#ttp 329821 - EDITOR: SETTINGS: Max and Min smoothed frame rates can be set invertedly
#branch UE4
#proj Runtime.Engine
#summary Implemented a struct customization for FRange and changed these two properties to an FRange, which can be accordingly clamped.
#add Implemented FRangeStructCustomization<> for FFloatRange and FInt32Range.
#add Implemented FFloatRange and FInt32Range mirrors in Object.h, for use by UHT.
#change Changed the Type member of TRangeBound to be an explicitly sized TEnumAsByte instead of a plain enum (so that the mirrored structs for UHT are guaranteed to match).
#change Changed Min and MaxSmoothedFrameRate members in UEngine to use a FFloatRange and added appropriate metadata.

diff --git a/Engine/Config/BaseEngine.ini b/Engine/Config/BaseEngine.ini
--- a/Engine/Config/BaseEngine.ini
+++ b/Engine/Config/BaseEngine.ini
@@ -159,8 +159,7 @@ NetDriverDefinitions=(DefName="GameNetDriver",DriverClassName="/Script/OnlineSub
 bUseNormalMapsForSimpleLightMaps=true
 
 bSmoothFrameRate=true
-MinSmoothedFrameRate=22
-MaxSmoothedFrameRate=62
+SmoothedFrameRateRange=(LowerBound=(Type="ERangeBoundTypes::Inclusive",Value=22),UpperBound=(Type="ERangeBoundTypes::Exclusive",Value=62))
 bCheckForMultiplePawnsSpawnedInAFrame=false
 NumPawnsAllowedToBeSpawnedInAFrame=2
 ; This is the default (additive) color of selected objects in the editor
@@ -2015,4 +2014,4 @@ bSupportsLandscapeRightOrientation=True
 DefaultGravityZ=-980.0
 
 [/Script/Engine.WorldComposition]
-TilesStreamingTimeThreshold=1.0
\ No newline at end of file
+TilesStreamingTimeThreshold=1.0
diff --git a/Engine/Source/Editor/DetailCustomizations/Private/DetailCustomizations.cpp b/Engine/Source/Editor/DetailCustomizations/Private/DetailCustomizations.cpp
--- a/Engine/Source/Editor/DetailCustomizations/Private/DetailCustomizations.cpp
+++ b/Engine/Source/Editor/DetailCustomizations/Private/DetailCustomizations.cpp
@@ -31,6 +31,7 @@
 #include "LandscapeSplineDetails.h"
 #include "MathStructCustomizations.h"
 #include "MathStructProxyCustomizations.h"
+#include "RangeStructCustomization.h"
 #include "StringAssetReferenceCustomization.h"
 #include "StringClassReferenceCustomization.h"
 #include "AttenuationSettingsCustomizations.h"
@@ -102,6 +103,8 @@ void FDetailCustomizationsModule::StartupModule()
 	RegisterStructPropertyLayout( "SlateSound", FOnGetStructCustomizationInstance::CreateStatic( &FSlateSoundStructCustomization::MakeInstance ) );
 	RegisterStructPropertyLayout( "Guid", FOnGetStructCustomizationInstance::CreateStatic( &FGuidStructCustomization::MakeInstance ) );
 	RegisterStructPropertyLayout( "Key", FOnGetStructCustomizationInstance::CreateStatic(&FKeyStructCustomization::MakeInstance));
+	RegisterStructPropertyLayout( "FloatRange", FOnGetStructCustomizationInstance::CreateStatic( &FRangeStructCustomization<float>::MakeInstance ) );
+	RegisterStructPropertyLayout( "Int32Range", FOnGetStructCustomizationInstance::CreateStatic( &FRangeStructCustomization<int32>::MakeInstance ) );
 
 	RegisterStructPropertyLayout( "BlackboardEntry", FOnGetStructCustomizationInstance::CreateStatic( &FBlackboardEntryDetails::MakeInstance ) );
 	RegisterStructPropertyLayout( "RuntimeFloatCurve", FOnGetStructCustomizationInstance::CreateStatic( &FCurveStructCustomization::MakeInstance ) );
diff --git a/Engine/Source/Editor/DetailCustomizations/Private/RangeStructCustomization.cpp b/Engine/Source/Editor/DetailCustomizations/Private/RangeStructCustomization.cpp
new file mode 100644
--- /dev/null
+++ b/Engine/Source/Editor/DetailCustomizations/Private/RangeStructCustomization.cpp
@@ -0,0 +1,385 @@
+// Copyright 1998-2014 Epic Games, Inc. All Rights Reserved.
+
+/*=============================================================================
+	RangeStructCustomization.cpp: Implements the FRangeStructCustomization class.
+  =============================================================================*/
+
+#include "DetailCustomizationsPrivatePCH.h"
+#include "RangeStructCustomization.h"
+
+
+#define LOCTEXT_NAMESPACE "RangeStructCustomization"
+
+
+/* Helper traits for getting a metadata property based on the template parameter type
+ *****************************************************************************/
+namespace
+{
+	template <typename NumericType>
+	struct FGetMetaDataHelper
+	{};
+
+	template <>
+	struct FGetMetaDataHelper<float>
+	{
+		static float GetMetaData(const UProperty* Property, const TCHAR* Key)
+		{
+			return Property->GetFLOATMetaData(Key);
+		}
+	};
+
+	template <>
+	struct FGetMetaDataHelper<int32>
+	{
+		static int32 GetMetaData(const UProperty* Property, const TCHAR* Key)
+		{
+			return Property->GetINTMetaData(Key);
+		}
+	};
+}
+
+
+/* FRangeStructCustomization static interface
+ *****************************************************************************/
+
+template <typename NumericType>
+TSharedRef<IStructCustomization> FRangeStructCustomization<NumericType>::MakeInstance()
+{
+	return MakeShareable(new FRangeStructCustomization<NumericType>);
+}
+
+
+/* IStructCustomization interface
+ *****************************************************************************/
+
+template <typename NumericType>
+void FRangeStructCustomization<NumericType>::CustomizeStructHeader(TSharedRef<IPropertyHandle> StructPropertyHandle, FDetailWidgetRow& HeaderRow, IStructCustomizationUtils& StructCustomizationUtils)
+{
+	// Get handles to the properties we're interested in
+	LowerBoundStructHandle = StructPropertyHandle->GetChildHandle(TEXT("LowerBound"));
+	UpperBoundStructHandle = StructPropertyHandle->GetChildHandle(TEXT("UpperBound"));
+	check(LowerBoundStructHandle.IsValid());
+	check(UpperBoundStructHandle.IsValid());
+
+	LowerBoundValueHandle = LowerBoundStructHandle->GetChildHandle(TEXT("Value"));
+	UpperBoundValueHandle = UpperBoundStructHandle->GetChildHandle(TEXT("Value"));
+	LowerBoundTypeHandle = LowerBoundStructHandle->GetChildHandle(TEXT("Type"));
+	UpperBoundTypeHandle = UpperBoundStructHandle->GetChildHandle(TEXT("Type"));
+	check(LowerBoundValueHandle.IsValid());
+	check(UpperBoundValueHandle.IsValid());
+	check(LowerBoundTypeHandle.IsValid());
+	check(UpperBoundTypeHandle.IsValid());
+
+	// Get min/max metadata values if defined
+	auto Property = StructPropertyHandle->GetProperty();
+	if (Property != nullptr)
+	{
+		if (Property->HasMetaData(TEXT("UIMin")))
+		{
+			MinAllowedValue = TOptional<NumericType>(FGetMetaDataHelper<NumericType>::GetMetaData(Property, TEXT("UIMin")));
+		}
+
+		if (Property->HasMetaData(TEXT("UIMax")))
+		{
+			MaxAllowedValue = TOptional<NumericType>(FGetMetaDataHelper<NumericType>::GetMetaData(Property, TEXT("UIMax")));
+		}
+	}
+
+	// Make weak pointers to be passed as payloads to the widgets
+	TWeakPtr<IPropertyHandle> LowerBoundValueWeakPtr = LowerBoundValueHandle;
+	TWeakPtr<IPropertyHandle> UpperBoundValueWeakPtr = UpperBoundValueHandle;
+	TWeakPtr<IPropertyHandle> LowerBoundTypeWeakPtr = LowerBoundTypeHandle;
+	TWeakPtr<IPropertyHandle> UpperBoundTypeWeakPtr = UpperBoundTypeHandle;
+
+	// Generate a list of enum values for the combo box from the LowerBound.Type property
+	TArray<bool> RestrictedList;
+	LowerBoundTypeHandle->GeneratePossibleValues(ComboBoxList, ComboBoxToolTips, RestrictedList);
+
+	// Get initial values for the combo box
+	uint8 LowerBoundType;
+	TSharedPtr<FString> LowerBoundTypeSelectedItem;
+	if (ensure(LowerBoundTypeHandle->GetValue(LowerBoundType) == FPropertyAccess::Success))
+	{
+		check(LowerBoundType < ComboBoxList.Num());
+		LowerBoundTypeSelectedItem = ComboBoxList[LowerBoundType];
+	}
+
+	uint8 UpperBoundType;
+	TSharedPtr<FString> UpperBoundTypeSelectedItem;
+	if (ensure(UpperBoundTypeHandle->GetValue(UpperBoundType) == FPropertyAccess::Success))
+	{
+		check(UpperBoundType < ComboBoxList.Num());
+		UpperBoundTypeSelectedItem = ComboBoxList[UpperBoundType];
+	}
+
+	// Build the widgets
+	HeaderRow.NameContent()
+	[
+		StructPropertyHandle->CreatePropertyNameWidget()
+	]
+	.ValueContent()
+	.MinDesiredWidth(200.0f)
+	.MaxDesiredWidth(200.0f)
+	[
+		SNew(SVerticalBox)
+
+		+SVerticalBox::Slot()
+		.Padding(FMargin(0.0f, 3.0f, 0.0f, 2.0f))
+		[
+			SNew(SHorizontalBox)
+
+			+SHorizontalBox::Slot()
+			.Padding(FMargin(0.0f, 0.0f, 5.0f, 0.0f))
+			.VAlign(VAlign_Center)
+			[
+				SNew(SNumericEntryBox<NumericType>)
+				.Value(this, &FRangeStructCustomization<NumericType>::OnGetValue, LowerBoundValueWeakPtr, LowerBoundTypeWeakPtr)
+				.MinValue(MinAllowedValue)
+				.MinSliderValue(MinAllowedValue)
+				.MaxValue(this, &FRangeStructCustomization<NumericType>::OnGetValue, UpperBoundValueWeakPtr, UpperBoundTypeWeakPtr)
+				.MaxSliderValue(this, &FRangeStructCustomization<NumericType>::OnGetValue, UpperBoundValueWeakPtr, UpperBoundTypeWeakPtr)
+				.OnValueCommitted(this, &FRangeStructCustomization<NumericType>::OnValueCommitted, LowerBoundValueWeakPtr)
+				.OnValueChanged(this, &FRangeStructCustomization<NumericType>::OnValueChanged, LowerBoundValueWeakPtr)
+				.OnBeginSliderMovement(this, &FRangeStructCustomization<NumericType>::OnBeginSliderMovement)
+				.OnEndSliderMovement(this, &FRangeStructCustomization<NumericType>::OnEndSliderMovement)
+				.IsEnabled(this, &FRangeStructCustomization<NumericType>::OnQueryIfEnabled, LowerBoundTypeWeakPtr)
+				.Font(IDetailLayoutBuilder::GetDetailFont())
+				.AllowSpin(true)
+				.LabelVAlign(VAlign_Center)
+				.Label()
+				[
+					SNew(STextBlock)
+					.Font(IDetailLayoutBuilder::GetDetailFont())
+					.Text(LOCTEXT("MinimumBoundLabel", "Min"))
+				]
+			]
+
+			+SHorizontalBox::Slot()
+			.VAlign(VAlign_Center)
+			.AutoWidth()
+			[
+				SNew(SComboBox< TSharedPtr<FString> >)
+				.OptionsSource(&ComboBoxList)
+				.OnGenerateWidget(this, &FRangeStructCustomization<NumericType>::OnGenerateComboWidget)
+				.OnSelectionChanged(this, &FRangeStructCustomization<NumericType>::OnComboSelectionChanged, LowerBoundTypeWeakPtr)
+				.InitiallySelectedItem(LowerBoundTypeSelectedItem)
+				[
+					// combo box button intentionally blank to avoid displaying excessive details
+					SNew(SSpacer)
+				]
+			]
+		]
+
+		+SVerticalBox::Slot()
+		.Padding(FMargin(0.0f, 2.0f, 0.0f, 3.0f))
+		[
+			SNew(SHorizontalBox)
+
+			+SHorizontalBox::Slot()
+			.Padding(FMargin(0.0f, 0.0f, 5.0f, 0.0f))
+			.VAlign(VAlign_Center)
+			[
+				SNew(SNumericEntryBox<NumericType>)
+				.Value(this, &FRangeStructCustomization<NumericType>::OnGetValue, UpperBoundValueWeakPtr, UpperBoundTypeWeakPtr)
+				.MinValue(this, &FRangeStructCustomization<NumericType>::OnGetValue, LowerBoundValueWeakPtr, LowerBoundTypeWeakPtr)
+				.MinSliderValue(this, &FRangeStructCustomization<NumericType>::OnGetValue, LowerBoundValueWeakPtr, LowerBoundTypeWeakPtr)
+				.MaxValue(MaxAllowedValue)
+				.MaxSliderValue(MaxAllowedValue)
+				.OnValueCommitted(this, &FRangeStructCustomization<NumericType>::OnValueCommitted, UpperBoundValueWeakPtr)
+				.OnValueChanged(this, &FRangeStructCustomization<NumericType>::OnValueChanged, UpperBoundValueWeakPtr)
+				.OnBeginSliderMovement(this, &FRangeStructCustomization<NumericType>::OnBeginSliderMovement)
+				.OnEndSliderMovement(this, &FRangeStructCustomization<NumericType>::OnEndSliderMovement)
+				.IsEnabled(this, &FRangeStructCustomization<NumericType>::OnQueryIfEnabled, UpperBoundTypeWeakPtr)
+				.Font(IDetailLayoutBuilder::GetDetailFont())
+				.AllowSpin(true)
+				.LabelVAlign(VAlign_Center)
+				.Label()
+				[
+					SNew(STextBlock)
+					.Font(IDetailLayoutBuilder::GetDetailFont())
+					.Text(LOCTEXT("MaximumBoundLabel", "Max"))
+				]
+			]
+
+			+SHorizontalBox::Slot()
+			.AutoWidth()
+			.VAlign(VAlign_Center)
+			[
+				SNew(SComboBox< TSharedPtr<FString> >)
+				.OptionsSource(&ComboBoxList)
+				.OnGenerateWidget(this, &FRangeStructCustomization<NumericType>::OnGenerateComboWidget)
+				.OnSelectionChanged(this, &FRangeStructCustomization<NumericType>::OnComboSelectionChanged, UpperBoundTypeWeakPtr)
+				.InitiallySelectedItem(UpperBoundTypeSelectedItem)
+				[
+					// combo box button intentionally blank to avoid displaying excessive details
+					SNew(SSpacer)
+				]
+			]
+		]
+	];
+}
+
+
+template <typename NumericType>
+void FRangeStructCustomization<NumericType>::CustomizeStructChildren(TSharedRef<IPropertyHandle> StructPropertyHandle, IDetailChildrenBuilder& StructBuilder, IStructCustomizationUtils& StructCustomizationUtils)
+{
+	// Don't display children, as editing them directly can break the constraints
+}
+
+
+
+/* FRangeStructCustomization callbacks
+ *****************************************************************************/
+
+template <typename NumericType>
+TOptional<NumericType> FRangeStructCustomization<NumericType>::OnGetValue(TWeakPtr<IPropertyHandle> ValueWeakPtr, TWeakPtr<IPropertyHandle> TypeWeakPtr) const
+{
+	auto ValueSharedPtr = ValueWeakPtr.Pin();
+	auto TypeSharedPtr = TypeWeakPtr.Pin();
+
+	if (TypeSharedPtr.IsValid())
+	{
+		uint8 Type;
+		if (ensure(TypeSharedPtr->GetValue(Type) == FPropertyAccess::Success))
+		{
+			if (Type != ERangeBoundTypes::Open && ValueSharedPtr.IsValid())
+			{
+				NumericType Value;
+				if (ensure(ValueSharedPtr->GetValue(Value) == FPropertyAccess::Success))
+				{
+					return TOptional<NumericType>(Value);
+				}
+			}
+		}
+	}
+
+	// Value couldn't be accessed, or was bound type 'open'.  Return an unset value
+	return TOptional<NumericType>();
+}
+
+template <typename NumericType>
+void FRangeStructCustomization<NumericType>::OnValueCommitted(NumericType NewValue, ETextCommit::Type CommitType, TWeakPtr<IPropertyHandle> HandleWeakPtr)
+{
+	auto HandleSharedPtr = HandleWeakPtr.Pin();
+	if (HandleSharedPtr.IsValid() && (!bIsUsingSlider || (bIsUsingSlider && ShouldAllowSpin())))
+	{
+		ensure(HandleSharedPtr->SetValue(NewValue) == FPropertyAccess::Success);
+	}
+}
+
+template <typename NumericType>
+void FRangeStructCustomization<NumericType>::OnValueChanged(NumericType NewValue, TWeakPtr<IPropertyHandle> HandleWeakPtr)
+{
+	if (bIsUsingSlider && ShouldAllowSpin())
+	{
+		auto HandleSharedPtr = HandleWeakPtr.Pin();
+		if (HandleSharedPtr.IsValid())
+		{
+			ensure(HandleSharedPtr->SetValue(NewValue, EPropertyValueSetFlags::InteractiveChange) == FPropertyAccess::Success);
+		}
+	}
+}
+
+template <typename NumericType>
+void FRangeStructCustomization<NumericType>::OnBeginSliderMovement()
+{
+	bIsUsingSlider = true;
+
+	if (ShouldAllowSpin())
+	{
+		GEditor->BeginTransaction(LOCTEXT("SetRangeProperty", "Set Range Property"));
+	}
+}
+
+template <typename NumericType>
+void FRangeStructCustomization<NumericType>::OnEndSliderMovement(NumericType /*NewValue*/)
+{
+	bIsUsingSlider = false;
+
+	if (ShouldAllowSpin())
+	{
+		GEditor->EndTransaction();
+	}
+}
+
+template <typename NumericType>
+bool FRangeStructCustomization<NumericType>::OnQueryIfEnabled(TWeakPtr<IPropertyHandle> HandleWeakPtr) const
+{
+	auto PropertyHandle = HandleWeakPtr.Pin();
+	if (PropertyHandle.IsValid())
+	{
+		uint8 BoundType;
+		if (ensure(PropertyHandle->GetValue(BoundType) == FPropertyAccess::Success))
+		{
+			return (BoundType != ERangeBoundTypes::Open);
+		}
+	}
+
+	return false;
+}
+
+template <typename NumericType>
+bool FRangeStructCustomization<NumericType>::ShouldAllowSpin() const
+{
+	uint8 LowerBoundType;
+	uint8 UpperBoundType;
+	if (ensure(LowerBoundTypeHandle->GetValue(LowerBoundType) == FPropertyAccess::Success) &&
+		ensure(UpperBoundTypeHandle->GetValue(UpperBoundType) == FPropertyAccess::Success))
+	{
+		return (LowerBoundType != ERangeBoundTypes::Open && UpperBoundType != ERangeBoundTypes::Open);
+	}
+
+	return false;
+}
+
+template <typename NumericType>
+TSharedRef<SWidget> FRangeStructCustomization<NumericType>::OnGenerateComboWidget(TSharedPtr<FString> InComboString)
+{
+	FString ToolTip;
+
+	// A list of tool tips should have been populated in a 1 to 1 correspondence
+	check(ComboBoxList.Num() == ComboBoxToolTips.Num());
+
+	if (ComboBoxToolTips.Num() > 0)
+	{
+		int32 Index = ComboBoxList.IndexOfByKey(InComboString);
+		if (ensure(Index >= 0))
+		{
+			ToolTip = *ComboBoxToolTips[Index];
+		}
+	}
+
+	return
+		SNew(SBox)
+		.WidthOverride(150.0f)
+		[
+			SNew(STextBlock)
+			.Text(*InComboString)
+			.ToolTipText(ToolTip)
+			.Font(IDetailLayoutBuilder::GetDetailFont())
+			.IsEnabled(true)
+		];
+}
+
+template <typename NumericType>
+void FRangeStructCustomization<NumericType>::OnComboSelectionChanged(TSharedPtr<FString> InSelectedItem, ESelectInfo::Type SelectInfo, TWeakPtr<IPropertyHandle> HandleWeakPtr)
+{
+	auto PropertyHandle = HandleWeakPtr.Pin();
+	if (PropertyHandle.IsValid())
+	{
+		int32 Index = ComboBoxList.IndexOfByKey(InSelectedItem);
+		if (ensure(Index >= 0))
+		{
+			ensure(PropertyHandle->SetValue(static_cast<uint8>(Index)) == FPropertyAccess::Success);
+		}
+	}
+}
+
+
+/* Only explicitly instantiate the types which are supported
+ *****************************************************************************/
+template class FRangeStructCustomization<float>;
+template class FRangeStructCustomization<int32>;
+
+
+#undef LOCTEXT_NAMESPACE
diff --git a/Engine/Source/Editor/DetailCustomizations/Private/RangeStructCustomization.h b/Engine/Source/Editor/DetailCustomizations/Private/RangeStructCustomization.h
new file mode 100644
--- /dev/null
+++ b/Engine/Source/Editor/DetailCustomizations/Private/RangeStructCustomization.h
@@ -0,0 +1,130 @@
+// Copyright 1998-2014 Epic Games, Inc. All Rights Reserved.
+
+/*=============================================================================
+	RangeStructCustomization.h: Declares the FRangeStructCustomization<> class.
+=============================================================================*/
+
+#pragma once
+
+/**
+* Implements a details panel customization for FFloatRange structures.
+*/
+template <typename NumericType>
+class FRangeStructCustomization : public IStructCustomization
+{
+public:
+
+	FRangeStructCustomization()
+		: bIsUsingSlider(false) {}
+
+	// Begin IStructCustomization interface
+
+	virtual void CustomizeStructHeader(TSharedRef<IPropertyHandle> StructPropertyHandle, FDetailWidgetRow& HeaderRow, IStructCustomizationUtils& StructCustomizationUtils) OVERRIDE;
+
+	virtual void CustomizeStructChildren(TSharedRef<IPropertyHandle> StructPropertyHandle, IDetailChildrenBuilder& StructBuilder, IStructCustomizationUtils& StructCustomizationUtils) OVERRIDE;
+
+	// End IStructCustomization interface
+
+public:
+
+	/**
+	* Creates a new instance.
+	*
+	* @return A new struct customization for Guids.
+	*/
+	static TSharedRef<IStructCustomization> MakeInstance();
+
+protected:
+
+	/**
+	 * Gets the value for the provided property handle
+	 *
+	 * @param	ValueWeakPtr	Handle to the property to get the value from
+	 * @param	TypeWeakPtr		Handle to the property to get the type from
+	 *
+	 * @return	The value or unset if it could not be accessed
+	 */
+	TOptional<NumericType> OnGetValue(TWeakPtr<IPropertyHandle> ValueWeakPtr, TWeakPtr<IPropertyHandle> TypeWeakPtr) const;
+
+	/**
+	 * Called when the value is committed from the property editor
+	 *
+	 * @param	NewValue		The new value of the property
+	 * @param	CommitType		How the value was committed (unused)
+	 * @param	HandleWeakPtr	Handle to the property that the new value is for
+	 */
+	void OnValueCommitted(NumericType NewValue, ETextCommit::Type CommitType, TWeakPtr<IPropertyHandle> HandleWeakPtr);
+
+	/**
+	 * Called when the value is changed in the property editor
+	 *
+	 * @param	NewValue		The new value of the property
+	 * @param	HandleWeakPtr	Handle to the property that the new value is for
+	 */
+	void OnValueChanged(NumericType NewValue, TWeakPtr<IPropertyHandle> HandleWeakPtr);
+
+	/**
+	 * Called when a value starts to be changed by a slider
+	 */
+	void OnBeginSliderMovement();
+
+	/**
+	 * Called when a value stops being changed by a slider
+	 *
+	 * @param	NewValue		The new value of the property
+	 */
+	void OnEndSliderMovement(NumericType NewValue);
+
+	/**
+	 * Determines if the value is valid from the handle of the range type passed in
+	 *
+	 * @param	HandleWeakPtr	Handle to the property to get the type from
+	 *
+	 * @return	Whether the value is valid or not.
+	 */
+	bool OnQueryIfEnabled(TWeakPtr<IPropertyHandle> HandleWeakPtr) const;
+
+	/**
+	 * Determines if the spinbox is enabled on a numeric value widget
+	 *
+	 * @return Whether the spinbox should be enabled.
+	 */
+	bool ShouldAllowSpin() const;
+
+	/**
+	 * Generates a row of the combo widget
+	 *
+	 * @param	InComboString	Item string to be displayed in the combo item
+	 *
+	 * @return	An SWidget representing the combo row.
+	 */
+	TSharedRef<SWidget> OnGenerateComboWidget(TSharedPtr<FString> InComboString);
+
+	/**
+	 * Called when an item is selected in the combo box
+	 *
+	 * @param	InSelectedItem	String of the item selected
+	 * @param	SelectInfo		Type of selection
+	 * @param	HandleWeakPtr	Handle to the type property which will be changed by the combo selection
+	 */
+	void OnComboSelectionChanged(TSharedPtr<FString> InSelectedItem, ESelectInfo::Type SelectInfo, TWeakPtr<IPropertyHandle> HandleWeakPtr);
+
+	// Cached shared pointers to properties that we are managing
+	TSharedPtr<IPropertyHandle> LowerBoundStructHandle;
+	TSharedPtr<IPropertyHandle> UpperBoundStructHandle;
+	TSharedPtr<IPropertyHandle> LowerBoundValueHandle;
+	TSharedPtr<IPropertyHandle> LowerBoundTypeHandle;
+	TSharedPtr<IPropertyHandle> UpperBoundValueHandle;
+	TSharedPtr<IPropertyHandle> UpperBoundTypeHandle;
+
+	// Min and max allowed values from the metadata
+	TOptional<NumericType> MinAllowedValue;
+	TOptional<NumericType> MaxAllowedValue;
+
+	// Arrays of combo box data: list items, and tooltips
+	TArray< TSharedPtr<FString> > ComboBoxList;
+	TArray< TSharedPtr<FString> > ComboBoxToolTips;
+
+	// Flags whether the slider is being moved at the moment on any of our widgets
+	bool bIsUsingSlider;
+};
diff --git a/Engine/Source/Editor/UnrealEd/Private/Editor.cpp b/Engine/Source/Editor/UnrealEd/Private/Editor.cpp
--- a/Engine/Source/Editor/UnrealEd/Private/Editor.cpp
+++ b/Engine/Source/Editor/UnrealEd/Private/Editor.cpp
@@ -1437,7 +1437,15 @@ float UEditorEngine::GetMaxTickRate( float DeltaTime, bool bAllowFrameRateSmooth
 	// Clamp editor frame rate, even if smoothing is disabled
 	if( !bSmoothFrameRate && GIsEditor && !GIsPlayInEditorWorld )
 	{
-		MaxTickRate = FMath::Clamp<float>( 1.0f / DeltaTime, MinSmoothedFrameRate, MaxSmoothedFrameRate );
+		MaxTickRate = 1.0f / DeltaTime;
+		if (SmoothedFrameRateRange.HasLowerBound())
+		{
+			MaxTickRate = FMath::Max(MaxTickRate, SmoothedFrameRateRange.GetLowerBoundValue());
+		}
+		if (SmoothedFrameRateRange.HasUpperBound())
+		{
+			MaxTickRate = FMath::Min(MaxTickRate, SmoothedFrameRateRange.GetUpperBoundValue());
+		}
 	}
 
 	return MaxTickRate;
diff --git a/Engine/Source/Runtime/Core/Public/Math/RangeBound.h b/Engine/Source/Runtime/Core/Public/Math/RangeBound.h
--- a/Engine/Source/Runtime/Core/Public/Math/RangeBound.h
+++ b/Engine/Source/Runtime/Core/Public/Math/RangeBound.h
@@ -329,7 +329,7 @@ public:
 private:
 
 	// Holds the type of the bound.
-	ERangeBoundTypes::Type Type;
+	TEnumAsByte<ERangeBoundTypes::Type> Type;
 
 	// Holds the bound's value.
 	ElementType Value;
diff --git a/Engine/Source/Runtime/CoreUObject/Classes/Object.h b/Engine/Source/Runtime/CoreUObject/Classes/Object.h
--- a/Engine/Source/Runtime/CoreUObject/Classes/Object.h
+++ b/Engine/Source/Runtime/CoreUObject/Classes/Object.h
@@ -677,6 +677,78 @@ struct FFallbackStruct
 {
 };
 
+UENUM()
+namespace ERangeBoundTypes
+{
+	/**
+	* Enumerates the valid types of range bounds.
+	*/
+	enum Type
+	{
+		/**
+		* The range excludes the bound.
+		*/
+		Exclusive,
+
+		/**
+		* The range includes the bound.
+		*/
+		Inclusive,
+
+		/**
+		* The bound is open.
+		*/
+		Open
+	};
+}
+
+// A float range bound
+
+USTRUCT(noexport)
+struct FFloatRangeBound
+{
+	UPROPERTY(EditAnywhere, Category=Range)
+	TEnumAsByte<ERangeBoundTypes::Type> Type;
+
+	UPROPERTY(EditAnywhere, Category=Range)
+	float Value;
+};
+
+// A float range
+
+USTRUCT(noexport)
+struct FFloatRange
+{
+	UPROPERTY(EditAnywhere, Category=Range)
+	FFloatRangeBound LowerBound;
+
+	UPROPERTY(EditAnywhere, Category=Range)
+	FFloatRangeBound UpperBound;
+};
+
+// An int32 range bound
+
+USTRUCT(noexport)
+struct FInt32RangeBound
+{
+	UPROPERTY(EditAnywhere, Category = Range)
+	TEnumAsByte<ERangeBoundTypes::Type> Type;
+
+	UPROPERTY(EditAnywhere, Category = Range)
+	int32 Value;
+};
+
+// An int32 range
+
+USTRUCT(noexport)
+struct FInt32Range
+{
+	UPROPERTY(EditAnywhere, Category = Range)
+	FInt32RangeBound LowerBound;
+
+	UPROPERTY(EditAnywhere, Category = Range)
+	FInt32RangeBound UpperBound;
+};
 
 //=============================================================================
 // Object: The base class all objects.
diff --git a/Engine/Source/Runtime/Engine/Classes/Engine/Engine.h b/Engine/Source/Runtime/Engine/Classes/Engine/Engine.h
--- a/Engine/Source/Runtime/Engine/Classes/Engine/Engine.h
+++ b/Engine/Source/Runtime/Engine/Classes/Engine/Engine.h
@@ -1026,13 +1026,9 @@ public:
 	UPROPERTY(config, EditAnywhere, Category=Framerate)
 	uint32 bSmoothFrameRate:1;
 
-	/** Maximum framerate to smooth. Code will try to not go over via waiting.										*/
-	UPROPERTY(config, EditAnywhere, Category=Framerate, meta=(UIMin=0, ClampMin=0))
-	float MaxSmoothedFrameRate;
-
-	/** Minimum framerate smoothing will kick in.																	*/
-	UPROPERTY(config, EditAnywhere, Category=Framerate, meta=(UIMin=0, ClampMin=0))
-	float MinSmoothedFrameRate;
+	/** Range of framerates in which smoothing will kick in */
+	UPROPERTY(config, EditAnywhere, Category=Framerate, meta=(UIMin=0, UIMax=200))
+	FFloatRange SmoothedFrameRateRange;
 
 	/** 
 	 * Whether we should check for more than N pawns spawning in a single frame.  
diff --git a/Engine/Source/Runtime/Engine/Private/GameEngine.cpp b/Engine/Source/Runtime/Engine/Private/GameEngine.cpp
--- a/Engine/Source/Runtime/Engine/Private/GameEngine.cpp
+++ b/Engine/Source/Runtime/Engine/Private/GameEngine.cpp
@@ -750,7 +750,10 @@ float UGameEngine::GetMaxTickRate( float DeltaTime, bool bAllowFrameRateSmoothin
 		// Limit framerate on console if VSYNC is enabled to avoid jumps from 30 to 60 and back.
 		if( CVar->GetValueOnGameThread() != 0 )
 		{
-			MaxTickRate = MaxSmoothedFrameRate;
+			if (SmoothedFrameRateRange.HasUpperBound())
+			{
+				MaxTickRate = SmoothedFrameRateRange.GetUpperBoundValue();
+			}
 		}
 	}
 	else 
diff --git a/Engine/Source/Runtime/Engine/Private/UnrealEngine.cpp b/Engine/Source/Runtime/Engine/Private/UnrealEngine.cpp
--- a/Engine/Source/Runtime/Engine/Private/UnrealEngine.cpp
+++ b/Engine/Source/Runtime/Engine/Private/UnrealEngine.cpp
@@ -5717,10 +5717,17 @@ float UEngine::GetMaxTickRate( float DeltaTime, bool bAllowFrameRateSmoothing )
 			RunningAverageDeltaTime = FMath::Lerp<float>( RunningAverageDeltaTime, FMath::Min<float>( DeltaTime, 0.2f ), 1 / 300.f );
 
 			// Work in FPS domain as that is what the function will return.
-			float AverageFPS = 1.f / RunningAverageDeltaTime;
+			MaxTickRate = 1.f / RunningAverageDeltaTime;
 
 			// Clamp FPS into ini defined min/ max range.
-			MaxTickRate = FMath::Clamp<float>( AverageFPS, MinSmoothedFrameRate, MaxSmoothedFrameRate );
+			if (SmoothedFrameRateRange.HasLowerBound())
+			{
+				MaxTickRate = FMath::Max( MaxTickRate, SmoothedFrameRateRange.GetLowerBoundValue() );
+			}
+			if (SmoothedFrameRateRange.HasUpperBound())
+			{
+				MaxTickRate = FMath::Min( MaxTickRate, SmoothedFrameRateRange.GetUpperBoundValue() );
+			}
 		}
 	}
 
