PrimitiveComponent without collision still has Simulate option

diff --git a/Engine/Source/Editor/DetailCustomizations/Private/PrimitiveComponentDetails.cpp b/Engine/Source/Editor/DetailCustomizations/Private/PrimitiveComponentDetails.cpp
--- a/Engine/Source/Editor/DetailCustomizations/Private/PrimitiveComponentDetails.cpp
+++ b/Engine/Source/Editor/DetailCustomizations/Private/PrimitiveComponentDetails.cpp
@@ -56,7 +56,8 @@ bool FPrimitiveComponentDetails::IsSimulatePhysicsEditable() const
 			}
 			else if(PrimitiveComponent->BodyInstance.BodySetup.IsValid())
 			{
-				bEnableSimulatePhysics &= (PrimitiveComponent->BodyInstance.BodySetup->AggGeom.GetElementCount() > 0 || PrimitiveComponent->BodyInstance.BodySetup->CollisionTraceFlag == CTF_UseComplexAsSimple);
+												//If there's no collision we still let them simulate physics. The object falls through the world - this behavior is debatable but what we decided on for now
+				bEnableSimulatePhysics &= true; // (PrimitiveComponent->BodyInstance.BodySetup->AggGeom.GetElementCount() > 0 || PrimitiveComponent->BodyInstance.BodySetup->CollisionTraceFlag == CTF_UseComplexAsSimple);
 			}
 			else if(StaticMeshComponent.IsValid() && StaticMeshComponent->StaticMesh && StaticMeshComponent->StaticMesh->BodySetup)
 			{
