#ttp 325958 - AUDIO: Sound Actor categories should be re-ordered to 1: Sound 2: Attenuation 3: Modulation 4: Auto Activate

#proj UE4.Editor

#summary re-ordered Sound Actor categories

#change Added ordering type enum for Sound, Attenuation and Modulation.

diff --git a/Engine/Source/Editor/DetailCustomizations/Private/AmbientSoundDetails.cpp b/Engine/Source/Editor/DetailCustomizations/Private/AmbientSoundDetails.cpp
--- a/Engine/Source/Editor/DetailCustomizations/Private/AmbientSoundDetails.cpp
+++ b/Engine/Source/Editor/DetailCustomizations/Private/AmbientSoundDetails.cpp
@@ -26,7 +26,7 @@ void FAmbientSoundDetails::CustomizeDetails( IDetailLayoutBuilder& DetailBuilder
 		}
 	}
 
-	DetailBuilder.EditCategory( "Sound" )
+	DetailBuilder.EditCategory( "Sound", TEXT(""), ECategoryPriority::Important )
 		.AddCustomRow( TEXT("") )
 		[
 			SNew(SVerticalBox)
@@ -91,6 +91,9 @@ void FAmbientSoundDetails::CustomizeDetails( IDetailLayoutBuilder& DetailBuilder
 					]
 			]
 		];
+
+	DetailBuilder.EditCategory("Attenuation", TEXT(""), ECategoryPriority::TypeSpecific);
+	DetailBuilder.EditCategory("Modulation", TEXT(""), ECategoryPriority::TypeSpecific);
 }
 
 bool FAmbientSoundDetails::IsEditSoundCueEnabled() const
