TTP 330596

We now scan for an AllowPreserveRatio metadata flag on properties, if it exists we show the preserve ratio button just like we do for scale.  Currently this is only enabled for the LightFunctionScale property on light components.

diff --git a/Engine/Source/Editor/DetailCustomizations/Private/MathStructCustomizations.cpp b/Engine/Source/Editor/DetailCustomizations/Private/MathStructCustomizations.cpp
--- a/Engine/Source/Editor/DetailCustomizations/Private/MathStructCustomizations.cpp
+++ b/Engine/Source/Editor/DetailCustomizations/Private/MathStructCustomizations.cpp
@@ -3,6 +3,8 @@
 #include "DetailCustomizationsPrivatePCH.h"
 #include "MathStructCustomizations.h"
 
+#define LOCTEXT_NAMESPACE "FMathStructCustomization"
+
 TSharedRef<IStructCustomization> FMathStructCustomization::MakeInstance() 
 {
 	return MakeShareable( new FMathStructCustomization );
@@ -15,7 +17,6 @@ void FMathStructCustomization::CustomizeStructHeader( TSharedRef<class IProperty
 	MakeHeaderRow( StructPropertyHandle, HeaderRow );
 }
 
-
 void FMathStructCustomization::CustomizeStructChildren( TSharedRef<class IPropertyHandle> StructPropertyHandle, class IDetailChildrenBuilder& StructBuilder, IStructCustomizationUtils& StructCustomizationUtils )
 {
 	for( int32 ChildIndex = 0; ChildIndex < SortedChildHandles.Num(); ++ChildIndex )
@@ -63,6 +64,52 @@ void FMathStructCustomization::MakeHeaderRow( TSharedRef<class IPropertyHandle>&
 			MakeChildWidget( ChildHandle )
 		];
 	}
+
+	if ( StructPropertyHandle->GetProperty()->HasMetaData("AllowPreserveRatio") )
+	{
+		if ( !GConfig->GetBool(TEXT("SelectionDetails"), *(StructPropertyHandle->GetProperty()->GetName() + TEXT("_PreserveScaleRatio")), bPreserveScaleRatio, GEditorUserSettingsIni) )
+		{
+			bPreserveScaleRatio = true;
+		}
+
+		HorizontalBox->AddSlot()
+		.AutoWidth()
+		.MaxWidth( 18.0f )
+		[
+			// Add a checkbox to toggle between preserving the ratio of x,y,z components of scale when a value is entered
+			SNew( SCheckBox )
+			.IsChecked(this, &FMathStructCustomization::IsPreserveScaleRatioChecked)
+			.OnCheckStateChanged(this, &FMathStructCustomization::OnPreserveScaleRatioToggled, StructWeakHandlePtr)
+			.Style( FEditorStyle::Get(), "TransparentCheckBox" )
+			.ToolTipText( LOCTEXT("PreserveScaleToolTip", "When locked, scales uniformly based on the current xyz scale values so the object maintains its shape in each direction when scaled" ) )
+			[
+				SNew( SImage )
+				.Image(this, &FMathStructCustomization::GetPreserveScaleRatioImage)
+				.ColorAndOpacity( FSlateColor::UseForeground() )
+			]
+		];
+	}
+}
+
+const FSlateBrush* FMathStructCustomization::GetPreserveScaleRatioImage() const
+{
+	return bPreserveScaleRatio ? FEditorStyle::GetBrush(TEXT("GenericLock")) : FEditorStyle::GetBrush(TEXT("GenericUnlock"));
+}
+
+ESlateCheckBoxState::Type FMathStructCustomization::IsPreserveScaleRatioChecked() const
+{
+	return bPreserveScaleRatio ? ESlateCheckBoxState::Checked : ESlateCheckBoxState::Unchecked;
+}
+
+void FMathStructCustomization::OnPreserveScaleRatioToggled(ESlateCheckBoxState::Type NewState, TWeakPtr<IPropertyHandle> PropertyHandle)
+{
+	bPreserveScaleRatio = ( NewState == ESlateCheckBoxState::Checked ) ? true : false;
+
+	if ( PropertyHandle.IsValid() )
+	{
+		FString SettingKey = ( PropertyHandle.Pin()->GetProperty()->GetName() + TEXT("_PreserveScaleRatio") );
+		GConfig->SetBool(TEXT("SelectionDetails"), *SettingKey, bPreserveScaleRatio, GEditorUserSettingsIni);
+	}
 }
 
 void FMathStructCustomization::GetSortedChildren( TSharedRef<IPropertyHandle> StructPropertyHandle, TArray< TSharedRef<IPropertyHandle> >& OutChildren )
@@ -147,7 +194,8 @@ TOptional<NumericType> FMathStructCustomization::OnGetValue( TWeakPtr<IPropertyH
 template<typename NumericType>
 void FMathStructCustomization::OnValueCommitted( NumericType NewValue, ETextCommit::Type CommitType, TWeakPtr<IPropertyHandle> WeakHandlePtr )
 {
-	WeakHandlePtr.Pin()->SetValue( NewValue );
+	EPropertyValueSetFlags::Type Flags = EPropertyValueSetFlags::DefaultFlags;
+	SetValue(NewValue, Flags, WeakHandlePtr);
 }	
 
 template<typename NumericType>
@@ -156,8 +204,62 @@ void FMathStructCustomization::OnValueChanged( NumericType NewValue, TWeakPtr<IP
 	if( bIsUsingSlider )
 	{
 		EPropertyValueSetFlags::Type Flags = EPropertyValueSetFlags::InteractiveChange;
-		WeakHandlePtr.Pin()->SetValue( NewValue, Flags );
+		SetValue(NewValue, Flags, WeakHandlePtr);
+	}
+}
+
+template<typename NumericType>
+void FMathStructCustomization::SetValue(NumericType NewValue, EPropertyValueSetFlags::Type Flags, TWeakPtr<IPropertyHandle> WeakHandlePtr)
+{
+	if ( bPreserveScaleRatio )
+	{
+		// Get the value for each object for the modified component
+		TArray<FString> OldValues;
+		if ( WeakHandlePtr.Pin()->GetPerObjectValues(OldValues) == FPropertyAccess::Success )
+		{
+			// Loop through each object and scale based on the new ratio for each object individually
+			for ( int32 OutputIndex = 0; OutputIndex < OldValues.Num(); ++OutputIndex )
+			{
+				NumericType OldValue;
+				TTypeFromString<NumericType>::FromString(OldValue, *OldValues[OutputIndex]);
+
+				// Account for the previous scale being zero.  Just set to the new value in that case?
+				NumericType Ratio = OldValue == 0 ? NewValue : NewValue / OldValue;
+				if ( Ratio == 0 )
+				{
+					Ratio = NewValue;
+				}
+
+				// Loop through all the child handles (each component of the math struct, like X, Y, Z...etc)
+				for ( int32 ChildIndex = 0; ChildIndex < SortedChildHandles.Num(); ++ChildIndex )
+				{
+					// Ignore scaling our selves.
+					TSharedRef<IPropertyHandle> ChildHandle = SortedChildHandles[ChildIndex];
+					if ( ChildHandle != WeakHandlePtr.Pin() )
+					{
+						// Get the value for each object.
+						TArray<FString> ObjectChildValues;
+						if ( ChildHandle->GetPerObjectValues(ObjectChildValues) == FPropertyAccess::Success )
+						{
+							// Individually scale each object's components by the same ratio.
+							for ( int32 ChildOutputIndex = 0; ChildOutputIndex < ObjectChildValues.Num(); ++ChildOutputIndex )
+							{
+								NumericType ChildOldValue;
+								TTypeFromString<NumericType>::FromString(ChildOldValue, *ObjectChildValues[ChildOutputIndex]);
+
+								NumericType ChildNewValue = ChildOldValue * Ratio;
+								ObjectChildValues[ChildOutputIndex] = TTypeToString<NumericType>::ToSanitizedString(ChildNewValue);
+							}
+
+							ChildHandle->SetPerObjectValues(ObjectChildValues);
+						}
+					}
+				}
+			}
+		}
 	}
+
+	WeakHandlePtr.Pin()->SetValue(NewValue, Flags);
 }
 
 bool FMathStructCustomization::IsValueEnabled(TWeakPtr<IPropertyHandle> WeakHandlePtr) const
@@ -580,3 +682,5 @@ FReply FColorStructCustomization::OnOpenFullColorPickerClicked()
 
 	return FReply::Handled();
 }
+
+#undef LOCTEXT_NAMESPACE
diff --git a/Engine/Source/Editor/DetailCustomizations/Private/MathStructCustomizations.h b/Engine/Source/Editor/DetailCustomizations/Private/MathStructCustomizations.h
--- a/Engine/Source/Editor/DetailCustomizations/Private/MathStructCustomizations.h
+++ b/Engine/Source/Editor/DetailCustomizations/Private/MathStructCustomizations.h
@@ -11,7 +11,9 @@ public:
 	static TSharedRef<IStructCustomization> MakeInstance();
 
 	FMathStructCustomization()
-		: bIsUsingSlider(false) {}
+		: bIsUsingSlider(false)
+		, bPreserveScaleRatio(false)
+	{ }
 
 	virtual ~FMathStructCustomization() {}
 
@@ -78,6 +80,26 @@ protected:
 	template<typename NumericType>
 	void OnValueChanged( NumericType NewValue, TWeakPtr<IPropertyHandle> WeakHandlePtr );
 
+	/**
+	 * Called to set the value of the property handle.
+	 *
+	 * @param NewValue		The new value of the property as a float
+	 * @param Flags         The flags to pass when setting the value on the property handle.
+	 * @param WeakHandlePtr	Handle to the property that the new value is for
+	 */
+	template<typename NumericType>
+	void SetValue(NumericType NewValue, EPropertyValueSetFlags::Type Flags, TWeakPtr<IPropertyHandle> WeakHandlePtr);
+
+private:
+	/** Gets the brush to use for the lock icon. */
+	const FSlateBrush* GetPreserveScaleRatioImage() const;
+
+	/** Gets the checked value of the preserve scale ratio option */
+	ESlateCheckBoxState::Type IsPreserveScaleRatioChecked() const;
+
+	/** Called when the user toggles preserve ratio. */
+	void OnPreserveScaleRatioToggled(ESlateCheckBoxState::Type NewState, TWeakPtr<IPropertyHandle> PropertyHandle);
+
 private:
 	/** Called when a value starts to be changed by a slider */
 	void OnBeginSliderMovement();
@@ -92,8 +114,12 @@ private:
 protected:
 	/** All the sorted children of the struct that should be displayed */
 	TArray< TSharedRef<IPropertyHandle> > SortedChildHandles;
+
 	/** True if a value is being changed by dragging a slider */
 	bool bIsUsingSlider;
+
+	/** True if the ratio is locked when scaling occurs (uniform scaling) */
+	bool bPreserveScaleRatio;
 };
 
 /**
diff --git a/Engine/Source/Runtime/Engine/Classes/Components/LightComponent.h b/Engine/Source/Runtime/Engine/Classes/Components/LightComponent.h
--- a/Engine/Source/Runtime/Engine/Classes/Components/LightComponent.h
+++ b/Engine/Source/Runtime/Engine/Classes/Components/LightComponent.h
@@ -103,7 +103,7 @@ class ENGINE_API ULightComponent : public ULightComponentBase
 	class UMaterialInterface* LightFunctionMaterial;
 
 	/** Scales the light function projection.  X and Y scale in the directions perpendicular to the light's direction, Z scales along the light direction. */
-	UPROPERTY(EditAnywhere, BlueprintReadOnly, Category=LightFunction)
+	UPROPERTY(EditAnywhere, BlueprintReadOnly, Category=LightFunction, meta=(AllowPreserveRatio = "true"))
 	FVector LightFunctionScale;
 
 	/** IES texture (light profiles from real world measured data) */
