#ttp 332185 - UE4: UI: There is no way to see what row name is assigned to a datatable lookup in the UI when the row name is long
#proj UE4
#branch UE4
#summary The width of a datatable combo button is no longer constrained
#extra Also added a tooltip in case the details panel is still too narrow

diff --git a/Engine/Source/Editor/DetailCustomizations/Private/DataTableCustomization.h b/Engine/Source/Editor/DetailCustomizations/Private/DataTableCustomization.h
--- a/Engine/Source/Editor/DetailCustomizations/Private/DataTableCustomization.h
+++ b/Engine/Source/Editor/DetailCustomizations/Private/DataTableCustomization.h
@@ -52,8 +52,10 @@ public:
 					.Font( StructCustomizationUtils.GetRegularFont() )
 				]
 			.ValueContent()
+			.MaxDesiredWidth(0.0f) // don't constrain the combo button width
 				[
 					SAssignNew( RowNameComboButton, SComboButton )
+					.ToolTipText( this, &FDataTableCustomizationLayout::GetRowNameComboBoxContentText )
 					.OnGetMenuContent( this, &FDataTableCustomizationLayout::GetListContent )
 					.ContentPadding( FMargin( 2.0f, 2.0f ) )
 					.ButtonContent()
