Tooltips for Sound Node and Cue properties
TTP #333033

diff --git a/Engine/Source/Editor/DetailCustomizations/Private/AttenuationSettingsCustomizations.cpp b/Engine/Source/Editor/DetailCustomizations/Private/AttenuationSettingsCustomizations.cpp
--- a/Engine/Source/Editor/DetailCustomizations/Private/AttenuationSettingsCustomizations.cpp
+++ b/Engine/Source/Editor/DetailCustomizations/Private/AttenuationSettingsCustomizations.cpp
@@ -81,15 +81,17 @@ void FAttenuationSettingsCustomization::CustomizeStructChildren( TSharedRef<IPro
 
 	IDetailPropertyRow& AttenuationShapeRow = ChildBuilder.AddChildProperty( AttenuationShapeHandle.ToSharedRef() );
 	
-	ChildBuilder.AddChildProperty( AttenuationExtentsHandle )
+	ChildBuilder.AddChildProperty(AttenuationExtentsHandle)
 		.Visibility(TAttribute<EVisibility>(this, &FAttenuationSettingsCustomization::IsBoxSelected))
-		.DisplayName(TEXT("Extents"));
+		.DisplayName(TEXT("Extents"))
+		.ToolTip(NSLOCTEXT("AttenuationSettings", "BoxExtents", "The dimensions of the of the box.").ToString());
 
 	ChildBuilder.AddChildContent(TEXT("Radius"))
 		.NameContent()
 		[
 			SNew(STextBlock)
 				.Text(NSLOCTEXT("AttenuationSettings", "RadiusLabel", "Radius"))
+				.ToolTipText(NSLOCTEXT("AttenuationSettings", "RadiusToolTip", "The distance from the location of the sound at which falloff begins."))
 				.Font(StructCustomizationUtils.GetRegularFont())
 		]
 		.ValueContent()
@@ -103,6 +105,7 @@ void FAttenuationSettingsCustomization::CustomizeStructChildren( TSharedRef<IPro
 			[
 				SNew(STextBlock)
 				.Text(NSLOCTEXT("AttenuationSettings", "CapsuleHalfHeightLabel", "Capsule Half Height"))
+				.ToolTipText(NSLOCTEXT("AttenuationSettings", "CapsuleHalfHeightToolTip", "The attenuation capsule's half height."))
 				.Font(StructCustomizationUtils.GetRegularFont())
 			]
 		.ValueContent()
@@ -116,6 +119,7 @@ void FAttenuationSettingsCustomization::CustomizeStructChildren( TSharedRef<IPro
 		[
 			SNew(STextBlock)
 			.Text(NSLOCTEXT("AttenuationSettings", "CapsuleRadiusLabel", "Capsule Radius"))
+			.ToolTipText(NSLOCTEXT("AttenuationSettings", "CapsuleRadiusToolTip", "The attenuation capsule's radius."))
 			.Font(StructCustomizationUtils.GetRegularFont())
 		]
 	.ValueContent()
@@ -129,6 +133,7 @@ void FAttenuationSettingsCustomization::CustomizeStructChildren( TSharedRef<IPro
 		[
 			SNew(STextBlock)
 			.Text(NSLOCTEXT("AttenuationSettings", "ConeRadiusLabel", "Cone Radius"))
+			.ToolTipText(NSLOCTEXT("AttenuationSettings", "ConeRadiusToolTip", "The attenuation cone's radius."))
 			.Font(StructCustomizationUtils.GetRegularFont())
 		]
 	.ValueContent()
@@ -142,6 +147,7 @@ void FAttenuationSettingsCustomization::CustomizeStructChildren( TSharedRef<IPro
 		[
 			SNew(STextBlock)
 			.Text(NSLOCTEXT("AttenuationSettings", "ConeAngleLabel", "Cone Angle"))
+			.ToolTipText(NSLOCTEXT("AttenuationSettings", "ConeAngleToolTip", "The angle of the inner edge of the attenuation cone's falloff. Inside this angle sounds will be at full volume."))
 			.Font(StructCustomizationUtils.GetRegularFont())
 		]
 	.ValueContent()
@@ -155,6 +161,7 @@ void FAttenuationSettingsCustomization::CustomizeStructChildren( TSharedRef<IPro
 		[
 			SNew(STextBlock)
 			.Text(NSLOCTEXT("AttenuationSettings", "ConeFalloffAngleLabel", "Cone Falloff Angle"))
+			.ToolTipText(NSLOCTEXT("AttenuationSettings", "ConeFalloffAngleToolTip", "The angle of the outer edge of the attenuation cone's falloff. Outside this angle sounds will be inaudible."))
 			.Font(StructCustomizationUtils.GetRegularFont())
 		]
 	.ValueContent()
diff --git a/Engine/Source/Runtime/Engine/Classes/Sound/SoundAttenuation.h b/Engine/Source/Runtime/Engine/Classes/Sound/SoundAttenuation.h
--- a/Engine/Source/Runtime/Engine/Classes/Sound/SoundAttenuation.h
+++ b/Engine/Source/Runtime/Engine/Classes/Sound/SoundAttenuation.h
@@ -37,15 +37,14 @@ namespace EAttenuationShape
 	};
 }
 
+/*
+The settings for attenuating.
+*/
 USTRUCT(BlueprintType)
 struct ENGINE_API FAttenuationSettings
 {
 	GENERATED_USTRUCT_BODY()
 
-	/* 
-	 The settings for attenuating. 
-	 */
-
 	/* Enable attenuation via volume. */
 	UPROPERTY(EditAnywhere, Category=Attenuation )
 	uint32 bAttenuate:1;
@@ -83,16 +82,24 @@ struct ENGINE_API FAttenuationSettings
 	UPROPERTY()
 	float RadiusMax_DEPRECATED;
 
+	/* The dimensions to use for the attenuation shape. Interpretation of the values differ per shape.
+	   Sphere  - X is Sphere Radius. Y and Z are unused
+	   Capsule - X is Capsule Half Height, Y is Capsule Radius, Z is unused
+	   Box     - X, Y, and Z are the Box's dimensions
+	   Cone    - X is Cone Radius, Y is Cone Angle, Z is Cone Falloff Angle
+	*/
 	UPROPERTY(EditAnywhere, Category=Attenuation)
 	FVector AttenuationShapeExtents;
 
+	/* The distance back from the sound's origin to begin the cone when using the cone attenuation shape. */
 	UPROPERTY(EditAnywhere, Category=Attenuation, meta=(ClampMin = "0"))
 	float ConeOffset;
 
+	/* The distance over which falloff occurs. */
 	UPROPERTY(EditAnywhere, Category=Attenuation, meta=(ClampMin = "0"))
 	float FalloffDistance;
 
-	/* The range at which to start applying a low passfilter. */
+	/* The range at which to start applying a low pass filter. */
 	UPROPERTY(EditAnywhere, Category=LowPassFilter )
 	float LPFRadiusMin;
 
diff --git a/Engine/Source/Runtime/Engine/Classes/Sound/SoundBase.h b/Engine/Source/Runtime/Engine/Classes/Sound/SoundBase.h
--- a/Engine/Source/Runtime/Engine/Classes/Sound/SoundBase.h
+++ b/Engine/Source/Runtime/Engine/Classes/Sound/SoundBase.h
@@ -33,11 +33,11 @@ class ENGINE_API USoundBase : public UObject
 	UPROPERTY()
 	FName SoundClass_DEPRECATED;
 
-	/** Sound group this sound cue belongs to */
+	/** Sound class this sound belongs to */
 	UPROPERTY(EditAnywhere, Category=Sound, meta=(DisplayName = "Sound Class"))
 	USoundClass* SoundClassObject;
 
-	/** For debugging purpose only . */
+	/** When "stat sounds -debug" has been specified, draw this sound's attenuation shape when the sound is audible. For debugging purpose only. */
 	UPROPERTY(EditAnywhere, Category=Playback)
 	uint32 bDebug:1;
 
@@ -45,7 +45,6 @@ class ENGINE_API USoundBase : public UObject
 	UPROPERTY(EditAnywhere, Category=Playback)
 	TEnumAsByte<EMaxConcurrentResolutionRule::Type> MaxConcurrentResolutionRule;
 
-
 	/** Maximum number of times this sound can be played concurrently. */
 	UPROPERTY(EditAnywhere, Category=Playback)
 	int32 MaxConcurrentPlayCount;
diff --git a/Engine/Source/Runtime/Engine/Classes/Sound/SoundCue.h b/Engine/Source/Runtime/Engine/Classes/Sound/SoundCue.h
--- a/Engine/Source/Runtime/Engine/Classes/Sound/SoundCue.h
+++ b/Engine/Source/Runtime/Engine/Classes/Sound/SoundCue.h
@@ -38,18 +38,22 @@ class USoundCue : public USoundBase
 {
 	GENERATED_UCLASS_BODY()
 
+	/* Indicates whether attenuation should use the Attenuation Overrides or the Attenuation Settings asset */
 	UPROPERTY(EditAnywhere, Category=Attenuation)
 	uint32 bOverrideAttenuation:1;
 
 	UPROPERTY()
 	class USoundNode* FirstNode;
 
+	/* Volume multiplier for the Sound Cue */
 	UPROPERTY(EditAnywhere, Category=Sound, AssetRegistrySearchable)
 	float VolumeMultiplier;
 
+	/* Pitch multiplier for the Sound Cue */
 	UPROPERTY(EditAnywhere, Category=Sound, AssetRegistrySearchable)
 	float PitchMultiplier;
 
+	/* Attenuation settings to use if Override Attenuation is set to true */
 	UPROPERTY(EditAnywhere, Category=Attenuation, meta=(EditCondition="bOverrideAttenuation"))
 	FAttenuationSettings AttenuationOverrides;
 
diff --git a/Engine/Source/Runtime/Engine/Classes/Sound/SoundNodeConcatenator.h b/Engine/Source/Runtime/Engine/Classes/Sound/SoundNodeConcatenator.h
--- a/Engine/Source/Runtime/Engine/Classes/Sound/SoundNodeConcatenator.h
+++ b/Engine/Source/Runtime/Engine/Classes/Sound/SoundNodeConcatenator.h
@@ -9,12 +9,15 @@
 #pragma once
 #include "SoundNodeConcatenator.generated.h"
 
-UCLASS(hidecategories=Object, editinlinenew, MinimalAPI, meta=( DisplayName="Concatenator" ))
+/**
+ Plays child nodes sequentially
+*/
+UCLASS(hidecategories = Object, editinlinenew, MinimalAPI, meta = (DisplayName = "Concatenator"))
 class USoundNodeConcatenator : public USoundNode
 {
 	GENERATED_UCLASS_BODY()
 
-	/** A volume for each input.  Automatically sized. */
+	/** Volume multiplier for each input. */
 	UPROPERTY(EditAnywhere, editfixedsize, Category=Concatenator)
 	TArray<float> InputVolume;
 
diff --git a/Engine/Source/Runtime/Engine/Classes/Sound/SoundNodeDialoguePlayer.h b/Engine/Source/Runtime/Engine/Classes/Sound/SoundNodeDialoguePlayer.h
--- a/Engine/Source/Runtime/Engine/Classes/Sound/SoundNodeDialoguePlayer.h
+++ b/Engine/Source/Runtime/Engine/Classes/Sound/SoundNodeDialoguePlayer.h
@@ -1,14 +1,13 @@
 // Copyright 1998-2014 Epic Games, Inc. All Rights Reserved.
 
 
-/** 
- * Sound node that contains a reference to the dialogue table to pull from and be played
- */
-
 #pragma once
 #include "SoundNodeDialoguePlayer.generated.h"
 
-UCLASS(hidecategories=Object, editinlinenew, MinimalAPI, meta=(DisplayName = "Dialogue Player"))
+/**
+ Sound node that contains a reference to the dialogue table to pull from and be played
+*/
+UCLASS(hidecategories = Object, editinlinenew, MinimalAPI, meta = (DisplayName = "Dialogue Player"))
 class USoundNodeDialoguePlayer : public USoundNode
 {
 	GENERATED_UCLASS_BODY()
@@ -16,6 +15,7 @@ class USoundNodeDialoguePlayer : public USoundNode
 	UPROPERTY(EditAnywhere, Category=DialoguePlayer)
 	FDialogueWaveParameter DialogueWaveParameter;
 
+	/* Whether the dialogue line should be played looping */
 	UPROPERTY(EditAnywhere, Category=DialoguePlayer)
 	uint32 bLooping:1;
 
diff --git a/Engine/Source/Runtime/Engine/Classes/Sound/SoundNodeEnveloper.h b/Engine/Source/Runtime/Engine/Classes/Sound/SoundNodeEnveloper.h
--- a/Engine/Source/Runtime/Engine/Classes/Sound/SoundNodeEnveloper.h
+++ b/Engine/Source/Runtime/Engine/Classes/Sound/SoundNodeEnveloper.h
@@ -11,21 +11,27 @@ class USoundNodeEnveloper : public USoundNode
 {
 	GENERATED_UCLASS_BODY()
 
+	// The time in seconds where the envelope's loop begins.
 	UPROPERTY(EditAnywhere, Category=Looping)
 	float LoopStart;
 
+	// The time in seconds where the envelope's loop ends.
 	UPROPERTY(EditAnywhere, Category=Looping)
 	float LoopEnd;
 
+	// The time in seconds it takes the evelope to fade out after the last loop is completed.
 	UPROPERTY(EditAnywhere, Category=Looping)
 	float DurationAfterLoop;
 
+	// The number of times the envelope should loop if looping is enabled and the envelope is not set to loop indefinitely.
 	UPROPERTY(EditAnywhere, Category=Looping)
 	int32 LoopCount;
 
+	// If enabled, the envelope will continue to loop indefenitely regardless of the Loop Count value.
 	UPROPERTY(EditAnywhere, Category=Looping)
 	uint32 bLoopIndefinitely:1;
-
+	
+	// If enabled, the envelope will loop using the loop settings.
 	UPROPERTY(EditAnywhere, Category=Looping)
 	uint32 bLoop:1;
 
@@ -35,21 +41,27 @@ class USoundNodeEnveloper : public USoundNode
 	UPROPERTY()
 	class UDistributionFloatConstantCurve* PitchInterpCurve_DEPRECATED;
 
+	// The distribution defining the volume envelope.
 	UPROPERTY(EditAnywhere, Category=Envelope)
 	FRuntimeFloatCurve VolumeCurve;
 
+	// The distribution defining the pitch envelope.
 	UPROPERTY(EditAnywhere, Category=Envelope)
 	FRuntimeFloatCurve PitchCurve;
 
+	// The minimum pitch for the input sound.
 	UPROPERTY(EditAnywhere, Category=Modulation, meta=(ToolTip="The lower bound of pitch (1.0 is no change)"))
 	float PitchMin;
 
+	// The maximum pitch for the input sound.
 	UPROPERTY(EditAnywhere, Category=Modulation, meta=(ToolTip="The upper bound of pitch (1.0 is no change)"))
 	float PitchMax;
 
+	// The minimum volume for the input sound.
 	UPROPERTY(EditAnywhere, Category=Modulation, meta=(ToolTip="The lower bound of volume (1.0 is no change)"))
 	float VolumeMin;
 
+	// The maximum volume for the input sound.
 	UPROPERTY(EditAnywhere, Category=Modulation, meta=(ToolTip="The upper bound of volume (1.0 is no change)"))
 	float VolumeMax;
 
diff --git a/Engine/Source/Runtime/Engine/Classes/Sound/SoundNodeGroupControl.h b/Engine/Source/Runtime/Engine/Classes/Sound/SoundNodeGroupControl.h
--- a/Engine/Source/Runtime/Engine/Classes/Sound/SoundNodeGroupControl.h
+++ b/Engine/Source/Runtime/Engine/Classes/Sound/SoundNodeGroupControl.h
@@ -4,13 +4,15 @@
 #include "SoundNodeGroupControl.generated.h"
 
 /** 
- * Plays different sounds depending on the number and proximity of the sounds to the listener
+ * Plays different sounds depending on the number of active sounds
+ * Any time a new sound is played, the first group that has an available slot will be chosen
  */
 UCLASS(hidecategories=Object, editinlinenew, MinimalAPI, meta=( DisplayName="Group Control" ))
 class USoundNodeGroupControl : public USoundNode
 {
 	GENERATED_UCLASS_BODY()
 
+	/* How many active sounds are allowed for each group */
 	UPROPERTY(EditAnywhere, editfixedsize, Category=GroupControl)
 	TArray<int32> GroupSizes;
 
diff --git a/Engine/Source/Runtime/Engine/Classes/Sound/SoundNodeModulatorContinuous.h b/Engine/Source/Runtime/Engine/Classes/Sound/SoundNodeModulatorContinuous.h
--- a/Engine/Source/Runtime/Engine/Classes/Sound/SoundNodeModulatorContinuous.h
+++ b/Engine/Source/Runtime/Engine/Classes/Sound/SoundNodeModulatorContinuous.h
@@ -28,7 +28,7 @@ struct FModulatorContinuousParams
 	{
 	}
 
-	/** The name of the sound instance parameter that drives this */
+	/** The name of the sound instance parameter that specifies the current value. */
 	UPROPERTY(EditAnywhere, Category=ModulatorContinousParameters)
 	FName ParameterName;
 
@@ -36,23 +36,23 @@ struct FModulatorContinuousParams
 	UPROPERTY(EditAnywhere, Category=ModulatorContinousParameters)
 	float Default;
 
-	/** todo document */
+	/** The minimum input value. Values will be clamped to the [MinInput, MaxInput] range.*/
 	UPROPERTY(EditAnywhere, Category=ModulatorContinousParameters)
 	float MinInput;
 
-	/** todo document */
-	UPROPERTY(EditAnywhere, Category=ModulatorContinousParameters)
+	/** The maximum input value. Values will be clamped to the [MinInput, MaxInput] range.*/
+	UPROPERTY(EditAnywhere, Category = ModulatorContinousParameters)
 	float MaxInput;
 
-	/** todo document */
-	UPROPERTY(EditAnywhere, Category=ModulatorContinousParameters)
+	/** The minimum output value. The input value will be scaled from the range [MinInput, MaxInput] to [MinOut, MaxOutput] */
+	UPROPERTY(EditAnywhere, Category = ModulatorContinousParameters)
 	float MinOutput;
 
-	/** todo document */
-	UPROPERTY(EditAnywhere, Category=ModulatorContinousParameters)
+	/** The maximum output value. The input value will be scaled from the range [MinInput, MaxInput] to [MinOut, MaxOutput] */
+	UPROPERTY(EditAnywhere, Category = ModulatorContinousParameters)
 	float MaxOutput;
 
-	/** todo document */
+	/** The mode with which to treat the input value */
 	UPROPERTY(EditAnywhere, Category=ModulatorContinousParameters)
 	TEnumAsByte<enum ModulationParamMode> ParamMode;
 
