L10N: Fixing localization bug

diff --git a/Engine/Source/Editor/DetailCustomizations/Private/DeviceProfileDetails.cpp b/Engine/Source/Editor/DetailCustomizations/Private/DeviceProfileDetails.cpp
--- a/Engine/Source/Editor/DetailCustomizations/Private/DeviceProfileDetails.cpp
+++ b/Engine/Source/Editor/DetailCustomizations/Private/DeviceProfileDetails.cpp
@@ -843,4 +843,4 @@ void FDeviceProfileConsoleVariablesPropertyDetails::OnCVarPropertyChanged()
 }
 
 
-#undef LOCTEXT_NAMESPACE // "DeviceProfileDetails"
\ No newline at end of file
+#undef LOCTEXT_NAMESPACE
\ No newline at end of file
