#UE4 - Change the wording on the movie player details panel prompt

diff --git a/Engine/Source/Editor/DetailCustomizations/Private/MoviePlayerSettingsDetails.cpp b/Engine/Source/Editor/DetailCustomizations/Private/MoviePlayerSettingsDetails.cpp
--- a/Engine/Source/Editor/DetailCustomizations/Private/MoviePlayerSettingsDetails.cpp
+++ b/Engine/Source/Editor/DetailCustomizations/Private/MoviePlayerSettingsDetails.cpp
@@ -37,11 +37,11 @@ bool FMoviePlayerSettingsDetails::HandlePathPicked(FString& InOutPath)
 	{
 		// ask the user if they want to import this movie
 		FSuppressableWarningDialog::FSetupInfo Info( 
-			LOCTEXT("ExternalMovieImportWarning", "This movie needs to be imported into your project, would you like to import the file now?"), 
-			LOCTEXT("ExternalMovieImportTitle", "Import Movie"), 
+			LOCTEXT("ExternalMovieImportWarning", "This movie needs to be copied into your project, would you like to copy the file now?"), 
+			LOCTEXT("ExternalMovieImportTitle", "Copy Movie"), 
 			TEXT("ImportMovieIntoProject") );
-		Info.ConfirmText = LOCTEXT("ExternalMovieImport_Confirm", "Import");
-		Info.CancelText = LOCTEXT("ExternalMovieImport_Cancel", "Don't Import");
+		Info.ConfirmText = LOCTEXT("ExternalMovieImport_Confirm", "Copy");
+		Info.CancelText = LOCTEXT("ExternalMovieImport_Cancel", "Don't Copy");
 
 		FSuppressableWarningDialog ImportWarningDialog( Info );
 
