Added an option to force loading movies to complete even if the game is done loading and added an option to toggle whether or not movies can be skipped by a user.

diff --git a/Engine/Source/Editor/DetailCustomizations/Private/MoviePlayerSettingsDetails.cpp b/Engine/Source/Editor/DetailCustomizations/Private/MoviePlayerSettingsDetails.cpp
--- a/Engine/Source/Editor/DetailCustomizations/Private/MoviePlayerSettingsDetails.cpp
+++ b/Engine/Source/Editor/DetailCustomizations/Private/MoviePlayerSettingsDetails.cpp
@@ -108,6 +108,9 @@ void FMoviePlayerSettingsDetails::CustomizeDetails( IDetailLayoutBuilder& Detail
 	TSharedRef<FDetailArrayBuilder> StartupMoviesBuilder = MakeShareable( new FDetailArrayBuilder( StartupMoviesPropertyHandle.ToSharedRef() ) );
 	StartupMoviesBuilder->OnGenerateArrayElementWidget( FOnGenerateArrayElementWidget::CreateSP(this, &FMoviePlayerSettingsDetails::GenerateArrayElementWidget) );
 
+	MoviesCategory.AddProperty( "bWaitForMoviesToComplete" );
+	MoviesCategory.AddProperty( "bMoviesAreSkippable" );
+
 	const bool bForAdvanced = false;
 	MoviesCategory.AddCustomBuilder( StartupMoviesBuilder, bForAdvanced );
 }
diff --git a/Engine/Source/Runtime/MoviePlayer/Private/DefaultGameMoviePlayer.cpp b/Engine/Source/Runtime/MoviePlayer/Private/DefaultGameMoviePlayer.cpp
--- a/Engine/Source/Runtime/MoviePlayer/Private/DefaultGameMoviePlayer.cpp
+++ b/Engine/Source/Runtime/MoviePlayer/Private/DefaultGameMoviePlayer.cpp
@@ -268,7 +268,9 @@ void FDefaultGameMoviePlayer::SetupLoadingScreenFromIni()
 		{
 			// fill out the attributes
 			FLoadingScreenAttributes LoadingScreen;
-			LoadingScreen.bAutoCompleteWhenLoadingCompletes = true;
+			LoadingScreen.bAutoCompleteWhenLoadingCompletes = !GetDefault<UMoviePlayerSettings>()->bWaitForMoviesToComplete;
+			LoadingScreen.bMoviesAreSkippable = GetDefault<UMoviePlayerSettings>()->bMoviesAreSkippable;
+
 			for(const auto& Movie : StartupMovies)
 			{
 				if(Movie.FilePath.Len() > 0)
diff --git a/Engine/Source/Runtime/MoviePlayer/Public/MoviePlayerSettings.h b/Engine/Source/Runtime/MoviePlayer/Public/MoviePlayerSettings.h
--- a/Engine/Source/Runtime/MoviePlayer/Public/MoviePlayerSettings.h
+++ b/Engine/Source/Runtime/MoviePlayer/Public/MoviePlayerSettings.h
@@ -14,9 +14,19 @@ public:
 
 	GENERATED_UCLASS_BODY()
 
+	/** If enabled the game waits for startup movies to complete even if loading has finished */
+	UPROPERTY(globalconfig, EditAnywhere, Category="Movies")
+	bool bWaitForMoviesToComplete;
+
+	/** If enabled startup movies can be skipped by the user when a mouse button is pressed */
+	UPROPERTY(globalconfig, EditAnywhere, Category = "Movies")
+	bool bMoviesAreSkippable;
+
 	/**
 	 * Movies to play on startup. Note that these must be in your game's Game/Content/Movies directory.
 	 */
 	UPROPERTY(globalconfig, EditAnywhere, Category="Movies", meta=(FilePathFilter = "mp4"))
 	TArray<FFilePath> StartupMovies;
+
+
 };
