Added the ability to copy and paste location,rotation, and scale from a details panel.

Cleaned up the ui for toggling absolute and relative scale.  We no longer show a hyperlink and use a dropdown menu instead.

diff --git a/Engine/Source/Editor/DetailCustomizations/Private/ActorDetails.cpp b/Engine/Source/Editor/DetailCustomizations/Private/ActorDetails.cpp
--- a/Engine/Source/Editor/DetailCustomizations/Private/ActorDetails.cpp
+++ b/Engine/Source/Editor/DetailCustomizations/Private/ActorDetails.cpp
@@ -415,7 +415,7 @@ void FActorDetails::AddTransformCategory( IDetailLayoutBuilder& DetailBuilder )
 		return;
 	}
 	
-	TSharedRef<FComponentTransformDetails> TransformDetails = MakeShareable( new FComponentTransformDetails( DetailBuilder.GetDetailsView().GetSelectedObjects(), SelectedActorInfo, DetailBuilder.GetPropertyUtilities()->GetNotifyHook() ) );
+	TSharedRef<FComponentTransformDetails> TransformDetails = MakeShareable( new FComponentTransformDetails( DetailBuilder.GetDetailsView().GetSelectedObjects(), SelectedActorInfo, DetailBuilder ) );
 
 	IDetailCategoryBuilder& TransformCategory = DetailBuilder.EditCategory( "TransformCommon", LOCTEXT("TransformCommonCategory", "Transform").ToString(), ECategoryPriority::Transform );
 
diff --git a/Engine/Source/Editor/DetailCustomizations/Private/ComponentTransformDetails.cpp b/Engine/Source/Editor/DetailCustomizations/Private/ComponentTransformDetails.cpp
--- a/Engine/Source/Editor/DetailCustomizations/Private/ComponentTransformDetails.cpp
+++ b/Engine/Source/Editor/DetailCustomizations/Private/ComponentTransformDetails.cpp
@@ -85,11 +85,11 @@ static void PropagateTransformPropertyChange(UObject* InObject, UProperty* InPro
 }
 
 
-FComponentTransformDetails::FComponentTransformDetails( const TArray< TWeakObjectPtr<UObject> >& InSelectedObjects, const FSelectedActorInfo& InSelectedActorInfo, FNotifyHook* InNotifyHook )
+FComponentTransformDetails::FComponentTransformDetails( const TArray< TWeakObjectPtr<UObject> >& InSelectedObjects, const FSelectedActorInfo& InSelectedActorInfo, IDetailLayoutBuilder& DetailBuilder )
 	: SelectedActorInfo( InSelectedActorInfo )
 	, SelectedObjects( InSelectedObjects )
 	, bPreserveScaleRatio( false )
-	, NotifyHook( InNotifyHook )
+	, NotifyHook( DetailBuilder.GetPropertyUtilities()->GetNotifyHook() )
 	, bEditingRotationInUI( false )
 {
 	GConfig->GetBool(TEXT("SelectionDetails"), TEXT("PreserveScaleRatio"), bPreserveScaleRatio, GEditorUserSettingsIni);
@@ -113,29 +113,292 @@ FComponentTransformDetails::FComponentTransformDetails( const TArray< TWeakObjec
 	}
 }
 
+TSharedRef<SWidget> FComponentTransformDetails::BuildTransformFieldLabel( ETransformField::Type TransformField )
+{
+	FText Label;
+	switch( TransformField )
+	{
+	case ETransformField::Rotation:
+		Label = LOCTEXT( "RotationLabel", "Rotation");
+		break;
+	case ETransformField::Scale:
+		Label = LOCTEXT( "ScaleLabel", "Scale" );
+		break;
+	case ETransformField::Location:
+	default:
+		Label = LOCTEXT("LocationLabel", "Location");
+		break;
+	}
+
+	FMenuBuilder MenuBuilder( true, NULL, NULL );
+
+	FUIAction SetRelativeLocationAction
+	(
+		FExecuteAction::CreateSP( this, &FComponentTransformDetails::OnSetRelativeTransform, TransformField ),
+		FCanExecuteAction(),
+		FIsActionChecked::CreateSP( this, &FComponentTransformDetails::IsRelativeTransformChecked, TransformField )
+	);
+
+	FUIAction SetWorldLocationAction
+	(
+		FExecuteAction::CreateSP( this, &FComponentTransformDetails::OnSetWorldTransform, TransformField ),
+		FCanExecuteAction(),
+		FIsActionChecked::CreateSP( this, &FComponentTransformDetails::IsWorldTransformChecked, TransformField )
+	);
+
+	MenuBuilder.BeginSection( TEXT("TransformType"), FText::Format( LOCTEXT("TransformType", "{0} Type"), Label ) );
+
+	MenuBuilder.AddMenuEntry
+	(
+		FText::Format( LOCTEXT( "RelativeLabel", "Relative"), Label ),
+		FText::Format( LOCTEXT( "RelativeLabel_ToolTip", "{0} is relative to its parent"), Label ),
+		FSlateIcon(),
+		SetRelativeLocationAction,
+		NAME_None, 
+		EUserInterfaceActionType::RadioButton
+	);
+
+	MenuBuilder.AddMenuEntry
+	(
+		FText::Format( LOCTEXT( "WorldLabel", "World"), Label ),
+		FText::Format( LOCTEXT( "WorldLabel_ToolTip", "{0} is relative to the world"), Label ),
+		FSlateIcon(),
+		SetWorldLocationAction,
+		NAME_None,
+		EUserInterfaceActionType::RadioButton
+	);
+
+	MenuBuilder.EndSection();
+
+
+	return 
+		SNew(SComboButton)
+		.ContentPadding( 0 )
+		.ButtonStyle( FEditorStyle::Get(), "NoBorder" )
+		.ForegroundColor( FSlateColor::UseForeground() )
+		.MenuContent()
+		[
+			MenuBuilder.MakeWidget()
+		]
+		.ButtonContent()
+		[
+			SNew( SBox )
+			.Padding( FMargin( 0.0f, 0.0f, 2.0f, 0.0f ) )
+			[
+				SNew(STextBlock)
+				.Text(this, &FComponentTransformDetails::GetTransformFieldText, TransformField)
+				.Font(IDetailLayoutBuilder::GetDetailFont())
+			]
+		];
+}
+
+FText FComponentTransformDetails::GetTransformFieldText( ETransformField::Type TransformField ) const
+{
+	switch (TransformField)
+	{
+	case ETransformField::Location:
+		return GetLocationText();
+		break;
+	case ETransformField::Rotation:
+		return GetRotationText();
+		break;
+	case ETransformField::Scale:
+		return GetScaleText();
+		break;
+	default:
+		return FText::GetEmpty();
+		break;
+	}
+}
+
+void FComponentTransformDetails::OnSetRelativeTransform(  ETransformField::Type TransformField )
+{
+	const bool bEnable = false;
+	switch (TransformField)
+	{
+	case ETransformField::Location:
+		OnToggleAbsoluteLocation(bEnable);
+		break;
+	case ETransformField::Rotation:
+		OnToggleAbsoluteRotation( bEnable );
+		break;
+	case ETransformField::Scale:
+		OnToggleAbsoluteScale( bEnable );
+		break;
+	default:
+		break;
+	}
+}
+
+bool FComponentTransformDetails::IsRelativeTransformChecked( ETransformField::Type TransformField ) const
+{
+	switch (TransformField)
+	{
+	case ETransformField::Location:
+		return !bAbsoluteLocation;
+		break;
+	case ETransformField::Rotation:
+		return !bAbsoluteRotation;
+		break;
+	case ETransformField::Scale:
+		return !bAbsoluteScale;
+		break;
+	default:
+		return false;
+		break;
+	}
+}
+
+void FComponentTransformDetails::OnSetWorldTransform( ETransformField::Type TransformField )
+{
+	const bool bEnable = true;
+	switch (TransformField)
+	{
+	case ETransformField::Location:
+		OnToggleAbsoluteLocation(bEnable);
+		break;
+	case ETransformField::Rotation:
+		OnToggleAbsoluteRotation(bEnable);
+		break;
+	case ETransformField::Scale:
+		OnToggleAbsoluteScale(bEnable);
+		break;
+	default:
+		break;
+	}
+}
+
+bool FComponentTransformDetails::IsWorldTransformChecked( ETransformField::Type TransformField ) const
+{
+	return !IsRelativeTransformChecked( TransformField );
+}
+
+bool FComponentTransformDetails::OnCanCopy( ETransformField::Type TransformField ) const
+{
+	// We can only copy values if the whole field is set.  If multiple values are defined we do not copy since we are unable to determine the value
+	switch (TransformField)
+	{
+	case ETransformField::Location:
+		return CachedLocation.IsSet();
+		break;
+	case ETransformField::Rotation:
+		return CachedRotation.IsSet();
+		break;
+	case ETransformField::Scale:
+		return CachedScale.IsSet();
+		break;
+	default:
+		return false;
+		break;
+	}
+}
+
+void FComponentTransformDetails::OnCopy( ETransformField::Type TransformField )
+{
+	CacheTransform();
+
+	FString CopyStr;
+	switch (TransformField)
+	{
+	case ETransformField::Location:
+		CopyStr = FString::Printf(TEXT("X=%f Y=%f Z=%f"), CachedLocation.X.GetValue(), CachedLocation.Y.GetValue(), CachedLocation.Z.GetValue());
+		break;
+	case ETransformField::Rotation:
+		CopyStr = FString::Printf(TEXT("P=%f Y=%f R=%f"), CachedRotation.Y.GetValue(), CachedRotation.Z.GetValue(), CachedRotation.X.GetValue());
+		break;
+	case ETransformField::Scale:
+		CopyStr = FString::Printf(TEXT("X=%f Y=%f Z=%f"), CachedScale.X.GetValue(), CachedScale.Y.GetValue(), CachedScale.Z.GetValue());
+		break;
+	default:
+		break;
+	}
+
+	if( !CopyStr.IsEmpty() )
+	{
+		FPlatformMisc::ClipboardCopy( *CopyStr );
+	}
+}
+
+void FComponentTransformDetails::OnPaste( ETransformField::Type TransformField )
+{
+	FString PastedText;
+	FPlatformMisc::ClipboardPaste(PastedText);
+
+	switch (TransformField)
+	{
+		case ETransformField::Location:
+		{
+			FVector Location;
+			if (Location.InitFromString(PastedText))
+			{
+				FScopedTransaction Transaction(LOCTEXT("PasteLocation", "Paste Location"));
+				OnSetLocation(Location.X, ETextCommit::Default, 0);
+				OnSetLocation(Location.Y, ETextCommit::Default, 1);
+				OnSetLocation(Location.Z, ETextCommit::Default, 2);
+			}
+		}
+		break;
+	case ETransformField::Rotation:
+		{
+			FRotator Rotation;
+			if (Rotation.InitFromString(PastedText))
+			{
+				FScopedTransaction Transaction(LOCTEXT("PasteRotation", "Paste Rotation"));
+				OnSetRotation(Rotation.Roll, ETextCommit::Default, 0);
+				OnSetRotation(Rotation.Pitch, ETextCommit::Default, 1);
+				OnSetRotation(Rotation.Yaw, ETextCommit::Default, 2);
+			}
+		}
+		break;
+	case ETransformField::Scale:
+		{
+			FVector Scale;
+			if (Scale.InitFromString(PastedText))
+			{
+				FScopedTransaction Transaction(LOCTEXT("PasteScale", "Paste Scale"));
+				OnSetScale(Scale.X, ETextCommit::Default, 0);
+				OnSetScale(Scale.Y, ETextCommit::Default, 1);
+				OnSetScale(Scale.Z, ETextCommit::Default, 2);
+			}
+		}
+		break;
+	default:
+		break;
+	}
+}
+
+FUIAction FComponentTransformDetails::CreateCopyAction( ETransformField::Type TransformField ) const
+{
+	return
+		FUIAction
+		(
+			FExecuteAction::CreateSP(this, &FComponentTransformDetails::OnCopy, TransformField ),
+			FCanExecuteAction::CreateSP(this, &FComponentTransformDetails::OnCanCopy, TransformField )
+		);
+}
+
+FUIAction FComponentTransformDetails::CreatePasteAction( ETransformField::Type TransformField ) const
+{
+	return 
+		 FUIAction( FExecuteAction::CreateSP(this, &FComponentTransformDetails::OnPaste, TransformField ) );
+}
+
 BEGIN_SLATE_FUNCTION_BUILD_OPTIMIZATION
 void FComponentTransformDetails::GenerateChildContent( IDetailChildrenBuilder& ChildrenBuilder )
 {
 	UClass* SceneComponentClass = USceneComponent::StaticClass();
 		
-
 	FSlateFontInfo FontInfo = IDetailLayoutBuilder::GetDetailFont();
-		
+
 	// Location
-	ChildrenBuilder.AddChildContent( LOCTEXT("Location", "Location").ToString() )
+	ChildrenBuilder.AddChildContent( LOCTEXT("LocationFilter", "Location").ToString() )
+		.CopyAction( CreateCopyAction( ETransformField::Location ) )
+		.PasteAction( CreatePasteAction( ETransformField::Location ) )
 		.NameContent()
+		.HAlign(HAlign_Left)
+		.VAlign(VAlign_Center)
 		[
-			SNew( SBox )
-			.HAlign(HAlign_Left)
-			.VAlign(VAlign_Center)
-			.Padding( FMargin( 2.0f, 0.0f ) )
-			[
-				SNew( SHyperlink )
-				.Text( this, &FComponentTransformDetails::GetLocationLabel )
-				.IsEnabled( this, &FComponentTransformDetails::GetIsEnabled )
-				.OnNavigate( this, &FComponentTransformDetails::OnLocationLabelClicked )
-				.TextStyle(FEditorStyle::Get(), "DetailsView.HyperlinkStyle")
-			]
+			BuildTransformFieldLabel( ETransformField::Location )
 		]
 		.ValueContent()
 		.MinDesiredWidth(125.0f * 3.0f)
@@ -186,19 +449,13 @@ void FComponentTransformDetails::GenerateChildContent( IDetailChildrenBuilder& C
 	
 	// Rotation
 	ChildrenBuilder.AddChildContent( LOCTEXT("RotationFilter", "Rotation").ToString() )
+		.CopyAction( CreateCopyAction(ETransformField::Rotation) )
+		.PasteAction( CreatePasteAction(ETransformField::Rotation) )
 		.NameContent()
+		.HAlign(HAlign_Left)
+		.VAlign(VAlign_Center)
 		[
-			SNew( SBox )
-			.HAlign(HAlign_Left)
-			.VAlign(VAlign_Center)
-			.Padding( FMargin( 2.0f, 0.0f ) )
-			[
-				SNew( SHyperlink )
-				.Text( this, &FComponentTransformDetails::GetRotationLabel )
-				.IsEnabled( this, &FComponentTransformDetails::GetIsEnabled )
-				.OnNavigate( this, &FComponentTransformDetails::OnRotationLabelClicked )
-				.TextStyle(FEditorStyle::Get(), "DetailsView.HyperlinkStyle")
-			]
+			BuildTransformFieldLabel(ETransformField::Rotation)
 		]
 		.ValueContent()
 		.MinDesiredWidth(125.0f * 3.0f)
@@ -254,19 +511,13 @@ void FComponentTransformDetails::GenerateChildContent( IDetailChildrenBuilder& C
 
 			
 	ChildrenBuilder.AddChildContent( LOCTEXT("ScaleFilter", "Scale").ToString() )
+		.CopyAction( CreateCopyAction(ETransformField::Scale) )
+		.PasteAction( CreatePasteAction(ETransformField::Scale) )
 		.NameContent()
+		.HAlign(HAlign_Left)
+		.VAlign(VAlign_Center)
 		[
-			SNew( SBox )
-			.HAlign(HAlign_Left)
-			.VAlign(VAlign_Center)
-			.Padding( FMargin( 2.0f, 0.0f ) )
-			[
-				SNew( SHyperlink )
-				.Text( this, &FComponentTransformDetails::GetScaleLabel )
-				.IsEnabled( this, &FComponentTransformDetails::GetIsEnabled )
-				.OnNavigate( this, &FComponentTransformDetails::OnScaleLabelClicked )
-				.TextStyle(FEditorStyle::Get(), "DetailsView.HyperlinkStyle")
-			]
+			BuildTransformFieldLabel(ETransformField::Scale)
 		]
 		.ValueContent()
 		.MinDesiredWidth(125.0f * 3.0f)
@@ -355,22 +606,22 @@ void FComponentTransformDetails::OnPreserveScaleRatioToggled( ESlateCheckBoxStat
 	GConfig->SetBool(TEXT("SelectionDetails"), TEXT("PreserveScaleRatio"), bPreserveScaleRatio, GEditorUserSettingsIni);
 }
 
-FString FComponentTransformDetails::GetLocationLabel() const
+FText FComponentTransformDetails::GetLocationText() const
 {
-	return bAbsoluteLocation ? LOCTEXT( "AbsoluteLocation", "Absolute Location" ).ToString() : LOCTEXT( "Location", "Location" ).ToString();
+	return bAbsoluteLocation ? LOCTEXT( "AbsoluteLocation", "Absolute Location" ) : LOCTEXT( "Location", "Location" );
 }
 
-FString FComponentTransformDetails::GetRotationLabel() const
+FText FComponentTransformDetails::GetRotationText() const
 {
-	return bAbsoluteRotation ? LOCTEXT( "AbsoluteRotation", "Absolute Rotation" ).ToString() : LOCTEXT( "Rotation", "Rotation" ).ToString();
+	return bAbsoluteRotation ? LOCTEXT( "AbsoluteRotation", "Absolute Rotation" ) : LOCTEXT( "Rotation", "Rotation" );
 }
 
-FString FComponentTransformDetails::GetScaleLabel() const
+FText FComponentTransformDetails::GetScaleText() const
 {
-	return bAbsoluteScale ? LOCTEXT( "AbsoluteScale", "Absolute Scale" ).ToString() : LOCTEXT( "Scale", "Scale" ).ToString();
+	return bAbsoluteScale ? LOCTEXT( "AbsoluteScale", "Absolute Scale" ) : LOCTEXT( "Scale", "Scale" );
 }
 
-void FComponentTransformDetails::OnLocationLabelClicked( )
+void FComponentTransformDetails::OnToggleAbsoluteLocation( bool bEnable )
 {
 	UProperty* AbsoluteLocationProperty = FindField<UProperty>( USceneComponent::StaticClass(), "bAbsoluteLocation" );
 
@@ -382,7 +633,7 @@ void FComponentTransformDetails::OnLocationLabelClicked( )
 		{
 			UObject* Object = ObjectPtr.Get();
 			USceneComponent* RootComponent = FComponentEditorUtils::GetSceneComponent( Object );
-			if( RootComponent )
+			if( RootComponent && RootComponent->bAbsoluteLocation != bEnable )
 			{
 				if( !bBeganTransaction )
 				{
@@ -436,7 +687,7 @@ void FComponentTransformDetails::OnLocationLabelClicked( )
 
 }
 
-void FComponentTransformDetails::OnRotationLabelClicked( )
+void FComponentTransformDetails::OnToggleAbsoluteRotation( bool bEnable )
 {
 	UProperty* AbsoluteRotationProperty = FindField<UProperty>( USceneComponent::StaticClass(), "bAbsoluteRotation" );
 
@@ -448,7 +699,7 @@ void FComponentTransformDetails::OnRotationLabelClicked( )
 		{
 			UObject* Object = ObjectPtr.Get();
 			USceneComponent* RootComponent = FComponentEditorUtils::GetSceneComponent( Object );
-			if( RootComponent )
+			if( RootComponent && RootComponent->bAbsoluteRotation != bEnable )
 			{
 				if( !bBeganTransaction )
 				{
@@ -501,7 +752,7 @@ void FComponentTransformDetails::OnRotationLabelClicked( )
 	}
 }
 
-void FComponentTransformDetails::OnScaleLabelClicked( )
+void FComponentTransformDetails::OnToggleAbsoluteScale( bool bEnable )
 {
 	UProperty* AbsoluteScaleProperty = FindField<UProperty>( USceneComponent::StaticClass(), "bAbsoluteScale" );
 
@@ -513,7 +764,7 @@ void FComponentTransformDetails::OnScaleLabelClicked( )
 		{
 			UObject* Object = ObjectPtr.Get();
 			USceneComponent* RootComponent = FComponentEditorUtils::GetSceneComponent( Object );
-			if( RootComponent )
+			if( RootComponent && RootComponent->bAbsoluteScale != bEnable )
 			{
 				if( !bBeganTransaction )
 				{
diff --git a/Engine/Source/Editor/DetailCustomizations/Private/ComponentTransformDetails.h b/Engine/Source/Editor/DetailCustomizations/Private/ComponentTransformDetails.h
--- a/Engine/Source/Editor/DetailCustomizations/Private/ComponentTransformDetails.h
+++ b/Engine/Source/Editor/DetailCustomizations/Private/ComponentTransformDetails.h
@@ -4,13 +4,22 @@
 
 #include "AssetSelection.h"
 
+namespace ETransformField
+{
+	enum Type
+	{
+		Location,
+		Rotation,
+		Scale
+	};
+}
 /**
  * Manages the Transform section of a details view                    
  */
 class FComponentTransformDetails : public TSharedFromThis<FComponentTransformDetails>, public IDetailCustomNodeBuilder
 {
 public:
-	FComponentTransformDetails( const TArray< TWeakObjectPtr<UObject> >& InSelectedObjects, const FSelectedActorInfo& InSelectedActorInfo, FNotifyHook* NotifyHook );
+	FComponentTransformDetails( const TArray< TWeakObjectPtr<UObject> >& InSelectedObjects, const FSelectedActorInfo& InSelectedActorInfo, IDetailLayoutBuilder& DetailBuilder );
 
 	/**
 	 * Caches the representation of the actor transform for the user input boxes                   
@@ -94,34 +103,99 @@ private:
 	void OnPreserveScaleRatioToggled( ESlateCheckBoxState::Type NewState );
 
 	/**
-	 * @return the Label to display for translation                   
+	 * Builds a transform field label
+	 *
+	 * @param TransformField The field to build the label for
+	 * @return The label widget
+	 */
+	TSharedRef<SWidget> BuildTransformFieldLabel( ETransformField::Type TransformField );
+
+	/**
+	 * Gets display text for a transform field
+	 *
+	 * @param TransformField	The field to get the text for
+	 * @return The text label for TransformField
 	 */
-	FString GetLocationLabel() const;
+	FText GetTransformFieldText( ETransformField::Type TransformField ) const;
 
 	/**
-	 * @return the Label to display for rotation                   
+	 * @return the text to display for translation                   
 	 */
-	FString GetRotationLabel() const;
+	FText GetLocationText() const;
+
+	/**
+	 * @return the text to display for rotation                   
+	 */
+	FText GetRotationText() const;
+
+	/**
+	 * @return the text to display for scale                   
+	 */
+	FText GetScaleText() const;
+
+	/**
+	 * Enables or disables absolute location
+	 *
+	 * @param bEnable True to enable, false to disable (use relative location)
+	 */
+	void OnToggleAbsoluteLocation( bool bEnable );
+
+	/**
+	 * Enables or disables absolute rotation
+	 *
+	 * @param bEnable True to enable, false to disable (use relative rotation)
+	 */
+	void OnToggleAbsoluteRotation( bool bEnable );
+
+	/**
+	 * Enables or disables absolute scale
+	 *
+	 * @param bEnable True to enable, false to disable (use relative scale)
+	 */
+	void OnToggleAbsoluteScale( bool bEnable );
+
+	/**
+	 * Sets relative transform on the specified field
+	 *
+	 * @param The field that should be set to relative
+	 */
+	void OnSetRelativeTransform( ETransformField::Type TransformField );
+
+	/**
+	 * Sets world transform on the specified field
+	 *
+	 * @param The field that should be set to world
+	 */
+	void OnSetWorldTransform( ETransformField::Type TransformField );
+
+	/** @return true if relative transform is enabled for the specified field */
+	bool IsRelativeTransformChecked( ETransformField::Type TransformField ) const;
+
+	/** @return true if world transform is enabled for the specified field */
+	bool IsWorldTransformChecked( ETransformField::Type TransformField ) const;
+
+	/** @return true of copy is enabled for the specified field */
+	bool OnCanCopy( ETransformField::Type TransformField ) const;
 
 	/**
-	 * @return the Label to display for scale                   
+	 * Copies the specified transform field to the clipboard
 	 */
-	FString GetScaleLabel() const;
+	void OnCopy( ETransformField::Type TransformField );
 
 	/**
-	 * Called when the location label is clicked (toggles between absolute and relative)                   
+	 * Pastes the specified transform field from the clipboard
 	 */
-	void OnLocationLabelClicked( );
+	void OnPaste( ETransformField::Type TransformField );
 
 	/**
-	 * Called when the rotation label is clicked (toggles between absolute and relative)                   
+	 * Creates a UI action for copying a specified transform field
 	 */
-	void OnRotationLabelClicked( );
+	FUIAction CreateCopyAction( ETransformField::Type TransformField ) const;
 
 	/**
-	 * Called when the scale label is clicked (toggles between absolute and relative)                   
+	 * Creates a UI action for pasting a specified transform field
 	 */
-	void OnScaleLabelClicked( );
+	FUIAction CreatePasteAction( ETransformField::Type TransformField ) const;
 
 	/** Called when the "Reset to Default" button for the location has been clicked */
 	FReply OnLocationResetClicked();
diff --git a/Engine/Source/Editor/DetailCustomizations/Private/SceneComponentDetails.cpp b/Engine/Source/Editor/DetailCustomizations/Private/SceneComponentDetails.cpp
--- a/Engine/Source/Editor/DetailCustomizations/Private/SceneComponentDetails.cpp
+++ b/Engine/Source/Editor/DetailCustomizations/Private/SceneComponentDetails.cpp
@@ -662,7 +662,7 @@ void FSceneComponentDetails::MakeTransformDetails( IDetailLayoutBuilder& DetailB
 
 		if( bShouldShowTransform )
 		{
-			TSharedRef<FComponentTransformDetails> TransformDetails = MakeShareable( new FComponentTransformDetails( bIsEditingBlueprintDefaults ? SceneComponentObjects : DetailBuilder.GetDetailsView().GetSelectedObjects(), SelectedActorInfo, DetailBuilder.GetPropertyUtilities()->GetNotifyHook() ) );
+			TSharedRef<FComponentTransformDetails> TransformDetails = MakeShareable( new FComponentTransformDetails( bIsEditingBlueprintDefaults ? SceneComponentObjects : DetailBuilder.GetDetailsView().GetSelectedObjects(), SelectedActorInfo, DetailBuilder ) );
 
 			IDetailCategoryBuilder& TransformCategory = DetailBuilder.EditCategory( "TransformCommon", LOCTEXT("TransformCommonCategory", "Transform").ToString(), ECategoryPriority::Transform );
 
diff --git a/Engine/Source/Editor/EditorStyle/Private/SlateEditorStyle.cpp b/Engine/Source/Editor/EditorStyle/Private/SlateEditorStyle.cpp
--- a/Engine/Source/Editor/EditorStyle/Private/SlateEditorStyle.cpp
+++ b/Engine/Source/Editor/EditorStyle/Private/SlateEditorStyle.cpp
@@ -2413,6 +2413,7 @@ void FSlateEditorStyle::FStyle::SetupPropertyEditorStyles()
 		Set( "DetailsView.CategoryTop_Hovered", new BOX_BRUSH( "PropertyView/DetailCategoryTop", FMargin( 4/16.0f, 8.0f/16.0f, 4/16.0f, 4/16.0f ), FLinearColor(0.5f,0.5f,0.5f,1.0f) ) );
 		Set( "DetailsView.CategoryBottom", new BOX_BRUSH( "PropertyView/DetailCategoryBottom", FMargin(4.0f/16.0f) ) );
 		Set( "DetailsView.CategoryMiddle", new IMAGE_BRUSH( "PropertyView/DetailCategoryMiddle", FVector2D( 16, 16 ) ) );
+		Set( "DetailsView.CategoryMiddle_Hovered", new IMAGE_BRUSH( "PropertyView/DetailCategoryMiddle_Hovered", FVector2D( 16, 16 ) ) );
 
 		Set( "DetailsView.Splitter", FSplitterStyle()
 			.SetHandleNormalBrush( IMAGE_BRUSH( "Common/SplitterHandleHighlight", Icon8x8, FLinearColor::Black ) )
diff --git a/Engine/Source/Editor/PropertyEditor/Private/SDetailSingleItemRow.cpp b/Engine/Source/Editor/PropertyEditor/Private/SDetailSingleItemRow.cpp
--- a/Engine/Source/Editor/PropertyEditor/Private/SDetailSingleItemRow.cpp
+++ b/Engine/Source/Editor/PropertyEditor/Private/SDetailSingleItemRow.cpp
@@ -168,7 +168,7 @@ void SDetailSingleItemRow::Construct( const FArguments& InArgs, FDetailLayoutCus
 	this->ChildSlot
 	[	
 		SNew( SBorder )
-		.BorderImage( FEditorStyle::GetBrush("DetailsView.CategoryMiddle") )
+		.BorderImage( this, &SDetailSingleItemRow::GetBorderImage )
 		.Padding( 0.0f )
 		[
 			Widget
@@ -185,11 +185,21 @@ void SDetailSingleItemRow::Construct( const FArguments& InArgs, FDetailLayoutCus
 
 bool SDetailSingleItemRow::OnContextMenuOpening( FMenuBuilder& MenuBuilder )
 {
-	if( Customization->HasPropertyNode() )
+	if( Customization->HasPropertyNode() || Customization->GetWidgetRow().IsCopyPasteBound() )
 	{
-		FUIAction CopyAction( FExecuteAction::CreateSP( this, &SDetailSingleItemRow::OnCopyProperty ) );
+		FUIAction CopyAction  = Customization->GetWidgetRow().CopyMenuAction;
+		FUIAction PasteAction = Customization->GetWidgetRow().PasteMenuAction;
 
-		FUIAction PasteAction( FExecuteAction::CreateSP( this, &SDetailSingleItemRow::OnPasteProperty ), FCanExecuteAction::CreateSP( this, &SDetailSingleItemRow::CanPasteProperty ) );
+		if( !CopyAction.ExecuteAction.IsBound() && Customization->HasPropertyNode() )
+		{
+			CopyAction.ExecuteAction = FExecuteAction::CreateSP( this, &SDetailSingleItemRow::OnCopyProperty );
+		}
+
+		if( !PasteAction.ExecuteAction.IsBound() && Customization->HasPropertyNode() )
+		{
+			PasteAction.ExecuteAction = FExecuteAction::CreateSP( this, &SDetailSingleItemRow::OnPasteProperty );
+			PasteAction.CanExecuteAction = FCanExecuteAction::CreateSP( this, &SDetailSingleItemRow::CanPasteProperty );
+		}
 
 		MenuBuilder.AddMenuSeparator();
 
@@ -248,10 +258,22 @@ void SDetailSingleItemRow::OnPasteProperty()
 bool SDetailSingleItemRow::CanPasteProperty() const
 {
 	FString ClipboardContent;
-	if( Customization->HasPropertyNode() && OwnerTreeNode.IsValid() )
+	if( OwnerTreeNode.IsValid() )
 	{
 		FPlatformMisc::ClipboardPaste(ClipboardContent);
 	}
 
 	return !ClipboardContent.IsEmpty();
+}
+
+const FSlateBrush* SDetailSingleItemRow::GetBorderImage() const
+{
+	if (IsHovered())
+	{
+		return FEditorStyle::GetBrush("DetailsView.CategoryMiddle_Hovered");
+	}
+	else
+	{
+		return FEditorStyle::GetBrush("DetailsView.CategoryMiddle");
+	}
 }
\ No newline at end of file
diff --git a/Engine/Source/Editor/PropertyEditor/Private/SDetailSingleItemRow.h b/Engine/Source/Editor/PropertyEditor/Private/SDetailSingleItemRow.h
--- a/Engine/Source/Editor/PropertyEditor/Private/SDetailSingleItemRow.h
+++ b/Engine/Source/Editor/PropertyEditor/Private/SDetailSingleItemRow.h
@@ -28,6 +28,7 @@ private:
 	void OnCopyProperty();
 	void OnPasteProperty();
 	bool CanPasteProperty() const;
+	const FSlateBrush* GetBorderImage() const;
 private:
 	/** Customization for this widget */
 	FDetailLayoutCustomization* Customization;
diff --git a/Engine/Source/Editor/PropertyEditor/Public/DetailWidgetRow.h b/Engine/Source/Editor/PropertyEditor/Public/DetailWidgetRow.h
--- a/Engine/Source/Editor/PropertyEditor/Public/DetailWidgetRow.h
+++ b/Engine/Source/Editor/PropertyEditor/Public/DetailWidgetRow.h
@@ -68,6 +68,8 @@ public:
 		, VisibilityAttr( EVisibility::Visible )
 		, IsEnabledAttr( true )
 		, FilterTextString()
+		, CopyMenuAction()
+		, PasteMenuAction()
 	{
 	}
 	
@@ -131,6 +133,24 @@ public:
 		return *this;
 	}
 
+	/**
+	 * Sets a custom copy action to take when copying the data from this row
+	 */
+	FDetailWidgetRow& CopyAction( const FUIAction& InCopyAction )
+	{
+		CopyMenuAction = InCopyAction;
+		return *this;
+	}
+
+	/**
+	 * Sets a custom paste action to take when copying the data from this row
+	 */
+	FDetailWidgetRow& PasteAction( const FUIAction& InPasteAction  )
+	{
+		PasteMenuAction = InPasteAction;
+		return *this;
+	}
+
 	/**
 	 * @return true if the row has columns, false if it spans the entire row
 	 */
@@ -147,6 +167,12 @@ public:
 		return WholeRowWidget.Widget != SNullWidget::NullWidget || HasColumns();
 	}
 
+	/** @return true if a custom copy/paste is bound on this row */
+	bool IsCopyPasteBound() const
+	{
+		return CopyMenuAction.ExecuteAction.IsBound() && PasteMenuAction.ExecuteAction.IsBound();
+	}
+
 public:
 	/** Name column content */
 	FDetailWidgetDecl NameWidget;
@@ -160,5 +186,9 @@ public:
 	TAttribute<bool> IsEnabledAttr;
 	/** String to filter with */
 	FString FilterTextString;
+	/** Action for coping data on this row */
+	FUIAction CopyMenuAction;
+	/** Action for pasting data on this row */
+	FUIAction PasteMenuAction;
 };
 
