Engine: Prevented Stationary mobility from showing up on anything except subclasses of ULightComponentBase

diff --git a/Engine/Source/Editor/DetailCustomizations/Private/SceneComponentDetails.cpp b/Engine/Source/Editor/DetailCustomizations/Private/SceneComponentDetails.cpp
--- a/Engine/Source/Editor/DetailCustomizations/Private/SceneComponentDetails.cpp
+++ b/Engine/Source/Editor/DetailCustomizations/Private/SceneComponentDetails.cpp
@@ -313,9 +313,9 @@ static bool IsMobilitySettingProhibited(EComponentMobility::Type const MobilityV
 
 	case EComponentMobility::Stationary:
 		{
-			if (SceneComponent->IsA(UStaticMeshComponent::StaticClass()))
+			if (!SceneComponent->IsA(ULightComponentBase::StaticClass()))
 			{
-				ProhibitedReasonOut = LOCTEXT("MeshStationaryRestriction", "Mesh components cannot be stationary.");
+				ProhibitedReasonOut = LOCTEXT("OnlyLightsCanBeStationary", "Only light components can be stationary.");
 				bIsProhibited = true;
 			}
 			else if (GetInheritedMobility(SceneComponent) == EComponentMobility::Movable)
