Editor: Add missing IPropertyUtilities.h include (non-unity)

diff --git a/Engine/Source/Editor/DetailCustomizations/Private/ComponentTransformDetails.cpp b/Engine/Source/Editor/DetailCustomizations/Private/ComponentTransformDetails.cpp
--- a/Engine/Source/Editor/DetailCustomizations/Private/ComponentTransformDetails.cpp
+++ b/Engine/Source/Editor/DetailCustomizations/Private/ComponentTransformDetails.cpp
@@ -8,6 +8,7 @@
 #include "ActorEditorUtils.h"
 #include "Editor/UnrealEd/Public/Kismet2/ComponentEditorUtils.h"
 #include "ScopedTransaction.h"
+#include "IPropertyUtilities.h"
 
 #define LOCTEXT_NAMESPACE "FComponentTransformDetails"
 
