Clean up steering curve UI a bit

diff --git a/Engine/Source/Editor/DetailCustomizations/Private/WheeledVehicleMovementComponent4WDetails.cpp b/Engine/Source/Editor/DetailCustomizations/Private/WheeledVehicleMovementComponent4WDetails.cpp
--- a/Engine/Source/Editor/DetailCustomizations/Private/WheeledVehicleMovementComponent4WDetails.cpp
+++ b/Engine/Source/Editor/DetailCustomizations/Private/WheeledVehicleMovementComponent4WDetails.cpp
@@ -6,6 +6,7 @@
 #include "ObjectEditorUtils.h"
 #include "Vehicles/WheeledVehicleMovementComponent4W.h"
 #include "IDocumentation.h"
+#include "SCurveEditor.h"
 
 #define LOCTEXT_NAMESPACE "WheeledVehicleMovementComponent4WDetails"
 
@@ -18,156 +19,92 @@ TSharedRef<IDetailCustomization> FWheeledVehicleMovementComponent4WDetails::Make
 	return MakeShareable(new FWheeledVehicleMovementComponent4WDetails);
 }
 
-void FWheeledVehicleMovementComponent4WDetails::FillAutoBoxWidget(UWheeledVehicleMovementComponent4W * VehicleMovement, FDetailWidgetRow & AutoBoxSetup)
+void FWheeledVehicleMovementComponent4WDetails::CustomizeDetails(IDetailLayoutBuilder& DetailBuilder)
 {
-	/*FSlateFontInfo DetailFont = IDetailLayoutBuilder::GetDetailFont();
-
-	AutoBoxSetup.NameContent()
-	[
-		SNew(STextBlock)
-		.Text(LOCTEXT("AutoBoxSetupHeader", "Auto Box Setup"))
-		.ToolTipText(LOCTEXT("AutoBoxSetupHeaderTip", "Describes the automatic transition of the vehicle - how it shifts gears etc."))
-		.Font(DetailFont)
-	];
+	DetailBuilder.GetObjectsBeingCustomized(SelectedObjects);
 
-	TSharedRef<SVerticalBox> AutoBoxWidget = SNew(SVerticalBox);
+	//we only do fancy customization if we have one vehicle component selected
+	if (SelectedObjects.Num() != 1)
+	{
+		return;
+	}
+	else if (UWheeledVehicleMovementComponent4W * VehicleComponent = Cast<UWheeledVehicleMovementComponent4W>(SelectedObjects[0].Get()))
+	{
+		SteeringCurveEditor = FSteeringCurveEditor(VehicleComponent);
+	}
+	else
+	{
+		return;
+	}
 
-	AutoBoxWidget
+	// See which categories are hidden
+	TArray<FString> HideCategories;
+	DetailBuilder.GetDetailsView().GetBaseClass()->GetHideCategories(HideCategories);
+	
+	IDetailCategoryBuilder& SteeringCategory = DetailBuilder.EditCategory("SteeringSetup");
+	TSharedRef<IPropertyHandle> SteeringCurveHandle = DetailBuilder.GetProperty("SteeringCurve");
+	SteeringCategory.AddProperty(SteeringCurveHandle).CustomWidget()
+	.NameContent()
+	[
+		SteeringCurveHandle->CreatePropertyNameWidget()
+	]
+	.ValueContent()
+	.MinDesiredWidth(125.f * 3.f)
 	[
-		SNew(STextBlock)
-		.Text(LOCTEXT("ASDS", "Testing"))
+		SAssignNew(SteeringCurveWidget, SCurveEditor)
+		.ViewMinInput(0.f)
+		.ViewMaxInput(150.f)
+		.ViewMinOutput(0.f)
+		.ViewMaxOutput(1.f)
+		.TimelineLength(150)
+		.HideUI(false)
+		.DesiredSize(FVector2D(512, 128))
 	];
 
-	AutoBoxSetup.ValueContent()
-	[
-		AutoBoxWidget
-	];*/
-	
-	//We want to inline the gear ratios with Neutral and Latency.
-	//We only want to show as many ratios as there are gears
-	/*FVehicleAutoBoxData AutoBox = VehicleMovement->AutoBoxSetup;
-	FVehicleGearData GearData = VehicleMovement->GearSetup;
-
-	uint32 NumChildren = 0;
-	AutoBoxSetup->GetNumChildren(NumChildren);
-	for (uint32 ChildIndex = 0; ChildIndex < NumChildren; ++ChildIndex)
-	{
-		TSharedPtr<IPropertyHandle> ChildProperty = AutoBoxSetup->GetChildHandle(ChildIndex);
-		// Only permit modifying bSimulatePhysics when the body has some geometry.
-		if (ChildProperty->GetProperty()->GetName() == TEXT("ForwardGearAutoBox"))
-		{
-			for (int32 GearIdx = 0; GearIdx != GearData.ForwardGearRatios.Num(); ++GearIdx)
-			{
-				const FGearUpDownRatio & GearUpDownRatio = AutoBox.ForwardGearAutoBox[GearIdx + FVehicleGearData::eFIRST];
-				FText Label = FText::Format(LOCTEXT("AutoBoxGear", "Gear {0}"), FText::AsNumber(GearIdx + 1));
-
-				VehicleSetupCategory.AddCustomRow(Label.ToString(), false)
-				.NameContent()
-					[
-						SNew(STextBlock)
-						.Text(Label)
-						.ToolTipText(LOCTEXT("MassInKG", "Mass of the body in KG"))
-						.Font(DetailFont)
-					]
-				.ValueContent()
-					[
-						SNew(SHorizontalBox)
-						+ SHorizontalBox::Slot()
-						.FillWidth(0.5f)
-						[
-							SNew(SNumericEntryBox<float>)
-							.Value(GearUpDownRatio.DownRatio)
-							.MinValue(0.f)
-							.MaxValue(1.f)
-							.Font(DetailFont)
-							.OnValueCommitted(this, &FWheeledVehicleMovementComponent4WDetails::SetGearShiftRatio, false, GearIdx)
-							.LabelVAlign(VAlign_Center)
-							.Label()
-								[
-									SNew(STextBlock)
-									.Font(DetailFont)
-									.Text(LOCTEXT("AutoBoxShiftMin", "Min"))
-								]
-							.ToolTipText(LOCTEXT("AutoBoxShiftMinTip", "Minimum currentRPM/maxRPM ratio required to stay in this gear. Going lower will cause autobox to gear down"))
-						]	
-						+ SHorizontalBox::Slot()
-						.FillWidth(0.5f)
-						.Padding(4, 0)
-						[
-							SNew(SNumericEntryBox<float>)
-							.Value(GearUpDownRatio.UpRatio)
-							.MinValue(0.f)
-							.MaxValue(1.f)
-							.Font(DetailFont)
-							.OnValueCommitted(this, &FWheeledVehicleMovementComponent4WDetails::SetGearShiftRatio, true, GearIdx)
-							.LabelVAlign(VAlign_Center)
-							.Label()
-							[
-								SNew(STextBlock)
-								.Font(DetailFont)
-								.Text(LOCTEXT("AutoBoxShiftMax", "Max"))
-							]
-							.ToolTipText(LOCTEXT("AutoBoxShiftMinTip", "Maximum currentRPM/maxRPM ratio required to stay in this gear. Going higher will cause autobox to gear up"))
-						]
-					];
-
-			}
-		}
-		else
-		{
-			VehicleSetupCategory.AddProperty(ChildProperty);
-		}
-	}*/
+	SteeringCurveWidget->SetCurveOwner(&SteeringCurveEditor);
 }
 
-void FWheeledVehicleMovementComponent4WDetails::CustomizeDetails(IDetailLayoutBuilder& DetailBuilder)
+TArray<FRichCurveEditInfoConst> FWheeledVehicleMovementComponent4WDetails::FSteeringCurveEditor::GetCurves() const
 {
-	DetailBuilder.GetObjectsBeingCustomized(SelectedObjects);
-	if (SelectedObjects.Num() != 1)	//only do customization if we have one object selected
-	{
-		return;
-	}
+	TArray<FRichCurveEditInfoConst> Curves;
+	Curves.Add(&VehicleComponent->SteeringCurve.EditorCurveData);
 
-	UWheeledVehicleMovementComponent4W * VehicleMovement = SelectedObjects[0].IsValid() ? Cast<UWheeledVehicleMovementComponent4W>(SelectedObjects[0].Get()) : NULL;
+	return Curves;
+}
 
-	// See which categories are hidden
-	TArray<FString> HideCategories;
-	DetailBuilder.GetDetailsView().GetBaseClass()->GetHideCategories(HideCategories);
+TArray<FRichCurveEditInfo> FWheeledVehicleMovementComponent4WDetails::FSteeringCurveEditor::GetCurves()
+{
+	TArray<FRichCurveEditInfo> Curves;
+	Curves.Add(&VehicleComponent->SteeringCurve.EditorCurveData);
 	
+	return Curves;
+}
 
-	//Vehicle Setup
-	if (!HideCategories.Contains(TEXT("VehicleSetup")))
-	{
-		IDetailCategoryBuilder& VehicleSetupCategory = DetailBuilder.EditCategory("VehicleSetup");
-		//Autobox
-		//FDetailWidgetRow & AutoBoxRow = VehicleSetupCategory.AddProperty(DetailBuilder.GetProperty("AutoBoxSetup")).CustomWidget();
-		//FillAutoBoxWidget(VehicleMovement, AutoBoxRow);
-	}
+UObject * FWheeledVehicleMovementComponent4WDetails::FSteeringCurveEditor::GetOwner()
+{
+	return Owner;
 }
 
-void FWheeledVehicleMovementComponent4WDetails::SetGearShiftRatio(float NewValue, ETextCommit::Type, bool bShiftUp, int32 GearIdx)
+void FWheeledVehicleMovementComponent4WDetails::FSteeringCurveEditor::ModifyOwner()
 {
-	/*if (SelectedObjects.Num() != 1)	//only do customization if we have one object selected
+	if (Owner)
 	{
-		return;
+		Owner->Modify();
 	}
+}
 
-	if (UWheeledVehicleMovementComponent4W * VehicleMovement = SelectedObjects[0].IsValid() ? Cast<UWheeledVehicleMovementComponent4W>(SelectedObjects[0].Get()) : NULL)
+void FWheeledVehicleMovementComponent4WDetails::FSteeringCurveEditor::MakeTransactional()
+{
+	if (Owner)
 	{
-		FVehicleAutoBoxData & AutoBox = VehicleMovement->AutoBoxSetup;
-		FGearUpDownRatio & GearUpDownRatio = AutoBox.ForwardGearAutoBox[GearIdx + FVehicleGearData::eFIRST];
-		
-		if (bShiftUp)
-		{
-			GearUpDownRatio.UpRatio = NewValue;
-		}
-		else
-		{
-			GearUpDownRatio.DownRatio = NewValue;
-		}
-
-	}*/
+		Owner->SetFlags(Owner->GetFlags() | RF_Transactional);
+	}
+}
 
+FWheeledVehicleMovementComponent4WDetails::FSteeringCurveEditor::FSteeringCurveEditor(UWheeledVehicleMovementComponent4W * InVehicle)
+{
+	VehicleComponent = InVehicle;
+	Owner = VehicleComponent;
 }
 
 #undef LOCTEXT_NAMESPACE
diff --git a/Engine/Source/Editor/DetailCustomizations/Private/WheeledVehicleMovementComponent4WDetails.h b/Engine/Source/Editor/DetailCustomizations/Private/WheeledVehicleMovementComponent4WDetails.h
--- a/Engine/Source/Editor/DetailCustomizations/Private/WheeledVehicleMovementComponent4WDetails.h
+++ b/Engine/Source/Editor/DetailCustomizations/Private/WheeledVehicleMovementComponent4WDetails.h
@@ -2,6 +2,8 @@
 
 #pragma once
 
+#include "SCurveEditor.h"
+
 class FWheeledVehicleMovementComponent4WDetails : public IDetailCustomization
 {
 public:
@@ -10,10 +12,28 @@ public:
 
 	/** IDetailCustomization interface */
 	virtual void CustomizeDetails( IDetailLayoutBuilder& DetailBuilder ) OVERRIDE;
+	
 private:
-	void SetGearShiftRatio(float NewValue, ETextCommit::Type CommitInfo, bool bShiftUp, int32 GearIdx);
-	void FillAutoBoxWidget(UWheeledVehicleMovementComponent4W * VehicleMovement, FDetailWidgetRow & AutoBoxSetup);
+
+	struct FSteeringCurveEditor : public FCurveOwnerInterface
+	{
+		FSteeringCurveEditor(UWheeledVehicleMovementComponent4W * InVehicle = NULL);
+		/** FCurveOwnerInterface interface */
+		virtual TArray<FRichCurveEditInfoConst> GetCurves() const OVERRIDE;
+		virtual TArray<FRichCurveEditInfo> GetCurves() OVERRIDE;
+		virtual UObject* GetOwner() OVERRIDE;
+		virtual void ModifyOwner() OVERRIDE;
+		virtual void MakeTransactional() OVERRIDE;
+
+	private:
+		UWheeledVehicleMovementComponent4W * VehicleComponent;
+		UObject * Owner;
+
+	} SteeringCurveEditor;
 
 	TArray<TWeakObjectPtr<UObject> > SelectedObjects;	//the objects we're showing details for
+	
+	/** Steering curve widget */
+	TSharedPtr<class SCurveEditor> SteeringCurveWidget;
 };
 
