DetailsView: added support for VisibleAnywhere and EditConst in the FGuid details customization (TTP# 334638 )

diff --git a/Engine/Source/Editor/DetailCustomizations/Private/GuidStructCustomization.cpp b/Engine/Source/Editor/DetailCustomizations/Private/GuidStructCustomization.cpp
--- a/Engine/Source/Editor/DetailCustomizations/Private/GuidStructCustomization.cpp
+++ b/Engine/Source/Editor/DetailCustomizations/Private/GuidStructCustomization.cpp
@@ -28,52 +28,65 @@ void FGuidStructCustomization::CustomizeStructHeader( TSharedRef<class IProperty
 	PropertyHandle = StructPropertyHandle;
 	InputValid = true;
 
-	// create quick-set menu
-	FMenuBuilder QuickSetMenuBuilder(true, NULL);
-	{
-		FUIAction GenerateAction(FExecuteAction::CreateSP(this, &FGuidStructCustomization::HandleGuidActionClicked, EPropertyEditorGuidActions::Generate));
-		QuickSetMenuBuilder.AddMenuEntry(LOCTEXT("GenerateAction", "Generate"), LOCTEXT("GenerateActionHint", "Generate a new random globally unique identifier (GUID)."), FSlateIcon(), GenerateAction);
+	TSharedPtr<SWidget> QuickSetSlotContent;
 
-		FUIAction InvalidateAction(FExecuteAction::CreateSP(this, &FGuidStructCustomization::HandleGuidActionClicked, EPropertyEditorGuidActions::Invalidate));
-		QuickSetMenuBuilder.AddMenuEntry(LOCTEXT("InvalidateAction", "Invalidate"), LOCTEXT("InvalidateActionHint", "Set an invalid globally unique identifier (GUID)."), FSlateIcon(), InvalidateAction);
+	// create quick-set menu if needed
+	if (PropertyHandle->IsEditConst())
+	{
+		QuickSetSlotContent = SNullWidget::NullWidget;
+	}
+	else
+	{
+		FMenuBuilder QuickSetMenuBuilder(true, nullptr);
+		{
+			FUIAction GenerateAction(FExecuteAction::CreateSP(this, &FGuidStructCustomization::HandleGuidActionClicked, EPropertyEditorGuidActions::Generate));
+			QuickSetMenuBuilder.AddMenuEntry(LOCTEXT("GenerateAction", "Generate"), LOCTEXT("GenerateActionHint", "Generate a new random globally unique identifier (GUID)."), FSlateIcon(), GenerateAction);
+
+			FUIAction InvalidateAction(FExecuteAction::CreateSP(this, &FGuidStructCustomization::HandleGuidActionClicked, EPropertyEditorGuidActions::Invalidate));
+			QuickSetMenuBuilder.AddMenuEntry(LOCTEXT("InvalidateAction", "Invalidate"), LOCTEXT("InvalidateActionHint", "Set an invalid globally unique identifier (GUID)."), FSlateIcon(), InvalidateAction);
+		}
+
+		QuickSetSlotContent = SNew(SComboButton)
+			.ContentPadding(FMargin(6.0, 2.0))
+			.MenuContent()
+			[
+				QuickSetMenuBuilder.MakeWidget()
+			];
 	}
 
 	// create struct header
-	HeaderRow.NameContent()
-	[
-		StructPropertyHandle->CreatePropertyNameWidget()
-	]
-	.ValueContent()
-	.MinDesiredWidth(325.0f)
-	.MaxDesiredWidth(325.0f)
-	[
-		SNew(SHorizontalBox)
-
-		+ SHorizontalBox::Slot()
-			.FillWidth(1.0f)
-			[
-				// text box
-				SAssignNew(TextBox, SEditableTextBox)
-					.ClearKeyboardFocusOnCommit(false)
-					.ForegroundColor(this, &FGuidStructCustomization::HandleTextBoxForegroundColor)
-					.OnTextChanged(this, &FGuidStructCustomization::HandleTextBoxTextChanged)
-					.OnTextCommitted(this, &FGuidStructCustomization::HandleTextBoxTextCommited)
-					.SelectAllTextOnCommit(true)
-					.Text(this, &FGuidStructCustomization::HandleTextBoxText)
-			]
-
-		+ SHorizontalBox::Slot()
-			.AutoWidth()
-			[
-				// quick-set menu
-				SNew(SComboButton)
-					.ContentPadding(FMargin(6.0, 2.0))
-					.MenuContent()
-					[
-						QuickSetMenuBuilder.MakeWidget()
-					]
-			]
-	];
+	HeaderRow
+		.NameContent()
+		[
+			StructPropertyHandle->CreatePropertyNameWidget()
+		]
+		.ValueContent()
+		.MinDesiredWidth(325.0f)
+		.MaxDesiredWidth(325.0f)
+		[
+			SNew(SHorizontalBox)
+
+			+ SHorizontalBox::Slot()
+				.FillWidth(1.0f)
+				[
+					// text box
+					SAssignNew(TextBox, SEditableTextBox)
+						.ClearKeyboardFocusOnCommit(false)
+						.IsEnabled(!PropertyHandle->IsEditConst())
+						.ForegroundColor(this, &FGuidStructCustomization::HandleTextBoxForegroundColor)
+						.OnTextChanged(this, &FGuidStructCustomization::HandleTextBoxTextChanged)
+						.OnTextCommitted(this, &FGuidStructCustomization::HandleTextBoxTextCommited)
+						.SelectAllTextOnCommit(true)
+						.Text(this, &FGuidStructCustomization::HandleTextBoxText)
+				]
+
+			+ SHorizontalBox::Slot()
+				.AutoWidth()
+				[
+					// quick set menu
+					QuickSetSlotContent.ToSharedRef()
+				]
+		];
 }
 
 
