Added support to configure which target RHI's we desire on platforms.
User can now select target RHI's for windows under: Edit->Project Preferences->Windows->Targeted RHIs

Split up GetShaderFormats into GetAllPossibleShaderFormats and GetAllTargetedShaderFormats.
D3D11 and D3D10 remain the only two on, by default, for windows.

diff --git a/Engine/Config/BaseEngine.ini b/Engine/Config/BaseEngine.ini
--- a/Engine/Config/BaseEngine.ini
+++ b/Engine/Config/BaseEngine.ini
@@ -2051,3 +2051,7 @@ DefaultGravityZ=-980.0
 
 [/Script/Engine.WorldComposition]
 TilesStreamingTimeThreshold=1.0
+
+[/Script/WindowsTargetPlatform.WindowsTargetSettings]
++TargetedRHIs=PCD3D_SM5
++TargetedRHIs=PCD3D_SM4
diff --git a/Engine/Source/Developer/Android/AndroidTargetPlatform/Private/AndroidTargetPlatform.h b/Engine/Source/Developer/Android/AndroidTargetPlatform/Private/AndroidTargetPlatform.h
--- a/Engine/Source/Developer/Android/AndroidTargetPlatform/Private/AndroidTargetPlatform.h
+++ b/Engine/Source/Developer/Android/AndroidTargetPlatform/Private/AndroidTargetPlatform.h
@@ -118,7 +118,9 @@ public:
 		OutFormats.Add(FName(TEXT("EncodedHDR")));
 	}
 
-	virtual void GetShaderFormats( TArray<FName>& OutFormats ) const OVERRIDE;
+	virtual void GetAllPossibleShaderFormats( TArray<FName>& OutFormats ) const OVERRIDE;
+
+	virtual void GetAllTargetedShaderFormats(TArray<FName>& OutFormats) const OVERRIDE;
 
 	virtual const class FStaticMeshLODSettings& GetStaticMeshLODSettings() const OVERRIDE;
 
diff --git a/Engine/Source/Developer/Android/AndroidTargetPlatform/Private/AndroidTargetPlatform.inl b/Engine/Source/Developer/Android/AndroidTargetPlatform/Private/AndroidTargetPlatform.inl
--- a/Engine/Source/Developer/Android/AndroidTargetPlatform/Private/AndroidTargetPlatform.inl
+++ b/Engine/Source/Developer/Android/AndroidTargetPlatform/Private/AndroidTargetPlatform.inl
@@ -161,13 +161,19 @@ bool FAndroidTargetPlatform::SupportsFeature( ETargetPlatformFeatures::Type Feat
 
 #if WITH_ENGINE
 
-inline void FAndroidTargetPlatform::GetShaderFormats( TArray<FName>& OutFormats ) const
+inline void FAndroidTargetPlatform::GetAllPossibleShaderFormats( TArray<FName>& OutFormats ) const
 {
 	static FName NAME_OPENGL_ES2(TEXT("GLSL_ES2"));
 	OutFormats.AddUnique(NAME_OPENGL_ES2);
 }
 
 
+inline void FAndroidTargetPlatform::GetAllTargetedShaderFormats( TArray<FName>& OutFormats ) const
+{
+	GetAllPossibleShaderFormats(OutFormats);
+}
+
+
 inline const FStaticMeshLODSettings& FAndroidTargetPlatform::GetStaticMeshLODSettings( ) const
 {
 	return StaticMeshLODSettings;
diff --git a/Engine/Source/Developer/HTML5/HTML5TargetPlatform/Private/HTML5TargetPlatform.cpp b/Engine/Source/Developer/HTML5/HTML5TargetPlatform/Private/HTML5TargetPlatform.cpp
--- a/Engine/Source/Developer/HTML5/HTML5TargetPlatform/Private/HTML5TargetPlatform.cpp
+++ b/Engine/Source/Developer/HTML5/HTML5TargetPlatform/Private/HTML5TargetPlatform.cpp
@@ -87,12 +87,18 @@ bool FHTML5TargetPlatform::IsRunningPlatform( ) const
 
 static FName NAME_OPENGL_ES2_WEBGL(TEXT("GLSL_ES2_WEBGL"));
 
-void FHTML5TargetPlatform::GetShaderFormats( TArray<FName>& OutFormats ) const
+void FHTML5TargetPlatform::GetAllPossibleShaderFormats( TArray<FName>& OutFormats ) const
 {
 	OutFormats.AddUnique(NAME_OPENGL_ES2_WEBGL);
 }
 
 
+void FHTML5TargetPlatform::GetAllTargetedShaderFormats( TArray<FName>& OutFormats ) const
+{
+	GetAllPossibleShaderFormats(OutFormats);
+}
+
+
 const class FStaticMeshLODSettings& FHTML5TargetPlatform::GetStaticMeshLODSettings( ) const
 {
 	return StaticMeshLODSettings;
diff --git a/Engine/Source/Developer/HTML5/HTML5TargetPlatform/Private/HTML5TargetPlatform.h b/Engine/Source/Developer/HTML5/HTML5TargetPlatform/Private/HTML5TargetPlatform.h
--- a/Engine/Source/Developer/HTML5/HTML5TargetPlatform/Private/HTML5TargetPlatform.h
+++ b/Engine/Source/Developer/HTML5/HTML5TargetPlatform/Private/HTML5TargetPlatform.h
@@ -44,7 +44,9 @@ public:
 	virtual bool IsRunningPlatform( ) const OVERRIDE;
 
 #if WITH_ENGINE
-	virtual void GetShaderFormats( TArray<FName>& OutFormats ) const OVERRIDE;
+	virtual void GetAllPossibleShaderFormats( TArray<FName>& OutFormats ) const OVERRIDE;
+
+	virtual void GetAllTargetedShaderFormats( TArray<FName>& OutFormats ) const OVERRIDE;
 
 	virtual const class FStaticMeshLODSettings& GetStaticMeshLODSettings() const OVERRIDE;
 
diff --git a/Engine/Source/Developer/Linux/LinuxTargetPlatform/Private/LinuxTargetPlatform.h b/Engine/Source/Developer/Linux/LinuxTargetPlatform/Private/LinuxTargetPlatform.h
--- a/Engine/Source/Developer/Linux/LinuxTargetPlatform/Private/LinuxTargetPlatform.h
+++ b/Engine/Source/Developer/Linux/LinuxTargetPlatform/Private/LinuxTargetPlatform.h
@@ -114,7 +114,7 @@ public:
 	}
 
 #if WITH_ENGINE
-	virtual void GetShaderFormats( TArray<FName>& OutFormats ) const OVERRIDE
+	virtual void GetAllPossibleShaderFormats( TArray<FName>& OutFormats ) const OVERRIDE
 	{
 		// no shaders needed for dedicated server target
 		if (!IS_DEDICATED_SERVER)
@@ -127,6 +127,11 @@ public:
 		}
 	}
 
+	virtual void GetAllTargetedShaderFormats( TArray<FName>& OutFormats ) const OVERRIDE
+	{
+		GetAllPossibleShaderFormats( OutFormats );
+	}
+
 	virtual const class FStaticMeshLODSettings& GetStaticMeshLODSettings( ) const OVERRIDE
 	{
 		return StaticMeshLODSettings;
diff --git a/Engine/Source/Developer/Mac/MacTargetPlatform/Private/GenericMacTargetPlatform.h b/Engine/Source/Developer/Mac/MacTargetPlatform/Private/GenericMacTargetPlatform.h
--- a/Engine/Source/Developer/Mac/MacTargetPlatform/Private/GenericMacTargetPlatform.h
+++ b/Engine/Source/Developer/Mac/MacTargetPlatform/Private/GenericMacTargetPlatform.h
@@ -133,7 +133,7 @@ return TSuper::SupportsFeature(Feature);
 	}
 
 #if WITH_ENGINE
-	virtual void GetShaderFormats( TArray<FName>& OutFormats ) const OVERRIDE
+	virtual void GetAllPossibleShaderFormats( TArray<FName>& OutFormats ) const OVERRIDE
 	{
 		// no shaders needed for dedicated server target
 		if (!IS_DEDICATED_SERVER)
@@ -143,6 +143,11 @@ return TSuper::SupportsFeature(Feature);
 		}
 	}
 
+	virtual void GetAllTargetedShaderFormats(TArray<FName>& OutFormats) const OVERRIDE
+	{
+		GetAllPossibleShaderFormats( OutFormats );
+	}
+
 
 	virtual const class FStaticMeshLODSettings& GetStaticMeshLODSettings( ) const OVERRIDE
 	{
diff --git a/Engine/Source/Developer/TargetPlatform/Public/Interfaces/ITargetPlatform.h b/Engine/Source/Developer/TargetPlatform/Public/Interfaces/ITargetPlatform.h
--- a/Engine/Source/Developer/TargetPlatform/Public/Interfaces/ITargetPlatform.h
+++ b/Engine/Source/Developer/TargetPlatform/Public/Interfaces/ITargetPlatform.h
@@ -330,11 +330,18 @@ public:
 	virtual void GetReflectionCaptureFormats( TArray<FName>& OutFormats ) const = 0;
 
 	/**
-	 * Gets the shader formats this platform uses.
+	 * Gets the shader formats this platform can use.
 	 *
 	 * @param OutFormats - Will contain the shader formats.
 	 */
-	virtual void GetShaderFormats( TArray<FName>& OutFormats ) const = 0;
+	virtual void GetAllPossibleShaderFormats( TArray<FName>& OutFormats ) const = 0;
+
+	/**
+	* Gets the shader formats that have been selected for this target platform
+	*
+	* @param OutFormats - Will contain the shader formats.
+	*/
+	virtual void GetAllTargetedShaderFormats(TArray<FName>& OutFormats) const = 0;
 
 	/**
 	 * Gets the format to use for a particular texture.
diff --git a/Engine/Source/Developer/Windows/WindowsTargetPlatform/Classes/WindowsTargetSettings.h b/Engine/Source/Developer/Windows/WindowsTargetPlatform/Classes/WindowsTargetSettings.h
--- a/Engine/Source/Developer/Windows/WindowsTargetPlatform/Classes/WindowsTargetSettings.h
+++ b/Engine/Source/Developer/Windows/WindowsTargetPlatform/Classes/WindowsTargetSettings.h
@@ -19,4 +19,11 @@ class WINDOWSTARGETPLATFORM_API UWindowsTargetSettings
 public:
 
 	GENERATED_UCLASS_BODY()
+
+	/** 
+	 * The collection of RHI's we want to support on this platform.
+	 * This is not always the full list of RHI we can support.
+	 */
+	UPROPERTY(EditAnywhere, config, Category=Rendering)
+	TArray<FString> TargetedRHIs;
 };
diff --git a/Engine/Source/Developer/Windows/WindowsTargetPlatform/Private/GenericWindowsTargetPlatform.h b/Engine/Source/Developer/Windows/WindowsTargetPlatform/Private/GenericWindowsTargetPlatform.h
--- a/Engine/Source/Developer/Windows/WindowsTargetPlatform/Private/GenericWindowsTargetPlatform.h
+++ b/Engine/Source/Developer/Windows/WindowsTargetPlatform/Private/GenericWindowsTargetPlatform.h
@@ -29,6 +29,22 @@ public:
 			FConfigCacheIni::LoadLocalIniFile(EngineSettings, TEXT("Engine"), true, *PlatformName());
 			TextureLODSettings.Initialize(EngineSettings, TEXT("SystemSettings"));
 			StaticMeshLODSettings.Initialize(EngineSettings);
+
+			// Get the Target RHIs for this platform, we do not always want all those that are supported.
+			GConfig->GetArray(TEXT("/Script/WindowsTargetPlatform.WindowsTargetSettings"), TEXT("TargetedRHIs"), TargetedShaderFormats, GEngineIni);
+			
+			// Gather the list of Target RHIs and filter out any that may be invalid.
+			TArray<FName> PossibleShaderFormats;
+			GetAllPossibleShaderFormats(PossibleShaderFormats);
+
+			for(int32 ShaderFormatIdx = TargetedShaderFormats.Num()-1; ShaderFormatIdx >= 0; ShaderFormatIdx--)
+			{
+				FString ShaderFormat = TargetedShaderFormats[ShaderFormatIdx];
+				if(PossibleShaderFormats.Contains(FName(*ShaderFormat)) == false)
+				{
+					TargetedShaderFormats.Remove(ShaderFormat);
+				}
+			}
 		#endif
 	}
 
@@ -129,7 +145,7 @@ public:
 	}
 
 #if WITH_ENGINE
-	virtual void GetShaderFormats( TArray<FName>& OutFormats ) const OVERRIDE
+	virtual void GetAllPossibleShaderFormats( TArray<FName>& OutFormats ) const OVERRIDE
 	{
 		// no shaders needed for dedicated server target
 		if (!IS_DEDICATED_SERVER)
@@ -139,9 +155,21 @@ public:
 			//       may not wish to support SM4 or OpenGL.
 			static FName NAME_PCD3D_SM5(TEXT("PCD3D_SM5"));
 			static FName NAME_PCD3D_SM4(TEXT("PCD3D_SM4"));
+			static FName NAME_GLSL_150(TEXT("GLSL_150"));
+			static FName NAME_GLSL_430(TEXT("GLSL_430"));
 
 			OutFormats.AddUnique(NAME_PCD3D_SM5);
 			OutFormats.AddUnique(NAME_PCD3D_SM4);
+			OutFormats.AddUnique(NAME_GLSL_150);
+			OutFormats.AddUnique(NAME_GLSL_430);
+		}
+	}
+
+	virtual void GetAllTargetedShaderFormats( TArray<FName>& OutFormats ) const OVERRIDE
+	{
+		for(const FString& ShaderFormat : TargetedShaderFormats)
+		{
+			OutFormats.AddUnique(FName(*ShaderFormat));
 		}
 	}
 
@@ -215,6 +243,9 @@ private:
 
 	// Holds static mesh LOD settings.
 	FStaticMeshLODSettings StaticMeshLODSettings;
+
+	// List of shader formats specified as targets
+	TArray<FString> TargetedShaderFormats;
 #endif // WITH_ENGINE
 
 private:
diff --git a/Engine/Source/Developer/Windows/WindowsTargetPlatform/Private/WindowsTargetPlatformModule.cpp b/Engine/Source/Developer/Windows/WindowsTargetPlatform/Private/WindowsTargetPlatformModule.cpp
--- a/Engine/Source/Developer/Windows/WindowsTargetPlatform/Private/WindowsTargetPlatformModule.cpp
+++ b/Engine/Source/Developer/Windows/WindowsTargetPlatform/Private/WindowsTargetPlatformModule.cpp
@@ -73,6 +73,9 @@ public:
 #endif
 
 		TargetSettings = ConstructObject<UWindowsTargetSettings>(UWindowsTargetSettings::StaticClass(), GetTransientPackage(), "WindowsTargetSettings", RF_Standalone);
+		
+		// We need to manually load the config properties here, as this module is loaded before the UObject system is setup to do this
+		GConfig->GetArray(TEXT("/Script/WindowsTargetPlatform.WindowsTargetSettings"), TEXT("TargetedRHIs"), TargetSettings->TargetedRHIs, GEngineIni);
 		TargetSettings->AddToRoot();
 
 		ISettingsModule* SettingsModule = ISettingsModule::Get();
diff --git a/Engine/Source/Developer/iOS/IOSTargetPlatform/Private/IOSTargetPlatform.cpp b/Engine/Source/Developer/iOS/IOSTargetPlatform/Private/IOSTargetPlatform.cpp
--- a/Engine/Source/Developer/iOS/IOSTargetPlatform/Private/IOSTargetPlatform.cpp
+++ b/Engine/Source/Developer/iOS/IOSTargetPlatform/Private/IOSTargetPlatform.cpp
@@ -217,13 +217,19 @@ bool FIOSTargetPlatform::HandleTicker(float DeltaTime )
 
 #if WITH_ENGINE
 
-void FIOSTargetPlatform::GetShaderFormats( TArray<FName>& OutFormats ) const
+void FIOSTargetPlatform::GetAllPossibleShaderFormats( TArray<FName>& OutFormats ) const
 {
 	static FName NAME_OPENGL_ES2_IOS(TEXT("GLSL_ES2_IOS"));
 	OutFormats.AddUnique(NAME_OPENGL_ES2_IOS);
 }
 
 
+void FIOSTargetPlatform::GetAllTargetedShaderFormats( TArray<FName>& OutFormats ) const
+{
+	GetAllPossibleShaderFormats(OutFormats);
+}
+
+
 void FIOSTargetPlatform::GetTextureFormats( const UTexture* Texture, TArray<FName>& OutFormats ) const
 {
 	check(Texture);
diff --git a/Engine/Source/Developer/iOS/IOSTargetPlatform/Private/IOSTargetPlatform.h b/Engine/Source/Developer/iOS/IOSTargetPlatform/Private/IOSTargetPlatform.h
--- a/Engine/Source/Developer/iOS/IOSTargetPlatform/Private/IOSTargetPlatform.h
+++ b/Engine/Source/Developer/iOS/IOSTargetPlatform/Private/IOSTargetPlatform.h
@@ -91,7 +91,9 @@ public:
 		OutFormats.Add(FName(TEXT("EncodedHDR")));
 	}
 
-	virtual void GetShaderFormats( TArray<FName>& OutFormats ) const OVERRIDE;
+	virtual void GetAllPossibleShaderFormats( TArray<FName>& OutFormats ) const OVERRIDE;
+
+	virtual void GetAllTargetedShaderFormats( TArray<FName>& OutFormats ) const OVERRIDE;
 
 	virtual const class FStaticMeshLODSettings& GetStaticMeshLODSettings( ) const OVERRIDE
 	{
diff --git a/Engine/Source/Editor/DetailCustomizations/Private/WindowsTargetSettingsDetails.cpp b/Engine/Source/Editor/DetailCustomizations/Private/WindowsTargetSettingsDetails.cpp
--- a/Engine/Source/Editor/DetailCustomizations/Private/WindowsTargetSettingsDetails.cpp
+++ b/Engine/Source/Editor/DetailCustomizations/Private/WindowsTargetSettingsDetails.cpp
@@ -19,6 +19,30 @@ namespace WindowsTargetSettingsDetailsConstants
 
 #define LOCTEXT_NAMESPACE "WindowsTargetSettingsDetails"
 
+FText GetFriendlyNameFromRHIName(const FString& InRHIName)
+{
+	FText FriendlyRHIName = LOCTEXT("UnknownRHI", "UnknownRHI");
+	if (InRHIName == TEXT("PCD3D_SM5"))
+	{
+		FriendlyRHIName = LOCTEXT("DirectX11", "DirectX 11 (SM5)");
+	}
+	else if (InRHIName == TEXT("PCD3D_SM4"))
+	{
+		FriendlyRHIName = LOCTEXT("DirectX10", "DirectX 10 (SM4)");
+	}
+	else if (InRHIName == TEXT("GLSL_150"))
+	{
+		FriendlyRHIName = LOCTEXT("OpenGL3", "OpenGL 3 (SM4)");
+	}
+	else if (InRHIName == TEXT("GLSL_430"))
+	{
+		FriendlyRHIName = LOCTEXT("OpenGL4", "OpenGL 4 (SM5, Experimental)");
+	}
+
+	return FriendlyRHIName;
+}
+
+
 TSharedRef<IDetailCustomization> FWindowsTargetSettingsDetails::MakeInstance()
 {
 	return MakeShareable(new FWindowsTargetSettingsDetails);
@@ -80,9 +104,13 @@ static FString GetIconFilename(EImageScope::Type Scope)
 
 void FWindowsTargetSettingsDetails::CustomizeDetails( IDetailLayoutBuilder& DetailBuilder )
 {
-	IDetailCategoryBuilder& DetailCategoryBuilder = DetailBuilder.EditCategory(TEXT("Splash"));
+	// Setup the supported/targeted RHI property view
+	TargetShaderFormatsDetails = MakeShareable(new FTargetShaderFormatsPropertyDetails(&DetailBuilder));
+	TargetShaderFormatsDetails->CreateTargetShaderFormatsPropertyView();
 
-	FDetailWidgetRow& EditorSplashWidgetRow = DetailCategoryBuilder.AddCustomRow(TEXT("Editor Splash"));
+	// Next add the splash image customization
+	IDetailCategoryBuilder& SplashCategoryBuilder = DetailBuilder.EditCategory(TEXT("Splash"));
+	FDetailWidgetRow& EditorSplashWidgetRow = SplashCategoryBuilder.AddCustomRow(TEXT("Editor Splash"));
 
 	const FText EditorSplashDesc(LOCTEXT("EditorSplashLabel", "Editor Splash"));
 	const FString EditorSplash_TargetImagePath = GetSplashFilename(EImageScope::GameOverride, true);
@@ -117,7 +145,7 @@ void FWindowsTargetSettingsDetails::CustomizeDetails( IDetailLayoutBuilder& Deta
 		]
 	];
 
-	FDetailWidgetRow& GameSplashWidgetRow = DetailCategoryBuilder.AddCustomRow(TEXT("Game Splash"));
+	FDetailWidgetRow& GameSplashWidgetRow = SplashCategoryBuilder.AddCustomRow(TEXT("Game Splash"));
 
 	const FText GameSplashDesc(LOCTEXT("GameSplashLabel", "Game Splash"));
 	const FString GameSplash_TargetImagePath = GetSplashFilename(EImageScope::GameOverride, false);
@@ -218,4 +246,95 @@ bool FWindowsTargetSettingsDetails::HandlePostExternalIconCopy(const FString& In
 	return true;
 }
 
+FTargetShaderFormatsPropertyDetails::FTargetShaderFormatsPropertyDetails(IDetailLayoutBuilder* InDetailBuilder)
+: DetailBuilder(InDetailBuilder)
+{
+	TargetShaderFormatsPropertyHandle = DetailBuilder->GetProperty("TargetedRHIs");
+	ensure(TargetShaderFormatsPropertyHandle.IsValid());
+}
+
+void FTargetShaderFormatsPropertyDetails::CreateTargetShaderFormatsPropertyView()
+{
+	DetailBuilder->HideProperty(TargetShaderFormatsPropertyHandle);
+
+	// List of supported RHI's and selected targets
+	ITargetPlatform* WindowsTargetPlatform = FModuleManager::GetModuleChecked<ITargetPlatformModule>("WindowsTargetPlatform").GetTargetPlatform();
+	TArray<FName> ShaderFormats;
+	WindowsTargetPlatform->GetAllPossibleShaderFormats(ShaderFormats);
+
+	IDetailCategoryBuilder& TargetedRHICategoryBuilder = DetailBuilder->EditCategory(TEXT("Targeted RHIs"));
+
+	for (const FName& ShaderFormat : ShaderFormats)
+	{
+		FDetailWidgetRow& TargetedRHIWidgetRow = TargetedRHICategoryBuilder.AddCustomRow(ShaderFormat.ToString());
+
+		FText FriendlyShaderFormatName = GetFriendlyNameFromRHIName(ShaderFormat.ToString());
+
+		TargetedRHIWidgetRow
+		.NameContent()
+		[
+			SNew(SHorizontalBox)
+			+ SHorizontalBox::Slot()
+			.Padding(FMargin(0, 1, 0, 1))
+			.FillWidth(1.0f)
+			[
+				SNew(STextBlock)
+				.Text(FriendlyShaderFormatName)
+				.Font(DetailBuilder->GetDetailFont())
+			]
+		]
+		.ValueContent()
+		[
+			SNew(SCheckBox)
+			.OnCheckStateChanged(this, &FTargetShaderFormatsPropertyDetails::OnTargetedRHIChanged, ShaderFormat)
+			.IsChecked(this, &FTargetShaderFormatsPropertyDetails::IsTargetedRHIChecked, ShaderFormat)
+		];
+	}
+}
+
+
+void FTargetShaderFormatsPropertyDetails::OnTargetedRHIChanged(ESlateCheckBoxState::Type InNewValue, FName InRHIName)
+{
+	TArray<void*> RawPtrs;
+	TargetShaderFormatsPropertyHandle->AccessRawData(RawPtrs);
+
+	// Update the CVars with the selection
+	{
+		TargetShaderFormatsPropertyHandle->NotifyPreChange();
+		for (void* RawPtr : RawPtrs)
+		{
+			TArray<FString>& Array = *(TArray<FString>*)RawPtr;
+			if(InNewValue == ESlateCheckBoxState::Checked)
+			{
+				Array.Add(InRHIName.ToString());
+			}
+			else
+			{
+				Array.Remove(InRHIName.ToString());
+			}
+		}
+		TargetShaderFormatsPropertyHandle->NotifyPostChange();
+	}
+}
+
+
+ESlateCheckBoxState::Type FTargetShaderFormatsPropertyDetails::IsTargetedRHIChecked(FName InRHIName) const
+{
+	ESlateCheckBoxState::Type CheckState = ESlateCheckBoxState::Unchecked;
+
+	TArray<void*> RawPtrs;
+	TargetShaderFormatsPropertyHandle->AccessRawData(RawPtrs);
+	
+	for(void* RawPtr : RawPtrs)
+	{
+		TArray<FString>& Array = *(TArray<FString>*)RawPtr;
+		if(Array.Contains(InRHIName.ToString()))
+		{
+			CheckState = ESlateCheckBoxState::Checked;
+		}
+	}
+	return CheckState;
+}
+
+
 #undef LOCTEXT_NAMESPACE
diff --git a/Engine/Source/Editor/DetailCustomizations/Private/WindowsTargetSettingsDetails.h b/Engine/Source/Editor/DetailCustomizations/Private/WindowsTargetSettingsDetails.h
--- a/Engine/Source/Editor/DetailCustomizations/Private/WindowsTargetSettingsDetails.h
+++ b/Engine/Source/Editor/DetailCustomizations/Private/WindowsTargetSettingsDetails.h
@@ -3,6 +3,47 @@
 #pragma once
 
 
+/* FTargetShaderFormatsPropertyDetails
+*****************************************************************************/
+
+/**
+* Helper which implements details panel customizations for a device profiles parent property
+*/
+class FTargetShaderFormatsPropertyDetails
+	: public TSharedFromThis<FTargetShaderFormatsPropertyDetails>
+{
+
+public:
+
+	/**
+	 * Constructor for the parent property details view
+	 *
+	 * @param InDetailsBuilder - Where we are adding our property view to
+	 */
+	FTargetShaderFormatsPropertyDetails(IDetailLayoutBuilder* InDetailBuilder);
+
+
+	/**
+	 * Create the UI to select which windows shader formats we are targetting
+	 */
+	void CreateTargetShaderFormatsPropertyView();
+
+private:
+
+	// Supported/Targeted RHI check boxes
+	void OnTargetedRHIChanged(ESlateCheckBoxState::Type InNewValue, FName InRHIName);
+	ESlateCheckBoxState::Type IsTargetedRHIChecked(FName InRHIName) const;
+
+private:
+
+	/** A handle to the detail view builder */
+	IDetailLayoutBuilder* DetailBuilder;
+
+	/** Access to the Parent Property */
+	TSharedPtr<IPropertyHandle> TargetShaderFormatsPropertyHandle;
+};
+
+
 /**
  * Manages the Transform section of a details view                    
  */
@@ -24,4 +65,8 @@ private:
 
 	/** Delegate handler to set the path to start picking from */
 	bool HandlePostExternalIconCopy(const FString& InChosenImage);
+
+private:
+	/** Reference to the target shader formats property view */
+	TSharedPtr<FTargetShaderFormatsPropertyDetails> TargetShaderFormatsDetails;
 };
\ No newline at end of file
diff --git a/Engine/Source/Editor/UnrealEd/Private/Commandlets/DerivedDataCacheCommandlet.cpp b/Engine/Source/Editor/UnrealEd/Private/Commandlets/DerivedDataCacheCommandlet.cpp
--- a/Engine/Source/Editor/UnrealEd/Private/Commandlets/DerivedDataCacheCommandlet.cpp
+++ b/Engine/Source/Editor/UnrealEd/Private/Commandlets/DerivedDataCacheCommandlet.cpp
@@ -92,7 +92,7 @@ int32 UDerivedDataCacheCommandlet::Main( const FString& Params )
 		for (int32 Index = 0; Index < Platforms.Num(); Index++)
 		{
 			TArray<FName> DesiredShaderFormats;
-			Platforms[Index]->GetShaderFormats(DesiredShaderFormats);
+			Platforms[Index]->GetAllTargetedShaderFormats(DesiredShaderFormats);
 
 			for (int32 FormatIndex = 0; FormatIndex < DesiredShaderFormats.Num(); FormatIndex++)
 			{
diff --git a/Engine/Source/Runtime/Engine/Private/GlobalShader.cpp b/Engine/Source/Runtime/Engine/Private/GlobalShader.cpp
--- a/Engine/Source/Runtime/Engine/Private/GlobalShader.cpp
+++ b/Engine/Source/Runtime/Engine/Private/GlobalShader.cpp
@@ -511,7 +511,7 @@ void RecompileShadersForRemote(
 	}
 
 	TArray<FName> DesiredShaderFormats;
-	TargetPlatform->GetShaderFormats(DesiredShaderFormats);
+	TargetPlatform->GetAllTargetedShaderFormats(DesiredShaderFormats);
 
 	UE_LOG(LogShaders, Display, TEXT("Loading %d materials..."), MaterialsToLoad.Num());
 	// make sure all materials the client has loaded will be processed
diff --git a/Engine/Source/Runtime/Engine/Private/Materials/Material.cpp b/Engine/Source/Runtime/Engine/Private/Materials/Material.cpp
--- a/Engine/Source/Runtime/Engine/Private/Materials/Material.cpp
+++ b/Engine/Source/Runtime/Engine/Private/Materials/Material.cpp
@@ -1987,7 +1987,7 @@ void UMaterial::PostLoad()
 void UMaterial::BeginCacheForCookedPlatformData( const ITargetPlatform *TargetPlatform )
 {
 	TArray<FName> DesiredShaderFormats;
-	TargetPlatform->GetShaderFormats( DesiredShaderFormats );
+	TargetPlatform->GetAllTargetedShaderFormats(DesiredShaderFormats);
 
 	TArray<FMaterialResource*> *CachedMaterialResourcesForPlatform = CachedMaterialResourcesForCooking.Find( TargetPlatform );
 
diff --git a/Engine/Source/Runtime/Engine/Private/Materials/MaterialInstance.cpp b/Engine/Source/Runtime/Engine/Private/Materials/MaterialInstance.cpp
--- a/Engine/Source/Runtime/Engine/Private/Materials/MaterialInstance.cpp
+++ b/Engine/Source/Runtime/Engine/Private/Materials/MaterialInstance.cpp
@@ -1480,7 +1480,7 @@ void TrimToOverriddenOnly(TArray<ParameterType>& Parameters)
 void UMaterialInstance::BeginCacheForCookedPlatformData( const ITargetPlatform *TargetPlatform )
 {
 	TArray<FName> DesiredShaderFormats;
-	TargetPlatform->GetShaderFormats( DesiredShaderFormats );
+	TargetPlatform->GetAllTargetedShaderFormats(DesiredShaderFormats);
 
 	TArray<FMaterialResource*> *CachedMaterialResourcesForPlatform = CachedMaterialResourcesForCooking.Find( TargetPlatform );
 
diff --git a/Engine/Source/Runtime/ShaderCore/Private/Shader.cpp b/Engine/Source/Runtime/ShaderCore/Private/Shader.cpp
--- a/Engine/Source/Runtime/ShaderCore/Private/Shader.cpp
+++ b/Engine/Source/Runtime/ShaderCore/Private/Shader.cpp
@@ -1073,7 +1073,7 @@ const TArray<FName>& GetTargetShaderFormats()
 
 			for (int32 Index = 0; Index < Platforms.Num(); Index++)
 			{
-				Platforms[Index]->GetShaderFormats(Results);
+				Platforms[Index]->GetAllTargetedShaderFormats(Results);
 			}
 		}
 	}
