Fixing (likely) coding errors or confusing debugging helpers.

Contributed by UE4Linux community.

diff --git a/Engine/Source/Developer/FunctionalTesting/Private/FunctionalTest.cpp b/Engine/Source/Developer/FunctionalTesting/Private/FunctionalTest.cpp
--- a/Engine/Source/Developer/FunctionalTesting/Private/FunctionalTest.cpp
+++ b/Engine/Source/Developer/FunctionalTesting/Private/FunctionalTest.cpp
@@ -177,7 +177,7 @@ void AFunctionalTest::SetTimeLimit(float InTimeLimit, TEnumAsByte<EFunctionalTes
 			, TEXT("%s> Trying to set test Result to \'Invalid\'. Falling back to \'Failed\'")
 			, *GetActorLabel());
 
-		InResult == EFunctionalTestResult::Failed;
+		InResult = EFunctionalTestResult::Failed;
 	}
 	TimesUpResult = InResult;
 }
diff --git a/Engine/Source/Editor/ClassViewer/Private/SClassViewer.cpp b/Engine/Source/Editor/ClassViewer/Private/SClassViewer.cpp
--- a/Engine/Source/Editor/ClassViewer/Private/SClassViewer.cpp
+++ b/Engine/Source/Editor/ClassViewer/Private/SClassViewer.cpp
@@ -1868,10 +1868,6 @@ void FClassHierarchy::LoadUnloadedTagData(TSharedPtr<FClassViewerNode>& InOutCla
 				{
 					UnloadedBlueprintData->AddImplementedInterfaces(InterfaceName);
 				}
-				else
-				{
-					ImplementedInterfaces = ImplementedInterfaces;
-				}
 			}
 			CurrentString = RemainingString;
 		}
diff --git a/Engine/Source/Editor/DetailCustomizations/Private/ActorDetails.cpp b/Engine/Source/Editor/DetailCustomizations/Private/ActorDetails.cpp
--- a/Engine/Source/Editor/DetailCustomizations/Private/ActorDetails.cpp
+++ b/Engine/Source/Editor/DetailCustomizations/Private/ActorDetails.cpp
@@ -160,10 +160,6 @@ public:
 
 	virtual bool IsClassAllowed(const FClassViewerInitializationOptions& InInitOptions, const UClass* InClass, TSharedRef< FClassViewerFilterFuncs > InFilterFuncs ) OVERRIDE
 	{
-		if(InClass->GetName() == TEXT("MyLightBlueprint_C"))
-		{
-			InClass = InClass;
-		}
 		EFilterReturn::Type eState = InFilterFuncs->IfInClassesSet(AllowedClasses, InClass);
 		if(eState == EFilterReturn::NoItems)
 		{
diff --git a/Engine/Source/Editor/Kismet/Private/BlueprintDetailsCustomization.cpp b/Engine/Source/Editor/Kismet/Private/BlueprintDetailsCustomization.cpp
--- a/Engine/Source/Editor/Kismet/Private/BlueprintDetailsCustomization.cpp
+++ b/Engine/Source/Editor/Kismet/Private/BlueprintDetailsCustomization.cpp
@@ -3317,7 +3317,7 @@ ESlateCheckBoxState::Type FBlueprintGraphActionDetails::GetIsReliableReplicatedF
 	const UK2Node_CustomEvent* CustomEvent = Cast<const UK2Node_CustomEvent>(FunctionEntryNode);
 	if(!CustomEvent)
 	{
-		ESlateCheckBoxState::Undetermined;
+		return ESlateCheckBoxState::Undetermined;
 	}
 
 	uint32 const NetReliableMask = (FUNC_Net | FUNC_NetReliable);
diff --git a/Engine/Source/Editor/UnrealEd/Private/Lightmass/Lightmass.cpp b/Engine/Source/Editor/UnrealEd/Private/Lightmass/Lightmass.cpp
--- a/Engine/Source/Editor/UnrealEd/Private/Lightmass/Lightmass.cpp
+++ b/Engine/Source/Editor/UnrealEd/Private/Lightmass/Lightmass.cpp
@@ -2768,7 +2768,7 @@ bool FLightmassProcessor::Update()
 	if (!bIsLightmassRunning)
 	{
 		bIsFinished = true;
-		bProcessingFailed == Status != 0;
+		bProcessingFailed = Status != 0;
 		bProcessingSuccessful = !bProcessingFailed;
 		bQuitReceived = true;
 	}
diff --git a/Engine/Source/Editor/UnrealEd/Private/UnrealWidget.cpp b/Engine/Source/Editor/UnrealEd/Private/UnrealWidget.cpp
--- a/Engine/Source/Editor/UnrealEd/Private/UnrealWidget.cpp
+++ b/Engine/Source/Editor/UnrealEd/Private/UnrealWidget.cpp
@@ -952,7 +952,7 @@ void FWidget::ConvertMouseMovementToAxisMovement( FEditorViewportClient* InViewp
 			{
 				if( CurrentAxis == EAxisList::ZRotation )
 				{
-					FVector Axis = Axis = FVector( 0, 0, 1 );
+					FVector Axis = FVector( 0, 0, 1 );
 					Axis = InputCoordSystem.TransformVector( Axis );
 
 					const float RotationSpeed = GetRotationSpeed();
