334254	DOCUMENTATION: Editor: Text displays incorrect number of physical surface types

diff --git a/Engine/Source/Editor/DetailCustomizations/Private/PhysicsSettingsDetails.cpp b/Engine/Source/Editor/DetailCustomizations/Private/PhysicsSettingsDetails.cpp
--- a/Engine/Source/Editor/DetailCustomizations/Private/PhysicsSettingsDetails.cpp
+++ b/Engine/Source/Editor/DetailCustomizations/Private/PhysicsSettingsDetails.cpp
@@ -151,7 +151,7 @@ void FPhysicsSettingsDetails::CustomizeDetails( IDetailLayoutBuilder& DetailBuil
 				.ColorAndOpacity(FSlateColor::UseSubduedForeground())
 				.Font(IDetailLayoutBuilder::GetDetailFont())
 				.ToolTip(PhysicalSurfaceTooltip)
-				.Text(LOCTEXT("PhysicalSurface_Menu_Description", " You can have up to 30 custom surface types for your project. \nOnce you name each type, they will show up as surface type in the physical material."))
+				.Text(LOCTEXT("PhysicalSurface_Menu_Description", " You can have up to 62 custom surface types for your project. \nOnce you name each type, they will show up as surface type in the physical material."))
 			]
 		]
 
