Ensure new Axis and Action Mappings have unique names so that they're not grouped together

TTP#335363 - EDITOR: Adding new input mapping fails if None mapping exists

Changed the add function to use a similar pattern to adding mappings to an existing group. Each add will create a new numbered name from a localized base name and check to see that it hasn't already been used.

diff --git a/Engine/Source/Editor/DetailCustomizations/Private/InputSettingsDetails.cpp b/Engine/Source/Editor/DetailCustomizations/Private/InputSettingsDetails.cpp
--- a/Engine/Source/Editor/DetailCustomizations/Private/InputSettingsDetails.cpp
+++ b/Engine/Source/Editor/DetailCustomizations/Private/InputSettingsDetails.cpp
@@ -126,7 +126,43 @@ void FActionMappingsNodeBuilder::GenerateChildContent( IDetailChildrenBuilder& C
 
 void FActionMappingsNodeBuilder::AddActionMappingButton_OnClick()
 {
-	ActionMappingsPropertyHandle->AsArray()->AddItem();
+	static const FName BaseActionMappingName(*LOCTEXT("NewActionMappingName", "NewActionMapping").ToString());
+	static int32 NewMappingCount = 0;
+	const FScopedTransaction Transaction(LOCTEXT("AddActionMapping_Transaction", "Add Action Mapping"));
+
+	TArray<UObject*> OuterObjects;
+	ActionMappingsPropertyHandle->GetOuterObjects(OuterObjects);
+
+	if (OuterObjects.Num() == 1)
+	{
+		UInputSettings* InputSettings = CastChecked<UInputSettings>(OuterObjects[0]);
+		InputSettings->Modify();
+		ActionMappingsPropertyHandle->NotifyPreChange();
+
+		EName BaseActionMappingIndex = (EName)BaseActionMappingName.GetIndex();
+		FName NewActionMappingName;
+		bool bFoundUniqueName;
+		do
+		{
+			// Create a numbered name and check whether it's already been used
+			NewActionMappingName = FName(BaseActionMappingIndex, ++NewMappingCount);
+			bFoundUniqueName = true;
+			for (int32 Index = 0; Index < InputSettings->ActionMappings.Num(); ++Index)
+			{
+				if (InputSettings->ActionMappings[Index].ActionName == NewActionMappingName)
+				{
+					bFoundUniqueName = false;
+					break;
+				}
+			}
+		}
+		while (!bFoundUniqueName);
+
+		FInputActionKeyMapping NewMapping(NewActionMappingName);
+		InputSettings->ActionMappings.Add(NewMapping);
+
+		ActionMappingsPropertyHandle->NotifyPostChange();
+	}
 }
 
 void FActionMappingsNodeBuilder::ClearActionMappingButton_OnClick()
@@ -359,7 +395,43 @@ void FAxisMappingsNodeBuilder::GenerateChildContent( IDetailChildrenBuilder& Chi
 
 void FAxisMappingsNodeBuilder::AddAxisMappingButton_OnClick()
 {
-	AxisMappingsPropertyHandle->AsArray()->AddItem();
+	static const FName BaseAxisMappingName(*LOCTEXT("NewAxisMappingName", "NewAxisMapping").ToString());
+	static int32 NewMappingCount = 0;
+	const FScopedTransaction Transaction(LOCTEXT("AddAxisMapping_Transaction", "Add Axis Mapping"));
+
+	TArray<UObject*> OuterObjects;
+	AxisMappingsPropertyHandle->GetOuterObjects(OuterObjects);
+
+	if (OuterObjects.Num() == 1)
+	{
+		UInputSettings* InputSettings = CastChecked<UInputSettings>(OuterObjects[0]);
+		InputSettings->Modify();
+		AxisMappingsPropertyHandle->NotifyPreChange();
+
+		EName BaseAxisMappingIndex = (EName)BaseAxisMappingName.GetIndex();
+		FName NewAxisMappingName;
+		bool bFoundUniqueName;
+		do
+		{
+			// Create a numbered name and check whether it's already been used
+			NewAxisMappingName = FName(BaseAxisMappingIndex, ++NewMappingCount);
+			bFoundUniqueName = true;
+			for (int32 Index = 0; Index < InputSettings->AxisMappings.Num(); ++Index)
+			{
+				if (InputSettings->AxisMappings[Index].AxisName == NewAxisMappingName)
+				{
+					bFoundUniqueName = false;
+					break;
+				}
+			}
+		}
+		while (!bFoundUniqueName);
+
+		FInputAxisKeyMapping NewMapping(NewAxisMappingName);
+		InputSettings->AxisMappings.Add(NewMapping);
+
+		AxisMappingsPropertyHandle->NotifyPostChange();
+	}
 }
 
 void FAxisMappingsNodeBuilder::ClearAxisMappingButton_OnClick()
