UAT now uses the game INI files to add additional directories to cook and stage.

TTP# 335070 - Add GUI based way to specify additional non-asset files or paths to package

CookCommandlet now uses the project settings to work out which additional directories it should cook (this removes the need for the editor to pass this information to that UAT via the command line). I've left the COOKDIR option in the UAT/CookCommandlet as, while non of our projects seem to be using it, I'm not comfortable removing it in case any scripts are using it.

CopyBuildToStagingDirectory also uses these project settings to work out which additional directories it should stage when producing the staging manifest (these are split into UFS and non-UFS directories). This allows users to define additional directories in the editor (or directly in the settings file) which will always be staged regardless of how UAT is invoked.

All of these paths are specified in the settings file as relative to the project Content directory, and the editor will enforce this fact when using the directory picker.

Also mirrored the meta-data interface from UField on IPropertyHandle. This takes care of accessing the meta-data from the correct property, even when you're accessing an element inside an array (MattK approved!).

diff --git a/Engine/Source/Editor/DetailCustomizations/Private/DirectoryPathStructCustomization.cpp b/Engine/Source/Editor/DetailCustomizations/Private/DirectoryPathStructCustomization.cpp
--- a/Engine/Source/Editor/DetailCustomizations/Private/DirectoryPathStructCustomization.cpp
+++ b/Engine/Source/Editor/DetailCustomizations/Private/DirectoryPathStructCustomization.cpp
@@ -15,7 +15,12 @@ TSharedRef<IStructCustomization> FDirectoryPathStructCustomization::MakeInstance
 void FDirectoryPathStructCustomization::CustomizeStructHeader( TSharedRef<IPropertyHandle> StructPropertyHandle, class FDetailWidgetRow& HeaderRow, IStructCustomizationUtils& StructCustomizationUtils )
 {
 	TSharedPtr<IPropertyHandle> PathProperty = StructPropertyHandle->GetChildHandle("Path");
-	bool bUseRelativePath = StructPropertyHandle->GetProperty()->HasMetaData( TEXT("RelativePath") );
+
+	const bool bRelativeToGameContentDir = StructPropertyHandle->HasMetaData( TEXT("RelativeToGameContentDir") );
+	const bool bUseRelativePath = StructPropertyHandle->HasMetaData( TEXT("RelativePath") );
+	
+	AbsoluteGameContentDir = FPaths::ConvertRelativePathToFull(FPaths::GameContentDir());
+
 	if(PathProperty.IsValid())
 	{
 		HeaderRow.ValueContent()
@@ -37,7 +42,7 @@ void FDirectoryPathStructCustomization::CustomizeStructHeader( TSharedRef<IPrope
 				SAssignNew(BrowseButton, SButton)
 				.ButtonStyle( FEditorStyle::Get(), "HoverHintOnly" )
 				.ToolTipText( LOCTEXT( "FolderButtonToolTipText", "Choose a directory from this computer").ToString() )
-				.OnClicked( FOnClicked::CreateSP(this, &FDirectoryPathStructCustomization::OnPickDirectory, PathProperty.ToSharedRef(), bUseRelativePath) )
+				.OnClicked( FOnClicked::CreateSP(this, &FDirectoryPathStructCustomization::OnPickDirectory, PathProperty.ToSharedRef(), bRelativeToGameContentDir, bUseRelativePath) )
 				.ContentPadding( 2.0f )
 				.ForegroundColor( FSlateColor::UseForeground() )
 				.IsFocusable( false )
@@ -59,7 +64,7 @@ void FDirectoryPathStructCustomization::CustomizeStructChildren( TSharedRef<IPro
 {
 }
 
-FReply FDirectoryPathStructCustomization::OnPickDirectory(TSharedRef<IPropertyHandle> PropertyHandle, const bool bUseRelativePath) const
+FReply FDirectoryPathStructCustomization::OnPickDirectory(TSharedRef<IPropertyHandle> PropertyHandle, const bool bRelativeToGameContentDir, const bool bUseRelativePath) const
 {
 	FString Directory;
 	IDesktopPlatform* DesktopPlatform = FDesktopPlatformModule::Get();
@@ -68,20 +73,63 @@ FReply FDirectoryPathStructCustomization::OnPickDirectory(TSharedRef<IPropertyHa
 		TSharedPtr<SWindow> ParentWindow = FSlateApplication::Get().FindWidgetWindow(BrowseButton.ToSharedRef());
 		void* ParentWindowHandle = (ParentWindow.IsValid() && ParentWindow->GetNativeWindow().IsValid()) ? ParentWindow->GetNativeWindow()->GetOSWindowHandle() : nullptr;
 
-		if(DesktopPlatform->OpenDirectoryDialog(ParentWindowHandle, LOCTEXT( "FolderDialogTitle", "Choose a directory").ToString(), FEditorDirectories::Get().GetLastDirectory(ELastDirectory::GENERIC_IMPORT), Directory))
+		FString StartDirectory = FEditorDirectories::Get().GetLastDirectory(ELastDirectory::GENERIC_IMPORT);
+		if(bRelativeToGameContentDir && !IsValidPath(StartDirectory, bRelativeToGameContentDir))
 		{
-			FEditorDirectories::Get().SetLastDirectory(ELastDirectory::GENERIC_IMPORT, Directory);
+			StartDirectory = AbsoluteGameContentDir;
+		}
 
-			if( bUseRelativePath )
+		// Loop until; a) the user cancels (OpenDirectoryDialog returns false), or, b) the chosen path is valid (IsValidPath returns true)
+		for(;;)
+		{
+			if( DesktopPlatform->OpenDirectoryDialog(ParentWindowHandle, LOCTEXT("FolderDialogTitle", "Choose a directory").ToString(), StartDirectory, Directory) )
 			{
-				Directory = IFileManager::Get().ConvertToRelativePath(*Directory);
-			}
+				FText FailureReason;
+				if( IsValidPath(Directory, bRelativeToGameContentDir, &FailureReason) )
+				{
+					FEditorDirectories::Get().SetLastDirectory(ELastDirectory::GENERIC_IMPORT, Directory);
 
-			PropertyHandle->SetValue(Directory);
+					if( bRelativeToGameContentDir )
+					{
+						Directory = Directory.RightChop(AbsoluteGameContentDir.Len());
+					}
+
+					if( bUseRelativePath )
+					{
+						Directory = IFileManager::Get().ConvertToRelativePath(*Directory);
+					}
+
+					PropertyHandle->SetValue(Directory);
+				}
+				else
+				{
+					StartDirectory = Directory;
+					FMessageDialog::Open(EAppMsgType::Ok, FailureReason);
+					continue;
+				}
+			}
+			break;
 		}
 	}
 
 	return FReply::Handled();
 }
 
+bool FDirectoryPathStructCustomization::IsValidPath(const FString& AbsolutePath, const bool bRelativeToGameContentDir, FText* const OutReason) const
+{
+	if(bRelativeToGameContentDir)
+	{
+		if(!AbsolutePath.StartsWith(AbsoluteGameContentDir))
+		{
+			if(OutReason)
+			{
+				*OutReason = FText::Format(LOCTEXT("Error_InvalidRootPath", "The chosen directory must be within {0}"), FText::FromString(AbsoluteGameContentDir));
+			}
+			return false;
+		}
+	}
+
+	return true;
+}
+
 #undef LOCTEXT_NAMESPACE
\ No newline at end of file
diff --git a/Engine/Source/Editor/DetailCustomizations/Private/DirectoryPathStructCustomization.h b/Engine/Source/Editor/DetailCustomizations/Private/DirectoryPathStructCustomization.h
--- a/Engine/Source/Editor/DetailCustomizations/Private/DirectoryPathStructCustomization.h
+++ b/Engine/Source/Editor/DetailCustomizations/Private/DirectoryPathStructCustomization.h
@@ -18,8 +18,14 @@ private:
 	FText GetDisplayedText(TSharedRef<IPropertyHandle> PropertyHandle) const;
 
 	/** Delegate used to display a directory picker */
-	FReply OnPickDirectory(TSharedRef<IPropertyHandle> PropertyHandle, const bool bUseRelativePaths) const;
+	FReply OnPickDirectory(TSharedRef<IPropertyHandle> PropertyHandle, const bool bRelativeToGameContentDir, const bool bUseRelativePaths) const;
+
+	/** Check whether that the chosen path is valid */
+	bool IsValidPath(const FString& AbsolutePath, const bool bRelativeToGameContentDir, FText* const OutReason = nullptr) const;
 
 	/** The browse button widget */
 	TSharedPtr<SButton> BrowseButton;
+
+	/** Absolute path to the game content directory */
+	FString AbsoluteGameContentDir;
 };
\ No newline at end of file
diff --git a/Engine/Source/Editor/DetailCustomizations/Private/StringAssetReferenceCustomization.cpp b/Engine/Source/Editor/DetailCustomizations/Private/StringAssetReferenceCustomization.cpp
--- a/Engine/Source/Editor/DetailCustomizations/Private/StringAssetReferenceCustomization.cpp
+++ b/Engine/Source/Editor/DetailCustomizations/Private/StringAssetReferenceCustomization.cpp
@@ -6,18 +6,9 @@
 
 void FStringAssetReferenceCustomization::CustomizeStructHeader( TSharedRef<IPropertyHandle> InStructPropertyHandle, FDetailWidgetRow& HeaderRow, IStructCustomizationUtils& StructCustomizationUtils )
 {
-	// If we are part of an array, we need to take our meta-data from the array property
-	const UProperty* MetaDataProperty = InStructPropertyHandle->GetProperty();
-	if(InStructPropertyHandle->GetIndexInArray() != INDEX_NONE)
-	{
-		TSharedPtr<IPropertyHandle> ParentPropertyHandle = InStructPropertyHandle->GetParentHandle();
-		check(ParentPropertyHandle.IsValid() && ParentPropertyHandle->AsArray().IsValid());
-		MetaDataProperty = ParentPropertyHandle->GetProperty();
-	}
-
-	this->StructPropertyHandle = InStructPropertyHandle;
+	StructPropertyHandle = InStructPropertyHandle;
 
-	FString ClassFilterString = MetaDataProperty->GetMetaData("AllowedClasses");
+	FString ClassFilterString = StructPropertyHandle->GetMetaData("AllowedClasses");
 	if( !ClassFilterString.IsEmpty() )
 	{
 		TArray<FString> CustomClassFilterNames;
@@ -39,7 +30,7 @@ void FStringAssetReferenceCustomization::CustomizeStructHeader( TSharedRef<IProp
 		}
 	}
 
-	bExactClass = MetaDataProperty->GetBoolMetaData("ExactClass");
+	bExactClass = StructPropertyHandle->GetBoolMetaData("ExactClass");
 
 	FOnShouldFilterAsset AssetFilter;
 	UClass* ClassFilter = UObject::StaticClass();
@@ -58,7 +49,7 @@ void FStringAssetReferenceCustomization::CustomizeStructHeader( TSharedRef<IProp
 	}
 
 	// Can the field be cleared
-	const bool bAllowClear = !(MetaDataProperty->PropertyFlags & CPF_NoClear);
+	const bool bAllowClear = !(StructPropertyHandle->GetMetaDataProperty()->PropertyFlags & CPF_NoClear);
 
 	HeaderRow
 	.NameContent()
diff --git a/Engine/Source/Editor/DetailCustomizations/Private/StringClassReferenceCustomization.cpp b/Engine/Source/Editor/DetailCustomizations/Private/StringClassReferenceCustomization.cpp
--- a/Engine/Source/Editor/DetailCustomizations/Private/StringClassReferenceCustomization.cpp
+++ b/Engine/Source/Editor/DetailCustomizations/Private/StringClassReferenceCustomization.cpp
@@ -5,25 +5,16 @@
 
 void FStringClassReferenceCustomization::CustomizeStructHeader(TSharedRef<IPropertyHandle> InStructPropertyHandle, FDetailWidgetRow& HeaderRow, IStructCustomizationUtils& StructCustomizationUtils)
 {
-	// If we are part of an array, we need to take our meta-data from the array property
-	const UProperty* MetaDataProperty = InStructPropertyHandle->GetProperty();
-	if(InStructPropertyHandle->GetIndexInArray() != INDEX_NONE)
-	{
-		TSharedPtr<IPropertyHandle> ParentPropertyHandle = InStructPropertyHandle->GetParentHandle();
-		check(ParentPropertyHandle.IsValid() && ParentPropertyHandle->AsArray().IsValid());
-		MetaDataProperty = ParentPropertyHandle->GetProperty();
-	}
-
 	StructPropertyHandle = InStructPropertyHandle;
 
 	ClassNamePropertyHandle = InStructPropertyHandle->GetChildHandle("ClassName");
 	check(ClassNamePropertyHandle.IsValid());
 
-	const FString& MetaClassName = MetaDataProperty->GetMetaData("MetaClass");
-	const FString& RequiredInterfaceName = MetaDataProperty->GetMetaData("RequiredInterface");
-	const bool bAllowAbstract = MetaDataProperty->GetBoolMetaData("AllowAbstract");
-	const bool bIsBlueprintBaseOnly = MetaDataProperty->GetBoolMetaData("IsBlueprintBaseOnly");
-	const bool bAllowNone = !(MetaDataProperty->PropertyFlags & CPF_NoClear);
+	const FString& MetaClassName = StructPropertyHandle->GetMetaData("MetaClass");
+	const FString& RequiredInterfaceName = StructPropertyHandle->GetMetaData("RequiredInterface");
+	const bool bAllowAbstract = StructPropertyHandle->GetBoolMetaData("AllowAbstract");
+	const bool bIsBlueprintBaseOnly = StructPropertyHandle->GetBoolMetaData("IsBlueprintBaseOnly");
+	const bool bAllowNone = !(StructPropertyHandle->GetMetaDataProperty()->PropertyFlags & CPF_NoClear);
 
 	check(!MetaClassName.IsEmpty());
 	const UClass* const MetaClass = StringToClass(MetaClassName);
diff --git a/Engine/Source/Editor/MainFrame/Private/Frame/MainFrameActions.cpp b/Engine/Source/Editor/MainFrame/Private/Frame/MainFrameActions.cpp
--- a/Engine/Source/Editor/MainFrame/Private/Frame/MainFrameActions.cpp
+++ b/Engine/Source/Editor/MainFrame/Private/Frame/MainFrameActions.cpp
@@ -403,6 +403,11 @@ void FMainFrameActionCallbacks::CookContent( const FString InPlatformName, const
 	FString PlatformName = InPlatformName;
 	FString OptionalParams;
 
+	if (!FModuleManager::LoadModuleChecked<IProjectTargetPlatformEditorModule>("ProjectTargetPlatformEditor").ShowUnsupportedTargetWarning(*PlatformName))
+	{
+		return;
+	}
+
 	if (PlatformName == "MacNoEditor")
 	{
 		OptionalParams += TEXT(" -targetplatform=Mac");
@@ -494,6 +499,11 @@ void FMainFrameActionCallbacks::PackageProject( const FString InPlatformName, co
 	}
 #endif
 
+	if (!FModuleManager::LoadModuleChecked<IProjectTargetPlatformEditorModule>("ProjectTargetPlatformEditor").ShowUnsupportedTargetWarning(*PlatformName))
+	{
+		return;
+	}
+
 	UProjectPackagingSettings* PackagingSettings = Cast<UProjectPackagingSettings>(UProjectPackagingSettings::StaticClass()->GetDefaultObject());
 
 	// let the user pick a target directory
@@ -573,22 +583,6 @@ void FMainFrameActionCallbacks::PackageProject( const FString InPlatformName, co
 		PlatformName = PlatformName.Mid(0, UnderscoreLoc);
 	}
 
-	const int32 DirCount = PackagingSettings->DirectoriesToAlwaysCook.Num();
-	if(DirCount> 0)
-	{
-		FString Dirs = TEXT(" -cookdir=");
-		for(int32 DirIndex = 0; DirIndex < DirCount; DirIndex++)
-		{
-			Dirs += FString::Printf(TEXT("\"%s\""), *PackagingSettings->DirectoriesToAlwaysCook[DirIndex].Path);
-			if(DirCount > 1 && DirIndex != DirCount - 1)
-			{
-				Dirs += TEXT("+");
-			}
-		}
-
-		OptionalParams += Dirs;
-	}
-
 	FString ExecutableName = FPlatformProcess::ExecutableName(false);
 #if PLATFORM_WINDOWS
 	// turn UE4editor into UE4editor.cmd
diff --git a/Engine/Source/Editor/PropertyEditor/Private/PropertyHandleImpl.cpp b/Engine/Source/Editor/PropertyEditor/Private/PropertyHandleImpl.cpp
--- a/Engine/Source/Editor/PropertyEditor/Private/PropertyHandleImpl.cpp
+++ b/Engine/Source/Editor/PropertyEditor/Private/PropertyHandleImpl.cpp
@@ -1625,6 +1625,66 @@ UProperty* FPropertyHandleBase::GetProperty() const
 	return NULL;
 }
 
+UProperty* FPropertyHandleBase::GetMetaDataProperty() const
+{
+	UProperty* MetaDataProperty = nullptr;
+
+	TSharedPtr<FPropertyNode> PropertyNode = Implementation->GetPropertyNode();
+	if( PropertyNode.IsValid() )
+	{
+		MetaDataProperty = PropertyNode->GetProperty();
+		
+		// If we are part of an array, we need to take our meta-data from the array property
+		if( PropertyNode->GetArrayIndex() != INDEX_NONE )
+		{
+			TSharedPtr<FPropertyNode> ParentNode = PropertyNode->GetParentNodeSharedPtr();
+			check(ParentNode.IsValid());
+			MetaDataProperty = ParentNode->GetProperty();
+		}
+	}
+
+	return MetaDataProperty;
+}
+
+bool FPropertyHandleBase::HasMetaData(const FName& Key) const
+{
+	UProperty* const MetaDataProperty = GetMetaDataProperty();
+	return (MetaDataProperty) ? MetaDataProperty->HasMetaData(Key) : false;
+}
+
+const FString& FPropertyHandleBase::GetMetaData(const FName& Key) const
+{
+	// if not found, return a static empty string
+	static const FString EmptyString;
+
+	UProperty* const MetaDataProperty = GetMetaDataProperty();
+	return (MetaDataProperty) ? MetaDataProperty->GetMetaData(Key) : EmptyString;
+}
+
+bool FPropertyHandleBase::GetBoolMetaData(const FName& Key) const
+{
+	UProperty* const MetaDataProperty = GetMetaDataProperty();
+	return (MetaDataProperty) ? MetaDataProperty->GetBoolMetaData(Key) : false;
+}
+
+int32 FPropertyHandleBase::GetINTMetaData(const FName& Key) const
+{
+	UProperty* const MetaDataProperty = GetMetaDataProperty();
+	return (MetaDataProperty) ? MetaDataProperty->GetINTMetaData(Key) : 0;
+}
+
+float FPropertyHandleBase::GetFLOATMetaData(const FName& Key) const
+{
+	UProperty* const MetaDataProperty = GetMetaDataProperty();
+	return (MetaDataProperty) ? MetaDataProperty->GetFLOATMetaData(Key) : 0.0f;
+}
+
+UClass* FPropertyHandleBase::GetClassMetaData(const FName& Key) const
+{
+	UProperty* const MetaDataProperty = GetMetaDataProperty();
+	return (MetaDataProperty) ? MetaDataProperty->GetClassMetaData(Key) : nullptr;
+}
+
 FString FPropertyHandleBase::GetToolTipText() const
 {
 	TSharedPtr<FPropertyNode> PropertyNode = Implementation->GetPropertyNode();
diff --git a/Engine/Source/Editor/PropertyEditor/Private/PropertyHandleImpl.h b/Engine/Source/Editor/PropertyEditor/Private/PropertyHandleImpl.h
--- a/Engine/Source/Editor/PropertyEditor/Private/PropertyHandleImpl.h
+++ b/Engine/Source/Editor/PropertyEditor/Private/PropertyHandleImpl.h
@@ -388,6 +388,13 @@ public:
 	virtual TSharedPtr<IPropertyHandleArray> AsArray() OVERRIDE { return NULL; }
 	virtual const UClass* GetPropertyClass() const OVERRIDE;
 	virtual UProperty* GetProperty() const OVERRIDE;
+	virtual UProperty* GetMetaDataProperty() const OVERRIDE;
+	virtual bool HasMetaData(const FName& Key) const OVERRIDE;
+	virtual const FString& GetMetaData(const FName& Key) const OVERRIDE;
+	virtual bool GetBoolMetaData(const FName& Key) const OVERRIDE;
+	virtual int32 GetINTMetaData(const FName& Key) const OVERRIDE;
+	virtual float GetFLOATMetaData(const FName& Key) const OVERRIDE;
+	virtual UClass* GetClassMetaData(const FName& Key) const OVERRIDE;
 	virtual FString GetToolTipText() const OVERRIDE;
 	virtual void SetToolTipText(const FString& ToolTip) OVERRIDE;
 	virtual FPropertyAccess::Result SetPerObjectValues( const TArray<FString>& InPerObjectValues,  EPropertyValueSetFlags::Type Flags = EPropertyValueSetFlags::DefaultFlags ) OVERRIDE;
diff --git a/Engine/Source/Editor/PropertyEditor/Public/PropertyHandle.h b/Engine/Source/Editor/PropertyEditor/Public/PropertyHandle.h
--- a/Engine/Source/Editor/PropertyEditor/Public/PropertyHandle.h
+++ b/Engine/Source/Editor/PropertyEditor/Public/PropertyHandle.h
@@ -46,6 +46,59 @@ public:
 	 */
 	virtual UProperty* GetProperty() const = 0;
 
+	/*
+	 * Gets the property we should use to read meta-data
+	 */
+	virtual UProperty* GetMetaDataProperty() const = 0;
+
+	/**
+	 * Determines if the property has any metadata associated with the key
+	 * 
+	 * @param Key The key to lookup in the metadata
+	 * @return true if there is a (possibly blank) value associated with this key
+	 */
+	virtual bool HasMetaData(const FName& Key) const = 0;
+
+	/**
+	 * Find the metadata value associated with the key
+	 * 
+	 * @param Key The key to lookup in the metadata
+	 * @return The value associated with the key
+	 */
+	virtual const FString& GetMetaData(const FName& Key) const = 0;
+
+	/**
+	 * Find the metadata value associated with the key and return bool
+	 *
+	 * @param Key The key to lookup in the metadata
+	 * @return return true if the value was true (case insensitive)
+	 */
+	virtual bool GetBoolMetaData(const FName& Key) const = 0;
+
+	/**
+	 * Find the metadata value associated with the key and return int32
+	 *
+	 * @param Key The key to lookup in the metadata
+	 * @return the int value stored in the metadata.
+	 */
+	virtual int32 GetINTMetaData(const FName& Key) const = 0;
+
+	/**
+	 * Find the metadata value associated with the key and return float
+	 *
+	 * @param Key The key to lookup in the metadata
+	 * @return the float value stored in the metadata.
+	 */
+	virtual float GetFLOATMetaData(const FName& Key) const = 0;
+
+	/**
+	 * Find the metadata value associated with the key and return UClass*
+	 *
+	 * @param Key The key to lookup in the metadata
+	 * @return the UClass value stored in the metadata.
+	 */
+	virtual UClass* GetClassMetaData(const FName& Key) const = 0;
+
 	/**
 	 * Gets the property tool tip text.
 	 */
diff --git a/Engine/Source/Editor/UnrealEd/Classes/Settings/ProjectPackagingSettings.h b/Engine/Source/Editor/UnrealEd/Classes/Settings/ProjectPackagingSettings.h
--- a/Engine/Source/Editor/UnrealEd/Classes/Settings/ProjectPackagingSettings.h
+++ b/Engine/Source/Editor/UnrealEd/Classes/Settings/ProjectPackagingSettings.h
@@ -77,12 +77,6 @@ public:
 	UPROPERTY(config, EditAnywhere, Category=Project, AdvancedDisplay)
 	bool ForDistribution;
 
-	/**
-	 * All assets in these directories will always be cooked.
-	 */
-	UPROPERTY(EditAnywhere, config, Category=Cooking)
-	TArray<FDirectoryPath> DirectoriesToAlwaysCook;
-
 	/**
 	 * If enabled, all content will be put into a single .pak file instead of many individual files (default = enabled).
 	 */
@@ -95,4 +89,27 @@ public:
 	*/
 	UPROPERTY(config, EditAnywhere, Category = Packaging)
 	bool UseOBB_InAPK;
+
+	/**
+	 * Directories containing .uasset files that should always be cooked regardless of whether they're referenced by anything in your project
+	 * Note: These paths are relative to your project Content directory
+	 */
+	UPROPERTY(config, EditAnywhere, Category=Packaging, AdvancedDisplay, meta=(DisplayName="Additional Asset Directories to Cook", RelativeToGameContentDir))
+	TArray<FDirectoryPath> DirectoriesToAlwaysCook;
+
+	/**
+	 * Directories containing files that should always be added to the .pak file (if using a .pak file; otherwise they're copied as individual files)
+	 * This is used to stage additional files that you manually load via the UFS (Unreal File System) file IO API
+	 * Note: These paths are relative to your project Content directory
+	 */
+	UPROPERTY(config, EditAnywhere, Category=Packaging, AdvancedDisplay, meta=(DisplayName="Additional Non-Asset Directories to Package", RelativeToGameContentDir))
+	TArray<FDirectoryPath> DirectoriesToAlwaysStageAsUFS;
+
+	/**
+	 * Directories containing files that should always be copied when packaging your project, but are not supposed to be part of the .pak file
+	 * This is used to stage additional files that you manually load without using the UFS (Unreal File System) file IO API, eg, third-party libraries that perform their own internal file IO
+	 * Note: These paths are relative to your project Content directory
+	 */
+	UPROPERTY(config, EditAnywhere, Category=Packaging, AdvancedDisplay, meta=(DisplayName="Additional Non-Asset Directories To Copy", RelativeToGameContentDir))
+	TArray<FDirectoryPath> DirectoriesToAlwaysStageAsNonUFS;
 };
diff --git a/Engine/Source/Editor/UnrealEd/Private/Commandlets/CookCommandlet.cpp b/Engine/Source/Editor/UnrealEd/Private/Commandlets/CookCommandlet.cpp
--- a/Engine/Source/Editor/UnrealEd/Private/Commandlets/CookCommandlet.cpp
+++ b/Engine/Source/Editor/UnrealEd/Private/Commandlets/CookCommandlet.cpp
@@ -770,10 +770,23 @@ void UCookCommandlet::CollectFilesToCook(TArray<FString>& FilesInPath)
 				DirToCook = DirToCook.Right(DirToCook.Len() - (PlusIdx + 1));
 				PlusIdx = DirToCook.Find(TEXT("+"));
 			}
+			// The dir may be contained within quotes
+			DirToCook = DirToCook.TrimQuotes();
+			FPaths::NormalizeDirectoryName(DirToCook);
 			CmdLineDirEntries.Add(DirToCook);
 		}
 	}
 
+	// Also append any cookdirs from the project ini files; these dirs are relative to the game content directory
+	{
+		const FString AbsoluteGameContentDir = FPaths::ConvertRelativePathToFull(FPaths::GameContentDir());
+		const UProjectPackagingSettings* const PackagingSettings = GetDefault<UProjectPackagingSettings>();
+		for(const auto& DirToCook : PackagingSettings->DirectoriesToAlwaysCook)
+		{
+			CmdLineDirEntries.Add(AbsoluteGameContentDir / DirToCook.Path);
+		}
+	}
+
 	for (int32 CmdLineMapIdx = 0; CmdLineMapIdx < CmdLineMapEntries.Num(); CmdLineMapIdx++)
 	{
 		FString CurrEntry = CmdLineMapEntries[CmdLineMapIdx];
diff --git a/Engine/Source/Programs/AutomationTool/ProjectParams.cs b/Engine/Source/Programs/AutomationTool/ProjectParams.cs
--- a/Engine/Source/Programs/AutomationTool/ProjectParams.cs
+++ b/Engine/Source/Programs/AutomationTool/ProjectParams.cs
@@ -1002,6 +1002,9 @@ namespace AutomationTool
 		#region Initialization
 
 		private Dictionary<TargetRules.TargetType, SingleTargetProperties> DetectedTargets;
+		private Dictionary<UnrealTargetPlatform, ConfigCacheIni> LoadedEngineConfigs;
+		private Dictionary<UnrealTargetPlatform, ConfigCacheIni> LoadedGameConfigs;
+
 		private void AutodetectSettings(bool bReset)
 		{
 			if (bReset)
@@ -1023,6 +1026,8 @@ namespace AutomationTool
 
 			bIsCodeBasedProject = Properties.bIsCodeBasedProject;			
 			DetectedTargets = Properties.Targets;
+			LoadedEngineConfigs = Properties.EngineConfigs;
+			LoadedGameConfigs = Properties.GameConfigs;
 
 			var GameTarget = String.Empty;
 			var EditorTarget = String.Empty;
@@ -1391,6 +1396,36 @@ namespace AutomationTool
 			}
 		}
 
+		/// <summary>
+		/// List of all Engine ini files for this project
+		/// </summary>
+		public Dictionary<UnrealTargetPlatform, ConfigCacheIni> EngineConfigs
+		{
+			get
+			{
+				if (LoadedEngineConfigs == null)
+				{
+					AutodetectSettings(false);
+				}
+				return LoadedEngineConfigs;
+			}
+		}
+
+		/// <summary>
+		/// List of all Game ini files for this project
+		/// </summary>
+		public Dictionary<UnrealTargetPlatform, ConfigCacheIni> GameConfigs
+		{
+			get
+			{
+				if (LoadedGameConfigs == null)
+				{
+					AutodetectSettings(false);
+				}
+				return LoadedGameConfigs;
+			}
+		}
+
 		public void Validate()
 		{
 			if (String.IsNullOrEmpty(RawProjectPath))
diff --git a/Engine/Source/Programs/AutomationTool/ProjectUtils.cs b/Engine/Source/Programs/AutomationTool/ProjectUtils.cs
--- a/Engine/Source/Programs/AutomationTool/ProjectUtils.cs
+++ b/Engine/Source/Programs/AutomationTool/ProjectUtils.cs
@@ -65,6 +65,11 @@ namespace AutomationTool
 		/// </summary>
 		public Dictionary<UnrealTargetPlatform, ConfigCacheIni> EngineConfigs = new Dictionary<UnrealTargetPlatform,ConfigCacheIni>();
 
+		/// <summary>
+		/// List of all Game ini files for this project
+		/// </summary>
+		public Dictionary<UnrealTargetPlatform, ConfigCacheIni> GameConfigs = new Dictionary<UnrealTargetPlatform, ConfigCacheIni>();
+
 		/// <summary>
 		/// List of all programs detected for this project.
 		/// </summary>
@@ -183,8 +188,9 @@ namespace AutomationTool
 			// Get all ini files
 			if (!String.IsNullOrWhiteSpace(RawProjectPath))
 			{
-				var EngineDirectory = CommandUtils.CombinePaths(CommandUtils.CmdEnv.LocalRoot, "Engine");
 				CommandUtils.Log("Loading ini files for {0}", RawProjectPath);
+
+				var EngineDirectory = CommandUtils.CombinePaths(CommandUtils.CmdEnv.LocalRoot, "Engine");
 				foreach (UnrealTargetPlatform TargetPlatformType in Enum.GetValues(typeof(UnrealTargetPlatform)))
 				{
 					if (TargetPlatformType != UnrealTargetPlatform.Unknown)
@@ -193,6 +199,15 @@ namespace AutomationTool
 						Properties.EngineConfigs.Add(TargetPlatformType, Config);
 					}
 				}
+
+				foreach (UnrealTargetPlatform TargetPlatformType in Enum.GetValues(typeof(UnrealTargetPlatform)))
+				{
+					if (TargetPlatformType != UnrealTargetPlatform.Unknown)
+					{
+						var Config = new ConfigCacheIni(TargetPlatformType, "Game", Path.GetDirectoryName(RawProjectPath));
+						Properties.GameConfigs.Add(TargetPlatformType, Config);
+					}
+				}
 			}
 
 			return Properties;
diff --git a/Engine/Source/Programs/AutomationTool/Scripts/CopyBuildToStagingDirectory.Automation.cs b/Engine/Source/Programs/AutomationTool/Scripts/CopyBuildToStagingDirectory.Automation.cs
--- a/Engine/Source/Programs/AutomationTool/Scripts/CopyBuildToStagingDirectory.Automation.cs
+++ b/Engine/Source/Programs/AutomationTool/Scripts/CopyBuildToStagingDirectory.Automation.cs
@@ -190,6 +190,44 @@ public partial class Project : CommandUtils
 
 			SC.StageFiles(StagedFileType.UFS, CombinePaths(SC.ProjectRoot, "Content/Localization/Engine"), "*.locres", true, null, CombinePaths(SC.RelativeProjectRootForStage, "Content/Localization/Engine"), true, !Params.Pak);
 
+			// Stage any additional UFS and NonUFS paths specified in the project ini files; these dirs are relative to the game content directory
+			ConfigCacheIni PlatformGameConfig = null;
+			if (Params.GameConfigs.TryGetValue(SC.StageTargetPlatform.PlatformType, out PlatformGameConfig))
+			{
+				var ProjectContentRoot = CombinePaths(SC.ProjectRoot, "Content");
+				var StageContentRoot = CombinePaths(SC.RelativeProjectRootForStage, "Content");
+
+				List<string> ExtraUFSDirs;
+				if (PlatformGameConfig.GetArray("/Script/UnrealEd.ProjectPackagingSettings", "DirectoriesToAlwaysStageAsUFS", out ExtraUFSDirs))
+				{
+					// Each string has the format '(Path="TheDirToStage")'
+					foreach (var PathStr in ExtraUFSDirs)
+					{
+						var PathParts = PathStr.Split('"');
+						if (PathParts.Length == 3)
+						{
+							var RelativePath = PathParts[1];
+							SC.StageFiles(StagedFileType.UFS, CombinePaths(ProjectContentRoot, RelativePath), "*", true, null, CombinePaths(StageContentRoot, RelativePath), true, !Params.Pak);
+						}
+					}
+				}
+
+				List<string> ExtraNonUFSDirs;
+				if (PlatformGameConfig.GetArray("/Script/UnrealEd.ProjectPackagingSettings", "DirectoriesToAlwaysStageAsNonUFS", out ExtraNonUFSDirs))
+				{
+					// Each string has the format '(Path="TheDirToStage")'
+					foreach (var PathStr in ExtraNonUFSDirs)
+					{
+						var PathParts = PathStr.Split('"');
+						if (PathParts.Length == 3)
+						{
+							var RelativePath = PathParts[1];
+							SC.StageFiles(StagedFileType.NonUFS, CombinePaths(ProjectContentRoot, RelativePath));
+						}
+					}
+				}
+			}
+
 			StagedFileType StagedFileTypeForMovies = StagedFileType.NonUFS;
 			if (Params.FileServer)
 			{
diff --git a/Engine/Source/Runtime/CoreUObject/Private/UObject/Class.cpp b/Engine/Source/Runtime/CoreUObject/Private/UObject/Class.cpp
--- a/Engine/Source/Runtime/CoreUObject/Private/UObject/Class.cpp
+++ b/Engine/Source/Runtime/CoreUObject/Private/UObject/Class.cpp
@@ -323,6 +323,13 @@ UClass* UField::GetClassMetaData(const TCHAR* Key) const
 	return FoundObject;
 }
 
+UClass* UField::GetClassMetaData(const FName& Key) const
+{
+	const FString & ClassName = GetMetaData(Key);
+	UClass* const FoundObject = FindObject<UClass>(ANY_PACKAGE, *ClassName);
+	return FoundObject;
+}
+
 void UField::RemoveMetaData(const TCHAR* Key)
 {
 	UPackage* Package = GetOutermost();
diff --git a/Engine/Source/Runtime/CoreUObject/Public/UObject/Class.h b/Engine/Source/Runtime/CoreUObject/Public/UObject/Class.h
--- a/Engine/Source/Runtime/CoreUObject/Public/UObject/Class.h
+++ b/Engine/Source/Runtime/CoreUObject/Public/UObject/Class.h
@@ -163,6 +163,19 @@ class COREUOBJECT_API UField : public UObject
 		return Value;
 	}
 
+	/**
+	* Find the metadata value associated with the key
+	* and return int32 
+	* @param Key The key to lookup in the metadata
+	* @return the int value stored in the metadata.
+	*/
+	int32 GetINTMetaData(const FName& Key) const
+	{
+		const FString & INTString = GetMetaData(Key);
+		int32 Value = FCString::Atoi(*INTString);
+		return Value;
+	}
+
 	/**
 	* Find the metadata value associated with the key
 	* and return float
@@ -177,7 +190,22 @@ class COREUOBJECT_API UField : public UObject
 		return Value;
 	}
 
+	/**
+	* Find the metadata value associated with the key
+	* and return float
+	* @param Key The key to lookup in the metadata
+	* @return the float value stored in the metadata.
+	*/
+	float GetFLOATMetaData(const FName& Key) const
+	{
+		const FString & FLOATString = GetMetaData(Key);
+		// FString == operator does case insensitive comparison
+		float Value = FCString::Atof(*FLOATString);
+		return Value;
+	}
+
 	UClass* GetClassMetaData(const TCHAR* Key) const;
+	UClass* GetClassMetaData(const FName& Key) const;
 
 	/** Clear any metadata associated with the key */
 	void RemoveMetaData(const TCHAR* Key);
