Moved some functionality from EmitterDetails.cpp to ParticleSystemComponent.cpp so that it appears for PSCs in the BP editor too.

#ttp 311726 - UE4: RENDERING: Rocket/Orion: Blueprint Editor: Components: ParticleSystems BP components are missing the particle 'expose parameters' button which is seen on particle component actors

diff --git a/Engine/Source/Editor/DetailCustomizations/Private/DetailCustomizations.cpp b/Engine/Source/Editor/DetailCustomizations/Private/DetailCustomizations.cpp
--- a/Engine/Source/Editor/DetailCustomizations/Private/DetailCustomizations.cpp
+++ b/Engine/Source/Editor/DetailCustomizations/Private/DetailCustomizations.cpp
@@ -17,7 +17,6 @@
 #include "ReflectionCaptureDetails.h"
 #include "SkyLightComponentDetails.h"
 #include "BrushDetails.h"
-#include "EmitterDetails.h"
 #include "ActorDetails.h"
 #include "SkeletalControlNodeDetails.h"
 #include "AnimMontageSegmentDetails.h"
@@ -75,6 +74,7 @@
 #include "WindowsTargetSettingsDetails.h"
 #include "MoviePlayerSettingsDetails.h"
 #include "SourceCodeAccessSettingsDetails.h"
+#include "ParticleSystemComponentDetails.h"
 
 
 IMPLEMENT_MODULE( FDetailCustomizationsModule, DetailCustomizations );
@@ -148,7 +148,6 @@ void FDetailCustomizationsModule::StartupModule()
 	RegisterCustomPropertyLayout( "ReflectionCapture", FOnGetDetailCustomizationInstance::CreateStatic( &FReflectionCaptureDetails::MakeInstance ) );
 	RegisterCustomPropertyLayout( "SkyLight", FOnGetDetailCustomizationInstance::CreateStatic( &FSkyLightComponentDetails::MakeInstance ) );
 	RegisterCustomPropertyLayout( "Brush", FOnGetDetailCustomizationInstance::CreateStatic( &FBrushDetails::MakeInstance ) );
-	RegisterCustomPropertyLayout( "Emitter", FOnGetDetailCustomizationInstance::CreateStatic( &FEmitterDetails::MakeInstance ) );
 	RegisterCustomPropertyLayout( "AmbientSound", FOnGetDetailCustomizationInstance::CreateStatic( &FAmbientSoundDetails::MakeInstance ) );
 	RegisterCustomPropertyLayout( "WorldSettings", FOnGetDetailCustomizationInstance::CreateStatic( &FWorldSettingsDetails::MakeInstance ) );
 	RegisterCustomPropertyLayout( "GeneralProjectSettings", FOnGetDetailCustomizationInstance::CreateStatic(&FGeneralProjectSettingsDetails::MakeInstance));
@@ -195,7 +194,8 @@ void FDetailCustomizationsModule::StartupModule()
 	RegisterCustomPropertyLayout( "WindowsTargetSettings", FOnGetDetailCustomizationInstance::CreateStatic( &FWindowsTargetSettingsDetails::MakeInstance ) );
 	RegisterCustomPropertyLayout( "MoviePlayerSettings", FOnGetDetailCustomizationInstance::CreateStatic( &FMoviePlayerSettingsDetails::MakeInstance ) );
 
-	RegisterCustomPropertyLayout( "SourceCodeAccessSettings", FOnGetDetailCustomizationInstance::CreateStatic( &FSourceCodeAccessSettingsDetails::MakeInstance ) );
+	RegisterCustomPropertyLayout("SourceCodeAccessSettings", FOnGetDetailCustomizationInstance::CreateStatic(&FSourceCodeAccessSettingsDetails::MakeInstance));
+	RegisterCustomPropertyLayout("ParticleSystemComponent", FOnGetDetailCustomizationInstance::CreateStatic(&FParticleSystemComponentDetails::MakeInstance));
 
 	PropertyModule.NotifyCustomizationModuleChanged();
 }
diff --git a/Engine/Source/Editor/DetailCustomizations/Private/EmitterDetails.cpp b/Engine/Source/Editor/DetailCustomizations/Private/EmitterDetails.cpp
deleted file mode 100644
diff --git a/Engine/Source/Editor/DetailCustomizations/Private/EmitterDetails.h b/Engine/Source/Editor/DetailCustomizations/Private/EmitterDetails.h
deleted file mode 100644
diff --git a/Engine/Source/Editor/DetailCustomizations/Private/ParticleSystemComponentDetails.cpp b/Engine/Source/Editor/DetailCustomizations/Private/ParticleSystemComponentDetails.cpp
new file mode 100644
--- /dev/null
+++ b/Engine/Source/Editor/DetailCustomizations/Private/ParticleSystemComponentDetails.cpp
@@ -0,0 +1,103 @@
+// Copyright 1998-2014 Epic Games, Inc. All Rights Reserved.
+
+#include "DetailCustomizationsPrivatePCH.h"
+#include "ParticleSystemComponentDetails.h"
+#include "Particles/Emitter.h"
+#include "Particles/ParticleSystemComponent.h"
+
+#define LOCTEXT_NAMESPACE "ParticleSystemComponentDetails"
+
+TSharedRef<IDetailCustomization> FParticleSystemComponentDetails::MakeInstance()
+{
+	return MakeShareable(new FParticleSystemComponentDetails);
+}
+
+void FParticleSystemComponentDetails::CustomizeDetails(IDetailLayoutBuilder& InDetailLayout)
+{
+	this->DetailLayout = &InDetailLayout;
+
+	InDetailLayout.EditCategory("Particles", TEXT(""), ECategoryPriority::Important);
+
+	IDetailCategoryBuilder& CustomCategory = InDetailLayout.EditCategory("EmitterActions", NSLOCTEXT("ParticleSystemComponentDetails", "EmitterActionCategoryName", "Emitter Actions").ToString(), ECategoryPriority::Important);
+
+	CustomCategory.AddCustomRow(TEXT(""))
+		[
+			SNew(SUniformGridPanel)
+			.SlotPadding(2.0f)
+			+ SUniformGridPanel::Slot(0, 0)
+			[
+				SNew(SButton)
+				.OnClicked(this, &FParticleSystemComponentDetails::OnAutoPopulateClicked)
+				.ToolTipText(NSLOCTEXT("ParticleSystemComponentDetails", "AutoPopulateButtonTooltip", "Copies properties from the source particle system into the instanced parameters of this system").ToString())
+				.HAlign(HAlign_Center)
+				[
+					SNew(STextBlock)
+					.Text(NSLOCTEXT("ParticleSystemComponentDetails", "AutoPopulateButton", "Expose Parameter").ToString())
+				]
+			]
+			+ SUniformGridPanel::Slot(1, 0)
+				[
+					SNew(SButton)
+					.OnClicked(this, &FParticleSystemComponentDetails::OnResetEmitter)
+					.ToolTipText(NSLOCTEXT("ParticleSystemComponentDetails", "ResetEmitterButtonTooltip", "Resets the selected particle system.").ToString())
+					.HAlign(HAlign_Center)
+					[
+						SNew(STextBlock)
+						.Text(NSLOCTEXT("ParticleSystemComponentDetails", "ResetEmitterButton", "Reset Emitter").ToString())
+					]
+				]
+		];
+}
+
+FReply FParticleSystemComponentDetails::OnAutoPopulateClicked()
+{
+	const TArray< TWeakObjectPtr<UObject> >& SelectedObjects = DetailLayout->GetDetailsView().GetSelectedObjects();
+
+	for (int32 Idx = 0; Idx < SelectedObjects.Num(); ++Idx)
+	{
+		if (SelectedObjects[Idx].IsValid())
+		{
+			if (UParticleSystemComponent* PSC = Cast<UParticleSystemComponent>(SelectedObjects[Idx].Get()))
+			{
+				PSC->AutoPopulateInstanceProperties();
+			}
+			else if (AEmitter* Emitter = Cast<AEmitter>(SelectedObjects[Idx].Get()))
+			{
+				Emitter->AutoPopulateInstanceProperties();
+			}
+		}
+	}
+
+	return FReply::Handled();
+}
+
+FReply FParticleSystemComponentDetails::OnResetEmitter()
+{
+	const TArray< TWeakObjectPtr<UObject> >& SelectedObjects = DetailLayout->GetDetailsView().GetSelectedObjects();
+	// Iterate over selected Actors.
+	for (int32 Idx = 0; Idx < SelectedObjects.Num(); ++Idx)
+	{
+		if (SelectedObjects[Idx].IsValid())
+		{
+			UParticleSystemComponent* PSC = Cast<UParticleSystemComponent>(SelectedObjects[Idx].Get());
+			if (!PSC)
+			{
+				if (AEmitter* Emitter = Cast<AEmitter>(SelectedObjects[Idx].Get()))
+				{
+					PSC = Emitter->ParticleSystemComponent;
+				}
+			}
+
+			if (PSC)
+			{
+				PSC->ResetParticles();
+				PSC->ActivateSystem();
+			}
+		}
+	}
+	return FReply::Handled();
+}
+
+
+
+#undef LOCTEXT_NAMESPACE
diff --git a/Engine/Source/Editor/DetailCustomizations/Private/ParticleSystemComponentDetails.h b/Engine/Source/Editor/DetailCustomizations/Private/ParticleSystemComponentDetails.h
new file mode 100644
--- /dev/null
+++ b/Engine/Source/Editor/DetailCustomizations/Private/ParticleSystemComponentDetails.h
@@ -0,0 +1,24 @@
+// Copyright 1998-2014 Epic Games, Inc. All Rights Reserved.
+
+#pragma once
+
+class FParticleSystemComponentDetails : public IDetailCustomization
+{
+public:
+	/** Makes a new instance of this detail layout class for a specific detail view requesting it */
+	static TSharedRef<IDetailCustomization> MakeInstance();
+private:
+	/** IDetailCustomization interface */
+	virtual void CustomizeDetails( IDetailLayoutBuilder& InDetailLayout ) OVERRIDE;
+
+	/** Handles the Auto-Populate button's on click event */
+	FReply OnAutoPopulateClicked();
+
+	/** Handles the Emitter Reset button's on click event */
+	FReply OnResetEmitter();
+
+private:
+
+	/** Cached off reference to the layout builder */
+	IDetailLayoutBuilder* DetailLayout;
+};
