Add content browser buttons (find in, use from) to animation blueprint selector in skeletal mesh details panel

#TTP 330078 - Editor: No "Find in Browser" button or "Use Selected" for AnimBlueprint property of Skeletallmeshes

diff --git a/Engine/Source/Editor/DetailCustomizations/Private/SkeletalMeshComponentDetails.cpp b/Engine/Source/Editor/DetailCustomizations/Private/SkeletalMeshComponentDetails.cpp
--- a/Engine/Source/Editor/DetailCustomizations/Private/SkeletalMeshComponentDetails.cpp
+++ b/Engine/Source/Editor/DetailCustomizations/Private/SkeletalMeshComponentDetails.cpp
@@ -135,14 +135,35 @@ void FSkeletalMeshComponentDetails::UpdateAnimationCategory( IDetailLayoutBuilde
 		]
 	.ValueContent()
 		[
-			SAssignNew(ClassPickerComboButton, SComboButton)
-			.OnGetMenuContent(this, &FSkeletalMeshComponentDetails::GetClassPickerMenuContent)
-			.ContentPadding(0)
-			.ButtonContent()
+			SNew(SHorizontalBox)
+			+ SHorizontalBox::Slot()
+			.FillWidth(1.0f)
 			[
-				SNew(STextBlock)
-				.Font(IDetailLayoutBuilder::GetDetailFont())
-				.Text(this, &FSkeletalMeshComponentDetails::GetSelectedAnimBlueprintName)
+				SAssignNew(ClassPickerComboButton, SComboButton)
+				.OnGetMenuContent(this, &FSkeletalMeshComponentDetails::GetClassPickerMenuContent)
+				.ContentPadding(0)
+				.ButtonContent()
+				[
+					SNew(STextBlock)
+					.Font(IDetailLayoutBuilder::GetDetailFont())
+					.Text(this, &FSkeletalMeshComponentDetails::GetSelectedAnimBlueprintName)
+				]
+			]
+			+ SHorizontalBox::Slot()
+			.AutoWidth()
+			.HAlign(HAlign_Center)
+			.VAlign(VAlign_Center)
+			.Padding(2.0f, 1.0f)
+			[
+				PropertyCustomizationHelpers::MakeBrowseButton(FSimpleDelegate::CreateSP(this, &FSkeletalMeshComponentDetails::OnBrowseToAnimBlueprint))
+			]
+			+ SHorizontalBox::Slot()
+			.AutoWidth()
+			.HAlign(HAlign_Center)
+			.VAlign(VAlign_Center)
+			.Padding(2.0f, 1.0f)
+			[
+				PropertyCustomizationHelpers::MakeUseSelectedButton(FSimpleDelegate::CreateSP(this, &FSkeletalMeshComponentDetails::UseSelectedAnimBlueprint))
 			]
 		];
 
@@ -340,4 +361,37 @@ void FSkeletalMeshComponentDetails::OnClassPicked( UClass* PickedClass )
 	}
 }
 
+void FSkeletalMeshComponentDetails::OnBrowseToAnimBlueprint()
+{
+	check(AnimationBlueprintHandle->IsValidHandle());
+
+	UObject* Object = NULL;
+	AnimationBlueprintHandle->GetValue(Object);
+
+	TArray<UObject*> Objects;
+	Objects.Add(Object);
+	GEditor->SyncBrowserToObjects(Objects);
+}
+
+void FSkeletalMeshComponentDetails::UseSelectedAnimBlueprint()
+{
+	FEditorDelegates::LoadSelectedAssetsIfNeeded.Broadcast();
+
+	USelection* AssetSelection = GEditor->GetSelectedObjects();
+	if (AssetSelection && AssetSelection->Num() == 1)
+	{
+		UAnimBlueprint* AnimBlueprintToAssign = AssetSelection->GetTop<UAnimBlueprint>();
+		if (AnimBlueprintToAssign)
+		{
+			if(USkeleton* AnimBlueprintSkeleton = AnimBlueprintToAssign->TargetSkeleton)
+			{
+				FString BlueprintSkeletonName = FString::Printf(TEXT("%s'%s'"), *AnimBlueprintSkeleton->GetClass()->GetName(), *AnimBlueprintSkeleton->GetPathName());
+				if (BlueprintSkeletonName == SelectedSkeletonName)
+				{
+					OnClassPicked(AnimBlueprintToAssign->GetAnimBlueprintGeneratedClass());
+				}
+			}
+		}
+	}
+}
 #undef LOCTEXT_NAMESPACE
diff --git a/Engine/Source/Editor/DetailCustomizations/Private/SkeletalMeshComponentDetails.h b/Engine/Source/Editor/DetailCustomizations/Private/SkeletalMeshComponentDetails.h
--- a/Engine/Source/Editor/DetailCustomizations/Private/SkeletalMeshComponentDetails.h
+++ b/Engine/Source/Editor/DetailCustomizations/Private/SkeletalMeshComponentDetails.h
@@ -48,6 +48,12 @@ private:
 	/** Callback from the class picker when the user selects a class */
 	void OnClassPicked(UClass* PickedClass);
 
+	/** Callback from the detail panel to browse to the selected anim asset */
+	void OnBrowseToAnimBlueprint();
+
+	/** Callback from the details panel to use the currently selected asset in the content browser */
+	void UseSelectedAnimBlueprint();
+
 	/** Cached layout builder for use after customization */
 	IDetailLayoutBuilder* CurrentDetailBuilder;
 
