Added ability to customize any low level property type.  This will replace the struct customization interface

diff --git a/Engine/Source/Editor/DetailCustomizations/Private/DetailCustomizations.cpp b/Engine/Source/Editor/DetailCustomizations/Private/DetailCustomizations.cpp
--- a/Engine/Source/Editor/DetailCustomizations/Private/DetailCustomizations.cpp
+++ b/Engine/Source/Editor/DetailCustomizations/Private/DetailCustomizations.cpp
@@ -84,118 +84,9 @@ void FDetailCustomizationsModule::StartupModule()
 {
 	FPropertyEditorModule& PropertyModule = FModuleManager::GetModuleChecked<FPropertyEditorModule>("PropertyEditor");
 	
-	RegisterStructPropertyLayout( "StringAssetReference", FOnGetStructCustomizationInstance::CreateStatic( &FStringAssetReferenceCustomization::MakeInstance ) );
-	RegisterStructPropertyLayout( "StringClassReference", FOnGetStructCustomizationInstance::CreateStatic( &FStringClassReferenceCustomization::MakeInstance ) );
-	RegisterStructPropertyLayout( "DataTableRowHandle", FOnGetStructCustomizationInstance::CreateStatic( &FDataTableCustomizationLayout::MakeInstance ) );
-	RegisterStructPropertyLayout( "DataTableCategoryHandle", FOnGetStructCustomizationInstance::CreateStatic( &FDataTableCategoryCustomizationLayout::MakeInstance ) );
-	RegisterStructPropertyLayout( "CurveTableRowHandle", FOnGetStructCustomizationInstance::CreateStatic( &FCurveTableCustomizationLayout::MakeInstance ) );
-	RegisterStructPropertyLayout( NAME_Vector, FOnGetStructCustomizationInstance::CreateStatic( &FMathStructCustomization::MakeInstance ) );
-	RegisterStructPropertyLayout( NAME_Vector4, FOnGetStructCustomizationInstance::CreateStatic( &FMathStructCustomization::MakeInstance ) );
-	RegisterStructPropertyLayout( NAME_Vector2D, FOnGetStructCustomizationInstance::CreateStatic( &FMathStructCustomization::MakeInstance ) );
-	RegisterStructPropertyLayout( NAME_IntPoint, FOnGetStructCustomizationInstance::CreateStatic( &FMathStructCustomization::MakeInstance ) );
-	RegisterStructPropertyLayout( NAME_Rotator, FOnGetStructCustomizationInstance::CreateStatic( &FRotatorStructCustomization::MakeInstance ) );
-	RegisterStructPropertyLayout( NAME_LinearColor, FOnGetStructCustomizationInstance::CreateStatic( &FColorStructCustomization::MakeInstance ) );
-	RegisterStructPropertyLayout( NAME_Color, FOnGetStructCustomizationInstance::CreateStatic( &FColorStructCustomization::MakeInstance ) );
-	RegisterStructPropertyLayout( NAME_Matrix, FOnGetStructCustomizationInstance::CreateStatic( &FMatrixStructCustomization::MakeInstance ) );
-	RegisterStructPropertyLayout( NAME_Transform, FOnGetStructCustomizationInstance::CreateStatic( &FTransformStructCustomization::MakeInstance ) );
-	RegisterStructPropertyLayout( "SlateColor", FOnGetStructCustomizationInstance::CreateStatic( &FSlateColorCustomization::MakeInstance ) );
-	RegisterStructPropertyLayout( "AttenuationSettings", FOnGetStructCustomizationInstance::CreateStatic( &FAttenuationSettingsCustomization::MakeInstance ) );
-	RegisterStructPropertyLayout( "DialogueContext", FOnGetStructCustomizationInstance::CreateStatic( &FDialogueContextStructCustomization::MakeInstance ) );
-	RegisterStructPropertyLayout( "DialogueWaveParameter", FOnGetStructCustomizationInstance::CreateStatic( &FDialogueWaveParameterStructCustomization::MakeInstance ) );
-	RegisterStructPropertyLayout( "BodyInstance", FOnGetStructCustomizationInstance::CreateStatic( &FBodyInstanceCustomization::MakeInstance ) );
-	RegisterStructPropertyLayout( "BodyInstance", FOnGetStructCustomizationInstance::CreateStatic(&FBodyInstanceCustomization::MakeInstance));
-	RegisterStructPropertyLayout( "VehicleTransmissionData", FOnGetStructCustomizationInstance::CreateStatic(&FVehicleTransmissionDataCustomization::MakeInstance));
-	RegisterStructPropertyLayout( "SlateBrush", FOnGetStructCustomizationInstance::CreateStatic( &FSlateBrushStructCustomization::MakeInstance ) );
-	RegisterStructPropertyLayout( "SlateSound", FOnGetStructCustomizationInstance::CreateStatic( &FSlateSoundStructCustomization::MakeInstance ) );
-	RegisterStructPropertyLayout( "Guid", FOnGetStructCustomizationInstance::CreateStatic( &FGuidStructCustomization::MakeInstance ) );
-	RegisterStructPropertyLayout( "Key", FOnGetStructCustomizationInstance::CreateStatic(&FKeyStructCustomization::MakeInstance));
-	RegisterStructPropertyLayout( "FloatRange", FOnGetStructCustomizationInstance::CreateStatic( &FRangeStructCustomization<float>::MakeInstance ) );
-	RegisterStructPropertyLayout( "Int32Range", FOnGetStructCustomizationInstance::CreateStatic( &FRangeStructCustomization<int32>::MakeInstance ) );
-
-	RegisterStructPropertyLayout( "BlackboardEntry", FOnGetStructCustomizationInstance::CreateStatic( &FBlackboardEntryDetails::MakeInstance ) );
-	RegisterStructPropertyLayout( "RuntimeFloatCurve", FOnGetStructCustomizationInstance::CreateStatic( &FCurveStructCustomization::MakeInstance ) );
-	RegisterStructPropertyLayout( "EnvNamedValue", FOnGetStructCustomizationInstance::CreateStatic( &FEnvQueryParamInstanceCustomization::MakeInstance ) );
-	RegisterStructPropertyLayout( "NavigationLink", FOnGetStructCustomizationInstance::CreateStatic( &FNavLinkStructCustomization::MakeInstance ) );
-	RegisterStructPropertyLayout( "NavigationSegmentLink", FOnGetStructCustomizationInstance::CreateStatic( &FNavLinkStructCustomization::MakeInstance ) );
-
-	RegisterStructPropertyLayout( "Margin", FOnGetStructCustomizationInstance::CreateStatic( &FMarginStructCustomization::MakeInstance ) );
-
-	RegisterStructPropertyLayout( "DirectoryPath", FOnGetStructCustomizationInstance::CreateStatic( &FDirectoryPathStructCustomization::MakeInstance ) );
-	RegisterStructPropertyLayout( "FilePath", FOnGetStructCustomizationInstance::CreateStatic( &FFilePathStructCustomization::MakeInstance ) );
-
-	RegisterStructPropertyLayout( "InputAxisConfigEntry", FOnGetStructCustomizationInstance::CreateStatic( &FInputAxisConfigCustomization::MakeInstance ) );
-	RegisterStructPropertyLayout( "InputActionKeyMapping", FOnGetStructCustomizationInstance::CreateStatic( &FInputActionMappingCustomization::MakeInstance ) );
-	RegisterStructPropertyLayout( "InputAxisKeyMapping", FOnGetStructCustomizationInstance::CreateStatic( &FInputAxisMappingCustomization::MakeInstance ) );
+	RegisterPropertyTypeCustomizations();
 
-
-	// Note: By default properties are displayed in script defined order (i.e the order in the header).  These layout detail classes are called in the order seen here which will display properties
-	// in the order they are customized.  This is only relevant for inheritance where both a child and a parent have properties that are customized.
-	// In the order below, Actor will get a chance to display details first, followed by USceneComponent.
-
-	RegisterCustomPropertyLayout( "Actor", FOnGetDetailCustomizationInstance::CreateStatic( &FActorDetails::MakeInstance ) );
-	RegisterCustomPropertyLayout( "SceneComponent", FOnGetDetailCustomizationInstance::CreateStatic( &FSceneComponentDetails::MakeInstance ) );
-	RegisterCustomPropertyLayout( "PrimitiveComponent", FOnGetDetailCustomizationInstance::CreateStatic( &FPrimitiveComponentDetails::MakeInstance ) );
-	RegisterCustomPropertyLayout( "StaticMeshComponent", FOnGetDetailCustomizationInstance::CreateStatic( &FStaticMeshComponentDetails::MakeInstance ) );
-	RegisterCustomPropertyLayout( "SkeletalMeshComponent", FOnGetDetailCustomizationInstance::CreateStatic( &FSkeletalMeshComponentDetails::MakeInstance ) );
-	RegisterCustomPropertyLayout( "SkinnedMeshComponent", FOnGetDetailCustomizationInstance::CreateStatic( &FSkinnedMeshComponentDetails::MakeInstance ) );
-	RegisterCustomPropertyLayout( "WheeledVehicleMovementComponent4W", FOnGetDetailCustomizationInstance::CreateStatic(&FWheeledVehicleMovementComponent4WDetails::MakeInstance));
-	RegisterCustomPropertyLayout( "LightComponent", FOnGetDetailCustomizationInstance::CreateStatic( &FLightComponentDetails::MakeInstance ) );
-	RegisterCustomPropertyLayout( "PointLightComponent", FOnGetDetailCustomizationInstance::CreateStatic( &FPointLightComponentDetails::MakeInstance ) );
-	RegisterCustomPropertyLayout( "DirectionalLightComponent", FOnGetDetailCustomizationInstance::CreateStatic( &FDirectionalLightComponentDetails::MakeInstance ) );
-	RegisterCustomPropertyLayout( "StaticMeshActor", FOnGetDetailCustomizationInstance::CreateStatic( &FStaticMeshActorDetails::MakeInstance ) );
-	RegisterCustomPropertyLayout( "MeshComponent", FOnGetDetailCustomizationInstance::CreateStatic( &FMeshComponentDetails::MakeInstance ) );
-	RegisterCustomPropertyLayout( "MatineeActor", FOnGetDetailCustomizationInstance::CreateStatic( &FMatineeActorDetails::MakeInstance ) );
-	RegisterCustomPropertyLayout( "ReflectionCapture", FOnGetDetailCustomizationInstance::CreateStatic( &FReflectionCaptureDetails::MakeInstance ) );
-	RegisterCustomPropertyLayout( "SkyLight", FOnGetDetailCustomizationInstance::CreateStatic( &FSkyLightComponentDetails::MakeInstance ) );
-	RegisterCustomPropertyLayout( "Brush", FOnGetDetailCustomizationInstance::CreateStatic( &FBrushDetails::MakeInstance ) );
-	RegisterCustomPropertyLayout( "AmbientSound", FOnGetDetailCustomizationInstance::CreateStatic( &FAmbientSoundDetails::MakeInstance ) );
-	RegisterCustomPropertyLayout( "WorldSettings", FOnGetDetailCustomizationInstance::CreateStatic( &FWorldSettingsDetails::MakeInstance ) );
-	RegisterCustomPropertyLayout( "GeneralProjectSettings", FOnGetDetailCustomizationInstance::CreateStatic(&FGeneralProjectSettingsDetails::MakeInstance));
-
-	//@TODO: A2REMOVAL: Rename FSkeletalControlNodeDetails to something more generic
-	RegisterCustomPropertyLayout( "K2Node_StructMemberGet", FOnGetDetailCustomizationInstance::CreateStatic( &FSkeletalControlNodeDetails::MakeInstance ) );
-	RegisterCustomPropertyLayout( "K2Node_StructMemberSet", FOnGetDetailCustomizationInstance::CreateStatic( &FSkeletalControlNodeDetails::MakeInstance ) );
-
-	RegisterCustomPropertyLayout( "AnimSequence", FOnGetDetailCustomizationInstance::CreateStatic( &FAnimSequenceDetails::MakeInstance ) );
-
-	RegisterCustomPropertyLayout( "EditorAnimSegment", FOnGetDetailCustomizationInstance::CreateStatic( &FAnimMontageSegmentDetails::MakeInstance ) );
-	RegisterCustomPropertyLayout( "EditorAnimCompositeSegment", FOnGetDetailCustomizationInstance::CreateStatic( &FAnimMontageSegmentDetails::MakeInstance ) );
-	RegisterCustomPropertyLayout( "EditorSkeletonNotifyObj", FOnGetDetailCustomizationInstance::CreateStatic( &FSkeletonNotifyDetails::MakeInstance ) );
-	RegisterCustomPropertyLayout( "AnimStateNode", FOnGetDetailCustomizationInstance::CreateStatic( &FAnimStateNodeDetails::MakeInstance ) );
-	RegisterCustomPropertyLayout( "AnimStateTransitionNode", FOnGetDetailCustomizationInstance::CreateStatic( &FAnimTransitionNodeDetails::MakeInstance ) );
-	RegisterCustomPropertyLayout( "AnimGraphNode_TransitionPoseEvaluator", FOnGetDetailCustomizationInstance::CreateStatic( &FTransitionPoseEvaluatorNodeDetails::MakeInstance ) );
-
-	RegisterCustomPropertyLayout( "LandscapeSplineControlPoint", FOnGetDetailCustomizationInstance::CreateStatic( &FLandscapeSplineDetails::MakeInstance ) );
-	RegisterCustomPropertyLayout( "LandscapeSplineSegment", FOnGetDetailCustomizationInstance::CreateStatic( &FLandscapeSplineDetails::MakeInstance ) );
-
-	RegisterCustomPropertyLayout( "DialogueWave", FOnGetDetailCustomizationInstance::CreateStatic( &FDialogueWaveDetails::MakeInstance ) );
-	RegisterCustomPropertyLayout( "BodySetup", FOnGetDetailCustomizationInstance::CreateStatic( &FBodySetupDetails::MakeInstance ) );
-	RegisterCustomPropertyLayout( "PhysicsConstraintTemplate", FOnGetDetailCustomizationInstance::CreateStatic( &FPhysicsConstraintComponentDetails::MakeInstance ) );
-	RegisterCustomPropertyLayout( "PhysicsConstraintComponent", FOnGetDetailCustomizationInstance::CreateStatic( &FPhysicsConstraintComponentDetails::MakeInstance ) );
-	RegisterCustomPropertyLayout( "CollisionProfile", FOnGetDetailCustomizationInstance::CreateStatic(&FCollisionProfileDetails::MakeInstance));
-	RegisterCustomPropertyLayout( "PhysicsSettings", FOnGetDetailCustomizationInstance::CreateStatic(&FPhysicsSettingsDetails::MakeInstance));
-
-	RegisterCustomPropertyLayout( "ParticleModuleRequired", FOnGetDetailCustomizationInstance::CreateStatic( &FParticleModuleRequiredDetails::MakeInstance ) );
-	RegisterCustomPropertyLayout( "ParticleModuleSubUV", FOnGetDetailCustomizationInstance::CreateStatic( &FParticleModuleSubUVDetails::MakeInstance ) );
-	RegisterCustomPropertyLayout( "ParticleModuleAccelerationDrag", FOnGetDetailCustomizationInstance::CreateStatic( &FParticleModuleAccelerationDragDetails::MakeInstance ) );
-	RegisterCustomPropertyLayout( "ParticleModuleAcceleration", FOnGetDetailCustomizationInstance::CreateStatic( &FParticleModuleAccelerationDetails::MakeInstance ) );
-	RegisterCustomPropertyLayout( "ParticleModuleAccelerationDragScaleOverLife", FOnGetDetailCustomizationInstance::CreateStatic( &FParticleModuleAccelerationDragScaleOverLifeDetails::MakeInstance ) );
-	RegisterCustomPropertyLayout( "ParticleModuleCollisionGPU", FOnGetDetailCustomizationInstance::CreateStatic( &FParticleModuleCollisionGPUDetails::MakeInstance ) );
-	RegisterCustomPropertyLayout( "ParticleModuleOrbit", FOnGetDetailCustomizationInstance::CreateStatic( &FParticleModuleOrbitDetails::MakeInstance ) );
-	RegisterCustomPropertyLayout( "ParticleModuleSizeMultiplyLife", FOnGetDetailCustomizationInstance::CreateStatic( &FParticleModuleSizeMultiplyLifeDetails::MakeInstance ) );
-	RegisterCustomPropertyLayout( "ParticleModuleSizeScale", FOnGetDetailCustomizationInstance::CreateStatic( &FParticleModuleSizeScaleDetails::MakeInstance ) );
-	RegisterCustomPropertyLayout( "ParticleModuleVectorFieldScale", FOnGetDetailCustomizationInstance::CreateStatic( &FParticleModuleVectorFieldScaleDetails::MakeInstance ) );
-	RegisterCustomPropertyLayout( "ParticleModuleVectorFieldScaleOverLife", FOnGetDetailCustomizationInstance::CreateStatic( &FParticleModuleVectorFieldScaleOverLifeDetails::MakeInstance ) );
-
-	RegisterCustomPropertyLayout( "CameraComponent", FOnGetDetailCustomizationInstance::CreateStatic( &FCameraDetails::MakeInstance ) );
-	RegisterCustomPropertyLayout( "DeviceProfile", FOnGetDetailCustomizationInstance::CreateStatic( &FDeviceProfileDetails::MakeInstance ) );
-	RegisterCustomPropertyLayout( "InternationalizationSettingsModel", FOnGetDetailCustomizationInstance::CreateStatic( &FInternationalizationSettingsModelDetails::MakeInstance ) );
-	RegisterCustomPropertyLayout( "InputSettings", FOnGetDetailCustomizationInstance::CreateStatic( &FInputSettingsDetails::MakeInstance ) );
-	RegisterCustomPropertyLayout( "WindowsTargetSettings", FOnGetDetailCustomizationInstance::CreateStatic( &FWindowsTargetSettingsDetails::MakeInstance ) );
-	RegisterCustomPropertyLayout( "MoviePlayerSettings", FOnGetDetailCustomizationInstance::CreateStatic( &FMoviePlayerSettingsDetails::MakeInstance ) );
-
-	RegisterCustomPropertyLayout("SourceCodeAccessSettings", FOnGetDetailCustomizationInstance::CreateStatic(&FSourceCodeAccessSettingsDetails::MakeInstance));
-	RegisterCustomPropertyLayout("ParticleSystemComponent", FOnGetDetailCustomizationInstance::CreateStatic(&FParticleSystemComponentDetails::MakeInstance));
+	RegisterObjectCustomizations();
 
 	PropertyModule.NotifyCustomizationModuleChanged();
 }
@@ -217,11 +108,11 @@ void FDetailCustomizationsModule::ShutdownModule()
 		}
 
 		// Unregister all structs
-		for(auto It = RegisteredStructs.CreateConstIterator(); It; ++It)
+		for(auto It = RegisteredPropertyTypes.CreateConstIterator(); It; ++It)
 		{
 			if(It->IsValid())
 			{
-				PropertyModule.UnregisterStructPropertyLayout(*It);
+				PropertyModule.UnregisterCustomPropertyTypeLayout(*It);
 			}
 		}
 	
@@ -230,26 +121,161 @@ void FDetailCustomizationsModule::ShutdownModule()
 }
 
 
+void FDetailCustomizationsModule::RegisterPropertyTypeCustomizations()
+{
+	// Structs
+	RegisterStructPropertyLayout("StringAssetReference", FOnGetStructCustomizationInstance::CreateStatic(&FStringAssetReferenceCustomization::MakeInstance));
+	RegisterStructPropertyLayout("StringClassReference", FOnGetStructCustomizationInstance::CreateStatic(&FStringClassReferenceCustomization::MakeInstance));
+	RegisterStructPropertyLayout("DataTableRowHandle", FOnGetStructCustomizationInstance::CreateStatic(&FDataTableCustomizationLayout::MakeInstance));
+	RegisterStructPropertyLayout("DataTableCategoryHandle", FOnGetStructCustomizationInstance::CreateStatic(&FDataTableCategoryCustomizationLayout::MakeInstance));
+	RegisterStructPropertyLayout("CurveTableRowHandle", FOnGetStructCustomizationInstance::CreateStatic(&FCurveTableCustomizationLayout::MakeInstance));
+	RegisterStructPropertyLayout(NAME_Vector, FOnGetStructCustomizationInstance::CreateStatic(&FMathStructCustomization::MakeInstance));
+	RegisterStructPropertyLayout(NAME_Vector4, FOnGetStructCustomizationInstance::CreateStatic(&FMathStructCustomization::MakeInstance));
+	RegisterStructPropertyLayout(NAME_Vector2D, FOnGetStructCustomizationInstance::CreateStatic(&FMathStructCustomization::MakeInstance));
+	RegisterStructPropertyLayout(NAME_IntPoint, FOnGetStructCustomizationInstance::CreateStatic(&FMathStructCustomization::MakeInstance));
+	RegisterStructPropertyLayout(NAME_Rotator, FOnGetStructCustomizationInstance::CreateStatic(&FRotatorStructCustomization::MakeInstance));
+	RegisterStructPropertyLayout(NAME_LinearColor, FOnGetStructCustomizationInstance::CreateStatic(&FColorStructCustomization::MakeInstance));
+	RegisterStructPropertyLayout(NAME_Color, FOnGetStructCustomizationInstance::CreateStatic(&FColorStructCustomization::MakeInstance));
+	RegisterStructPropertyLayout(NAME_Matrix, FOnGetStructCustomizationInstance::CreateStatic(&FMatrixStructCustomization::MakeInstance));
+	RegisterStructPropertyLayout(NAME_Transform, FOnGetStructCustomizationInstance::CreateStatic(&FTransformStructCustomization::MakeInstance));
+	RegisterStructPropertyLayout("SlateColor", FOnGetStructCustomizationInstance::CreateStatic(&FSlateColorCustomization::MakeInstance));
+	RegisterStructPropertyLayout("AttenuationSettings", FOnGetStructCustomizationInstance::CreateStatic(&FAttenuationSettingsCustomization::MakeInstance));
+	RegisterStructPropertyLayout("DialogueContext", FOnGetStructCustomizationInstance::CreateStatic(&FDialogueContextStructCustomization::MakeInstance));
+	RegisterStructPropertyLayout("DialogueWaveParameter", FOnGetStructCustomizationInstance::CreateStatic(&FDialogueWaveParameterStructCustomization::MakeInstance));
+	RegisterStructPropertyLayout("BodyInstance", FOnGetStructCustomizationInstance::CreateStatic(&FBodyInstanceCustomization::MakeInstance));
+	RegisterStructPropertyLayout("VehicleTransmissionData", FOnGetStructCustomizationInstance::CreateStatic(&FVehicleTransmissionDataCustomization::MakeInstance));
+	RegisterStructPropertyLayout("SlateBrush", FOnGetStructCustomizationInstance::CreateStatic(&FSlateBrushStructCustomization::MakeInstance));
+	RegisterStructPropertyLayout("SlateSound", FOnGetStructCustomizationInstance::CreateStatic(&FSlateSoundStructCustomization::MakeInstance));
+	RegisterStructPropertyLayout("Guid", FOnGetStructCustomizationInstance::CreateStatic(&FGuidStructCustomization::MakeInstance));
+	RegisterStructPropertyLayout("Key", FOnGetStructCustomizationInstance::CreateStatic(&FKeyStructCustomization::MakeInstance));
+	RegisterStructPropertyLayout("FloatRange", FOnGetStructCustomizationInstance::CreateStatic(&FRangeStructCustomization<float>::MakeInstance));
+	RegisterStructPropertyLayout("Int32Range", FOnGetStructCustomizationInstance::CreateStatic(&FRangeStructCustomization<int32>::MakeInstance));
+
+	RegisterStructPropertyLayout("BlackboardEntry", FOnGetStructCustomizationInstance::CreateStatic(&FBlackboardEntryDetails::MakeInstance));
+	RegisterStructPropertyLayout("RuntimeFloatCurve", FOnGetStructCustomizationInstance::CreateStatic(&FCurveStructCustomization::MakeInstance));
+	RegisterStructPropertyLayout("EnvNamedValue", FOnGetStructCustomizationInstance::CreateStatic(&FEnvQueryParamInstanceCustomization::MakeInstance));
+	RegisterStructPropertyLayout("NavigationLink", FOnGetStructCustomizationInstance::CreateStatic(&FNavLinkStructCustomization::MakeInstance));
+	RegisterStructPropertyLayout("NavigationSegmentLink", FOnGetStructCustomizationInstance::CreateStatic(&FNavLinkStructCustomization::MakeInstance));
+
+	RegisterStructPropertyLayout("Margin", FOnGetStructCustomizationInstance::CreateStatic(&FMarginStructCustomization::MakeInstance));
+
+	RegisterStructPropertyLayout("DirectoryPath", FOnGetStructCustomizationInstance::CreateStatic(&FDirectoryPathStructCustomization::MakeInstance));
+	RegisterStructPropertyLayout("FilePath", FOnGetStructCustomizationInstance::CreateStatic(&FFilePathStructCustomization::MakeInstance));
+
+	RegisterStructPropertyLayout("InputAxisConfigEntry", FOnGetStructCustomizationInstance::CreateStatic(&FInputAxisConfigCustomization::MakeInstance));
+	RegisterStructPropertyLayout("InputActionKeyMapping", FOnGetStructCustomizationInstance::CreateStatic(&FInputActionMappingCustomization::MakeInstance));
+	RegisterStructPropertyLayout("InputAxisKeyMapping", FOnGetStructCustomizationInstance::CreateStatic(&FInputAxisMappingCustomization::MakeInstance));
+}
+
+void FDetailCustomizationsModule::RegisterObjectCustomizations()
+{
+	// Note: By default properties are displayed in script defined order (i.e the order in the header).  These layout detail classes are called in the order seen here which will display properties
+	// in the order they are customized.  This is only relevant for inheritance where both a child and a parent have properties that are customized.
+	// In the order below, Actor will get a chance to display details first, followed by USceneComponent.
+
+	RegisterCustomPropertyLayout("Actor", FOnGetDetailCustomizationInstance::CreateStatic(&FActorDetails::MakeInstance));
+	RegisterCustomPropertyLayout("SceneComponent", FOnGetDetailCustomizationInstance::CreateStatic(&FSceneComponentDetails::MakeInstance));
+	RegisterCustomPropertyLayout("PrimitiveComponent", FOnGetDetailCustomizationInstance::CreateStatic(&FPrimitiveComponentDetails::MakeInstance));
+	RegisterCustomPropertyLayout("StaticMeshComponent", FOnGetDetailCustomizationInstance::CreateStatic(&FStaticMeshComponentDetails::MakeInstance));
+	RegisterCustomPropertyLayout("SkeletalMeshComponent", FOnGetDetailCustomizationInstance::CreateStatic(&FSkeletalMeshComponentDetails::MakeInstance));
+	RegisterCustomPropertyLayout("SkinnedMeshComponent", FOnGetDetailCustomizationInstance::CreateStatic(&FSkinnedMeshComponentDetails::MakeInstance));
+	RegisterCustomPropertyLayout("WheeledVehicleMovementComponent4W", FOnGetDetailCustomizationInstance::CreateStatic(&FWheeledVehicleMovementComponent4WDetails::MakeInstance));
+	RegisterCustomPropertyLayout("LightComponent", FOnGetDetailCustomizationInstance::CreateStatic(&FLightComponentDetails::MakeInstance));
+	RegisterCustomPropertyLayout("PointLightComponent", FOnGetDetailCustomizationInstance::CreateStatic(&FPointLightComponentDetails::MakeInstance));
+	RegisterCustomPropertyLayout("DirectionalLightComponent", FOnGetDetailCustomizationInstance::CreateStatic(&FDirectionalLightComponentDetails::MakeInstance));
+	RegisterCustomPropertyLayout("StaticMeshActor", FOnGetDetailCustomizationInstance::CreateStatic(&FStaticMeshActorDetails::MakeInstance));
+	RegisterCustomPropertyLayout("MeshComponent", FOnGetDetailCustomizationInstance::CreateStatic(&FMeshComponentDetails::MakeInstance));
+	RegisterCustomPropertyLayout("MatineeActor", FOnGetDetailCustomizationInstance::CreateStatic(&FMatineeActorDetails::MakeInstance));
+	RegisterCustomPropertyLayout("ReflectionCapture", FOnGetDetailCustomizationInstance::CreateStatic(&FReflectionCaptureDetails::MakeInstance));
+	RegisterCustomPropertyLayout("SkyLight", FOnGetDetailCustomizationInstance::CreateStatic(&FSkyLightComponentDetails::MakeInstance));
+	RegisterCustomPropertyLayout("Brush", FOnGetDetailCustomizationInstance::CreateStatic(&FBrushDetails::MakeInstance));
+	RegisterCustomPropertyLayout("AmbientSound", FOnGetDetailCustomizationInstance::CreateStatic(&FAmbientSoundDetails::MakeInstance));
+	RegisterCustomPropertyLayout("WorldSettings", FOnGetDetailCustomizationInstance::CreateStatic(&FWorldSettingsDetails::MakeInstance));
+	RegisterCustomPropertyLayout("GeneralProjectSettings", FOnGetDetailCustomizationInstance::CreateStatic(&FGeneralProjectSettingsDetails::MakeInstance));
+
+	//@TODO: A2REMOVAL: Rename FSkeletalControlNodeDetails to something more generic
+	RegisterCustomPropertyLayout("K2Node_StructMemberGet", FOnGetDetailCustomizationInstance::CreateStatic(&FSkeletalControlNodeDetails::MakeInstance));
+	RegisterCustomPropertyLayout("K2Node_StructMemberSet", FOnGetDetailCustomizationInstance::CreateStatic(&FSkeletalControlNodeDetails::MakeInstance));
+
+	RegisterCustomPropertyLayout("AnimSequence", FOnGetDetailCustomizationInstance::CreateStatic(&FAnimSequenceDetails::MakeInstance));
+
+	RegisterCustomPropertyLayout("EditorAnimSegment", FOnGetDetailCustomizationInstance::CreateStatic(&FAnimMontageSegmentDetails::MakeInstance));
+	RegisterCustomPropertyLayout("EditorAnimCompositeSegment", FOnGetDetailCustomizationInstance::CreateStatic(&FAnimMontageSegmentDetails::MakeInstance));
+	RegisterCustomPropertyLayout("EditorSkeletonNotifyObj", FOnGetDetailCustomizationInstance::CreateStatic(&FSkeletonNotifyDetails::MakeInstance));
+	RegisterCustomPropertyLayout("AnimStateNode", FOnGetDetailCustomizationInstance::CreateStatic(&FAnimStateNodeDetails::MakeInstance));
+	RegisterCustomPropertyLayout("AnimStateTransitionNode", FOnGetDetailCustomizationInstance::CreateStatic(&FAnimTransitionNodeDetails::MakeInstance));
+	RegisterCustomPropertyLayout("AnimGraphNode_TransitionPoseEvaluator", FOnGetDetailCustomizationInstance::CreateStatic(&FTransitionPoseEvaluatorNodeDetails::MakeInstance));
+
+	RegisterCustomPropertyLayout("LandscapeSplineControlPoint", FOnGetDetailCustomizationInstance::CreateStatic(&FLandscapeSplineDetails::MakeInstance));
+	RegisterCustomPropertyLayout("LandscapeSplineSegment", FOnGetDetailCustomizationInstance::CreateStatic(&FLandscapeSplineDetails::MakeInstance));
+
+	RegisterCustomPropertyLayout("DialogueWave", FOnGetDetailCustomizationInstance::CreateStatic(&FDialogueWaveDetails::MakeInstance));
+	RegisterCustomPropertyLayout("BodySetup", FOnGetDetailCustomizationInstance::CreateStatic(&FBodySetupDetails::MakeInstance));
+	RegisterCustomPropertyLayout("PhysicsConstraintTemplate", FOnGetDetailCustomizationInstance::CreateStatic(&FPhysicsConstraintComponentDetails::MakeInstance));
+	RegisterCustomPropertyLayout("PhysicsConstraintComponent", FOnGetDetailCustomizationInstance::CreateStatic(&FPhysicsConstraintComponentDetails::MakeInstance));
+	RegisterCustomPropertyLayout("CollisionProfile", FOnGetDetailCustomizationInstance::CreateStatic(&FCollisionProfileDetails::MakeInstance));
+	RegisterCustomPropertyLayout("PhysicsSettings", FOnGetDetailCustomizationInstance::CreateStatic(&FPhysicsSettingsDetails::MakeInstance));
+
+	RegisterCustomPropertyLayout("ParticleModuleRequired", FOnGetDetailCustomizationInstance::CreateStatic(&FParticleModuleRequiredDetails::MakeInstance));
+	RegisterCustomPropertyLayout("ParticleModuleSubUV", FOnGetDetailCustomizationInstance::CreateStatic(&FParticleModuleSubUVDetails::MakeInstance));
+	RegisterCustomPropertyLayout("ParticleModuleAccelerationDrag", FOnGetDetailCustomizationInstance::CreateStatic(&FParticleModuleAccelerationDragDetails::MakeInstance));
+	RegisterCustomPropertyLayout("ParticleModuleAcceleration", FOnGetDetailCustomizationInstance::CreateStatic(&FParticleModuleAccelerationDetails::MakeInstance));
+	RegisterCustomPropertyLayout("ParticleModuleAccelerationDragScaleOverLife", FOnGetDetailCustomizationInstance::CreateStatic(&FParticleModuleAccelerationDragScaleOverLifeDetails::MakeInstance));
+	RegisterCustomPropertyLayout("ParticleModuleCollisionGPU", FOnGetDetailCustomizationInstance::CreateStatic(&FParticleModuleCollisionGPUDetails::MakeInstance));
+	RegisterCustomPropertyLayout("ParticleModuleOrbit", FOnGetDetailCustomizationInstance::CreateStatic(&FParticleModuleOrbitDetails::MakeInstance));
+	RegisterCustomPropertyLayout("ParticleModuleSizeMultiplyLife", FOnGetDetailCustomizationInstance::CreateStatic(&FParticleModuleSizeMultiplyLifeDetails::MakeInstance));
+	RegisterCustomPropertyLayout("ParticleModuleSizeScale", FOnGetDetailCustomizationInstance::CreateStatic(&FParticleModuleSizeScaleDetails::MakeInstance));
+	RegisterCustomPropertyLayout("ParticleModuleVectorFieldScale", FOnGetDetailCustomizationInstance::CreateStatic(&FParticleModuleVectorFieldScaleDetails::MakeInstance));
+	RegisterCustomPropertyLayout("ParticleModuleVectorFieldScaleOverLife", FOnGetDetailCustomizationInstance::CreateStatic(&FParticleModuleVectorFieldScaleOverLifeDetails::MakeInstance));
+
+	RegisterCustomPropertyLayout("CameraComponent", FOnGetDetailCustomizationInstance::CreateStatic(&FCameraDetails::MakeInstance));
+	RegisterCustomPropertyLayout("DeviceProfile", FOnGetDetailCustomizationInstance::CreateStatic(&FDeviceProfileDetails::MakeInstance));
+	RegisterCustomPropertyLayout("InternationalizationSettingsModel", FOnGetDetailCustomizationInstance::CreateStatic(&FInternationalizationSettingsModelDetails::MakeInstance));
+	RegisterCustomPropertyLayout("InputSettings", FOnGetDetailCustomizationInstance::CreateStatic(&FInputSettingsDetails::MakeInstance));
+	RegisterCustomPropertyLayout("WindowsTargetSettings", FOnGetDetailCustomizationInstance::CreateStatic(&FWindowsTargetSettingsDetails::MakeInstance));
+	RegisterCustomPropertyLayout("MoviePlayerSettings", FOnGetDetailCustomizationInstance::CreateStatic(&FMoviePlayerSettingsDetails::MakeInstance));
+
+	RegisterCustomPropertyLayout("SourceCodeAccessSettings", FOnGetDetailCustomizationInstance::CreateStatic(&FSourceCodeAccessSettingsDetails::MakeInstance));
+	RegisterCustomPropertyLayout("ParticleSystemComponent", FOnGetDetailCustomizationInstance::CreateStatic(&FParticleSystemComponentDetails::MakeInstance));
+}
+
 void FDetailCustomizationsModule::RegisterCustomPropertyLayout(FName ClassName, FOnGetDetailCustomizationInstance DetailLayoutDelegate )
 {
 	check( ClassName != NAME_None );
 
+	static FName PropertyEditor("PropertyEditor");
+
 	// Add the class to the list of classes we should unregister later
 	RegisteredClassNames.Add( ClassName );
 
-	FPropertyEditorModule& PropertyModule = FModuleManager::GetModuleChecked<FPropertyEditorModule>("PropertyEditor");
+	FPropertyEditorModule& PropertyModule = FModuleManager::GetModuleChecked<FPropertyEditorModule>(PropertyEditor);
 
 	PropertyModule.RegisterCustomPropertyLayout( ClassName, DetailLayoutDelegate );
 }
 
 
-void FDetailCustomizationsModule::RegisterStructPropertyLayout(FName StructTypeName, FOnGetStructCustomizationInstance StructLayoutDelegate)
+void FDetailCustomizationsModule::RegisterCustomPropertyTypeLayout(FName PropertyTypeName, FOnGetPropertyTypeCustomizationInstance PropertyTypeLayoutDelegate)
 {
-	check(StructTypeName != NAME_None);
+	check(PropertyTypeName != NAME_None);
 
-	RegisteredStructs.Add(StructTypeName);
+	static FName PropertyEditor("PropertyEditor");
 
-	FPropertyEditorModule& PropertyModule = FModuleManager::GetModuleChecked<FPropertyEditorModule>("PropertyEditor");
+	RegisteredPropertyTypes.Add(PropertyTypeName);
+
+	FPropertyEditorModule& PropertyModule = FModuleManager::GetModuleChecked<FPropertyEditorModule>(PropertyEditor);
+
+	PropertyModule.RegisterCustomPropertyTypeLayout(PropertyTypeName, PropertyTypeLayoutDelegate);
+}
+
+void FDetailCustomizationsModule::RegisterStructPropertyLayout(FName StructPropertyName, FOnGetStructCustomizationInstance StructLayoutDelegate)
+{
+	check(StructPropertyName != NAME_None);
+
+	static FName PropertyEditor("PropertyEditor");
+
+	RegisteredPropertyTypes.Add(StructPropertyName);
+
+	FPropertyEditorModule& PropertyModule = FModuleManager::GetModuleChecked<FPropertyEditorModule>(PropertyEditor);
 
-	PropertyModule.RegisterStructPropertyLayout(StructTypeName, StructLayoutDelegate);
+	PropertyModule.RegisterStructPropertyLayout(StructPropertyName, StructLayoutDelegate);
 }
diff --git a/Engine/Source/Editor/DetailCustomizations/Public/DetailCustomizations.h b/Engine/Source/Editor/DetailCustomizations/Public/DetailCustomizations.h
--- a/Engine/Source/Editor/DetailCustomizations/Public/DetailCustomizations.h
+++ b/Engine/Source/Editor/DetailCustomizations/Public/DetailCustomizations.h
@@ -28,6 +28,8 @@ public:
 	}
 
 private:
+	void RegisterPropertyTypeCustomizations();
+	void RegisterObjectCustomizations();
 
 	/**
 	 * Registers a custom class
@@ -43,10 +45,11 @@ private:
 	* @param StructName				The name of the struct to register for property customization
 	* @param StructLayoutDelegate	The delegate to call to get the custom detail layout instance
 	*/
-	void RegisterStructPropertyLayout(FName StructTypeName, FOnGetStructCustomizationInstance StructLayoutDelegate);
+	void RegisterCustomPropertyTypeLayout(FName PropertyTypeName, FOnGetPropertyTypeCustomizationInstance PropertyTypeLayoutDelegate );
 
+	void RegisterStructPropertyLayout(FName StructPropertyName, FOnGetStructCustomizationInstance StructLayoutDelegate );
 private:
 	/** List of registered class that we must unregister when the module shuts down */
 	TSet< FName > RegisteredClassNames;
-	TSet< FName > RegisteredStructs;
+	TSet< FName > RegisteredPropertyTypes;
 };
diff --git a/Engine/Source/Editor/PropertyEditor/Private/CustomChildBuilder.cpp b/Engine/Source/Editor/PropertyEditor/Private/CustomChildBuilder.cpp
--- a/Engine/Source/Editor/PropertyEditor/Private/CustomChildBuilder.cpp
+++ b/Engine/Source/Editor/PropertyEditor/Private/CustomChildBuilder.cpp
@@ -50,21 +50,21 @@ IDetailPropertyRow& FCustomChildrenBuilder::AddChildProperty( TSharedRef<IProper
 	return *NewCustomization.PropertyRow;
 }
 
-class SStandaloneCustomStructValue : public SCompoundWidget, public IStructCustomizationUtils
+class SStandaloneCustomStructValue : public SCompoundWidget, public IPropertyTypeCustomizationUtils
 {
 public:
 	SLATE_BEGIN_ARGS( SStandaloneCustomStructValue )
 	{}
 	SLATE_END_ARGS()
 	
-	void Construct( const FArguments& InArgs, TSharedPtr<IStructCustomization> InCustomizationInterface, TSharedRef<IPropertyHandle> InStructPropertyHandle, TSharedRef<FDetailCategoryImpl> InParentCategory )
+	void Construct( const FArguments& InArgs, TSharedPtr<IPropertyTypeCustomization> InCustomizationInterface, TSharedRef<IPropertyHandle> InStructPropertyHandle, TSharedRef<FDetailCategoryImpl> InParentCategory )
 	{
 		CustomizationInterface = InCustomizationInterface;
 		StructPropertyHandle = InStructPropertyHandle;
 		ParentCategory = InParentCategory;
 		CustomPropertyWidget = MakeShareable(new FDetailWidgetRow);
 
-		CustomizationInterface->CustomizeStructHeader(InStructPropertyHandle, *CustomPropertyWidget, *this);
+		CustomizationInterface->CustomizeHeader(InStructPropertyHandle, *CustomPropertyWidget, *this);
 
 		ChildSlot
 		[
@@ -80,7 +80,7 @@ public:
 
 private:
 	TWeakPtr<FDetailCategoryImpl> ParentCategory;
-	TSharedPtr<IStructCustomization> CustomizationInterface;
+	TSharedPtr<IPropertyTypeCustomization> CustomizationInterface;
 	TSharedPtr<IPropertyHandle> StructPropertyHandle;
 	TSharedPtr<FDetailWidgetRow> CustomPropertyWidget;
 };
@@ -91,10 +91,10 @@ TSharedRef<SWidget> FCustomChildrenBuilder::GenerateStructValueWidget( TSharedRe
 	UStructProperty* StructProperty = CastChecked<UStructProperty>( StructPropertyHandle->GetProperty() );
 
 	FPropertyEditorModule& PropertyEditorModule = FModuleManager::GetModuleChecked<FPropertyEditorModule>("PropertyEditor");
-	FOnGetStructCustomizationInstance StructLayoutInstanceDelegate = PropertyEditorModule.GetStructCustomizaton(StructProperty->Struct);
-	if (StructLayoutInstanceDelegate.IsBound())
+	FPropertyTypeLayoutCallback LayoutCallback = PropertyEditorModule.GetPropertyTypeCustomization(StructProperty);
+	if (LayoutCallback.IsValid())
 	{
-		TSharedRef<IStructCustomization> CustomStructInterface = StructLayoutInstanceDelegate.Execute();
+		TSharedRef<IPropertyTypeCustomization> CustomStructInterface = LayoutCallback.GetCustomizationInstance();
 
 		return SNew( SStandaloneCustomStructValue, CustomStructInterface, StructPropertyHandle, ParentCategory.Pin().ToSharedRef() );
 	}
diff --git a/Engine/Source/Editor/PropertyEditor/Private/DetailPropertyRow.cpp b/Engine/Source/Editor/PropertyEditor/Private/DetailPropertyRow.cpp
--- a/Engine/Source/Editor/PropertyEditor/Private/DetailPropertyRow.cpp
+++ b/Engine/Source/Editor/PropertyEditor/Private/DetailPropertyRow.cpp
@@ -30,17 +30,17 @@ FDetailPropertyRow::FDetailPropertyRow(TSharedPtr<FPropertyNode> InPropertyNode,
 			MakePropertyEditor(InParentCategory->GetParentLayoutImpl().GetPropertyUtilities());
 		}
 
-		UStructProperty* StructProperty = Cast<UStructProperty>(Property);
-		// Check if the property is valid for struct customization.  Note: Static arrays of structs will be a UStructProperty with array elements as children
-		if (!PropertyEditorHelpers::IsStaticArray(*PropertyNodeRef) && StructProperty && StructProperty->Struct)
+		// Check if the property is valid for type customization.  Note: Static arrays of types will be a UProperty with array elements as children
+		if (!PropertyEditorHelpers::IsStaticArray(*PropertyNodeRef))
 		{
-			FPropertyEditorModule& PropertyEditorModule = FModuleManager::GetModuleChecked<FPropertyEditorModule>("PropertyEditor");
-			FOnGetStructCustomizationInstance StructLayoutInstanceDelegate = PropertyEditorModule.GetStructCustomizaton(StructProperty->Struct);
-			if (StructLayoutInstanceDelegate.IsBound())
+			static FName PropertyEditor("PropertyEditor");
+			FPropertyEditorModule& PropertyEditorModule = FModuleManager::GetModuleChecked<FPropertyEditorModule>(PropertyEditor);
+			FPropertyTypeLayoutCallback LayoutCallback = PropertyEditorModule.GetPropertyTypeCustomization(Property);
+			if (LayoutCallback.IsValid())
 			{
 				if (PropertyHandle->IsValidHandle())
 				{
-					CustomStructInterface = StructLayoutInstanceDelegate.Execute();
+					CustomTypeInterface = LayoutCallback.GetCustomizationInstance();
 				}
 			}
 		}
@@ -110,17 +110,17 @@ void FDetailPropertyRow::GetDefaultWidgets( TSharedPtr<SWidget>& OutNameWidget,
 
 void FDetailPropertyRow::GetDefaultWidgets( TSharedPtr<SWidget>& OutNameWidget, TSharedPtr<SWidget>& OutValueWidget, FDetailWidgetRow& Row )
 {
-	TSharedPtr<FDetailWidgetRow> CustomStructRow;
-	if ( CustomStructInterface.IsValid() ) 
+	TSharedPtr<FDetailWidgetRow> CustomTypeRow;
+	if ( CustomTypeInterface.IsValid() ) 
 	{
-		CustomStructRow = MakeShareable(new FDetailWidgetRow);
+		CustomTypeRow = MakeShareable(new FDetailWidgetRow);
 
-		CustomStructInterface->CustomizeStructHeader(PropertyHandle.ToSharedRef(), *CustomStructRow, *this);
+		CustomTypeInterface->CustomizeHeader(PropertyHandle.ToSharedRef(), *CustomTypeRow, *this);
 	}
 
 	const bool bAddWidgetDecoration = false;
-	MakeNameWidget(Row,CustomStructRow);
-	MakeValueWidget(Row,CustomStructRow,bAddWidgetDecoration);
+	MakeNameWidget(Row,CustomTypeRow);
+	MakeValueWidget(Row,CustomTypeRow,bAddWidgetDecoration);
 
 	OutNameWidget = Row.NameWidget.Widget;
 	OutValueWidget = Row.ValueWidget.Widget;
@@ -135,7 +135,7 @@ bool FDetailPropertyRow::HasColumns() const
 
 bool FDetailPropertyRow::ShowOnlyChildren() const
 {
-	return StructLayoutBuilder.IsValid() && CustomPropertyWidget.IsValid() && !CustomPropertyWidget->HasAnyContent();
+	return PropertyTypeLayoutBuilder.IsValid() && CustomPropertyWidget.IsValid() && !CustomPropertyWidget->HasAnyContent();
 }
 
 bool FDetailPropertyRow::RequiresTick() const
@@ -189,11 +189,11 @@ void FDetailPropertyRow::OnItemNodeInitialized( TSharedRef<FDetailCategoryImpl>
 	IsParentEnabled = InIsParentEnabled;
 
 	// Don't customize the user already customized
-	if( !CustomPropertyWidget.IsValid() && CustomStructInterface.IsValid() )
+	if( !CustomPropertyWidget.IsValid() && CustomTypeInterface.IsValid() )
 	{
 		CustomPropertyWidget = MakeShareable(new FDetailWidgetRow);
 
-		CustomStructInterface->CustomizeStructHeader(PropertyHandle.ToSharedRef(), *CustomPropertyWidget, *this);
+		CustomTypeInterface->CustomizeHeader(PropertyHandle.ToSharedRef(), *CustomPropertyWidget, *this);
 
 		// set initial value of enabled attribute to settings from struct customization
 		if (CustomPropertyWidget->IsEnabledAttr.IsBound())
@@ -202,10 +202,10 @@ void FDetailPropertyRow::OnItemNodeInitialized( TSharedRef<FDetailCategoryImpl>
 		}
 	}
 
-	if( bShowCustomPropertyChildren && CustomStructInterface.IsValid() )
+	if( bShowCustomPropertyChildren && CustomTypeInterface.IsValid() )
 	{
-		StructLayoutBuilder = MakeShareable(new FCustomChildrenBuilder(InParentCategory));
-		CustomStructInterface->CustomizeStructChildren(PropertyHandle.ToSharedRef(), *StructLayoutBuilder, *this);
+		PropertyTypeLayoutBuilder = MakeShareable(new FCustomChildrenBuilder(InParentCategory));
+		CustomTypeInterface->CustomizeChildren(PropertyHandle.ToSharedRef(), *PropertyTypeLayoutBuilder, *this);
 	}
 }
 
@@ -227,9 +227,9 @@ void FDetailPropertyRow::GenerateChildrenForPropertyNode( TSharedPtr<FPropertyNo
 		ParentEnabledState.Bind( this, &FDetailPropertyRow::GetEnabledState );
 	}
 
-	if( StructLayoutBuilder.IsValid() && bShowCustomPropertyChildren )
+	if( PropertyTypeLayoutBuilder.IsValid() && bShowCustomPropertyChildren )
 	{
-		const TArray< FDetailLayoutCustomization >& ChildRows = StructLayoutBuilder->GetChildCustomizations();
+		const TArray< FDetailLayoutCustomization >& ChildRows = PropertyTypeLayoutBuilder->GetChildCustomizations();
 
 		for( int32 ChildIndex = 0; ChildIndex < ChildRows.Num(); ++ChildIndex )
 		{
diff --git a/Engine/Source/Editor/PropertyEditor/Private/DetailPropertyRow.h b/Engine/Source/Editor/PropertyEditor/Private/DetailPropertyRow.h
--- a/Engine/Source/Editor/PropertyEditor/Private/DetailPropertyRow.h
+++ b/Engine/Source/Editor/PropertyEditor/Private/DetailPropertyRow.h
@@ -2,7 +2,7 @@
 
 #pragma once
 
-class FDetailPropertyRow : public IDetailPropertyRow, public IStructCustomizationUtils, public TSharedFromThis<FDetailPropertyRow>
+class FDetailPropertyRow : public IDetailPropertyRow, public IPropertyTypeCustomizationUtils, public TSharedFromThis<FDetailPropertyRow>
 {
 public:
 	FDetailPropertyRow( TSharedPtr<FPropertyNode> InPropertyNode, TSharedRef<FDetailCategoryImpl> InParentCategory, TSharedPtr<FObjectPropertyNode> InExternalRootNode = NULL);
@@ -19,7 +19,7 @@ public:
 	virtual void GetDefaultWidgets( TSharedPtr<SWidget>& OutNameWidget, TSharedPtr<SWidget>& OutValueWidget ) OVERRIDE;
 	virtual void GetDefaultWidgets( TSharedPtr<SWidget>& OutNameWidget, TSharedPtr<SWidget>& OutValueWidget, FDetailWidgetRow& Row ) OVERRIDE;
 
-	/** IStructCustomizationUtils interface */
+	/** IPropertyTypeCustomizationUtils interface */
 	virtual TSharedPtr<class FAssetThumbnailPool> GetThumbnailPool() const OVERRIDE;
 	virtual TSharedPtr<class IPropertyUtilities> GetPropertyUtilities() const OVERRIDE;
 
@@ -101,10 +101,10 @@ private:
 	TAttribute<bool> IsParentEnabled;
 	/** Visibility of the property */
 	TAttribute<EVisibility> PropertyVisibility;
-	/** If the property on this row is a customized struct, this is the interface to that customization */
-	TSharedPtr<IStructCustomization> CustomStructInterface;
-	/** Builder for children of a customized struct */
-	TSharedPtr<class FCustomChildrenBuilder> StructLayoutBuilder;
+	/** If the property on this row is a customized property type, this is the interface to that customization */
+	TSharedPtr<IPropertyTypeCustomization> CustomTypeInterface;
+	/** Builder for children of a customized property type */
+	TSharedPtr<class FCustomChildrenBuilder> PropertyTypeLayoutBuilder;
 	/** The property handle for this row */
 	TSharedPtr<IPropertyHandle> PropertyHandle;
 	/** The property node for this row */
diff --git a/Engine/Source/Editor/PropertyEditor/Private/PropertyEditorModule.cpp b/Engine/Source/Editor/PropertyEditor/Private/PropertyEditorModule.cpp
--- a/Engine/Source/Editor/PropertyEditor/Private/PropertyEditorModule.cpp
+++ b/Engine/Source/Editor/PropertyEditor/Private/PropertyEditorModule.cpp
@@ -24,6 +24,53 @@
 
 IMPLEMENT_MODULE( FPropertyEditorModule, PropertyEditor );
 
+
+void FPropertyTypeLayoutCallbackList::Add( const FPropertyTypeLayoutCallback& NewCallback )
+{
+	if( !NewCallback.PropertyTypeIdentifier.IsValid() )
+	{
+		BaseCallback = NewCallback;
+	}
+	else
+	{
+		IdentifierList.Add( NewCallback );
+	}
+}
+
+void FPropertyTypeLayoutCallbackList::Remove( const TSharedPtr<IPropertyTypeIdentifier>& InIdentifier )
+{
+	if( !InIdentifier.IsValid() )
+	{
+		BaseCallback = FPropertyTypeLayoutCallback();
+	}
+	else
+	{
+		IdentifierList.RemoveAllSwap( [&InIdentifier]( FPropertyTypeLayoutCallback& Callback) { return Callback.PropertyTypeIdentifier == InIdentifier; } );
+	}
+}
+
+const FPropertyTypeLayoutCallback& FPropertyTypeLayoutCallbackList::Find( const UProperty& Property )
+{
+	if( IdentifierList.Num() > 0 )
+	{
+		FPropertyTypeLayoutCallback* Callback =
+			IdentifierList.FindByPredicate
+			(
+				[&Property]( const FPropertyTypeLayoutCallback& Callback )
+				{
+					return Callback.PropertyTypeIdentifier->IsPropertyTypeCustomized( Property );
+				}
+			);
+
+		if( Callback )
+		{
+			return *Callback;
+		}
+	}
+
+	return BaseCallback;
+}
+
 void FPropertyEditorModule::StartupModule()
 {
 }
@@ -330,27 +377,73 @@ void FPropertyEditorModule::RegisterCustomPropertyLayout( FName ClassName, FOnGe
 	}
 }
 
-void FPropertyEditorModule::RegisterStructPropertyLayout( FName StructTypeName, FOnGetStructCustomizationInstance StructLayoutDelegate )
+void FPropertyEditorModule::UnregisterCustomPropertyLayout(FName ClassName)
 {
-	if( StructTypeName != NAME_None )
+	if (ClassName != NAME_None)
 	{
-		StructTypeToLayoutMap.Add( StructTypeName, StructLayoutDelegate );
+		ClassNameToDetailLayoutNameMap.Remove(ClassName);
 	}
 }
 
-void FPropertyEditorModule::UnregisterCustomPropertyLayout( FName ClassName )
+void FPropertyEditorModule::RegisterStructPropertyLayout( FName StructTypeName, FOnGetStructCustomizationInstance StructLayoutDelegate )
 {
-	if( ClassName != NAME_None )
+	if (StructTypeName != NAME_None)
 	{
-		ClassNameToDetailLayoutNameMap.Remove( ClassName );
+		FPropertyTypeLayoutCallback Callback;
+		Callback.DeprecatedLayoutDelegate = StructLayoutDelegate;
+
+		FPropertyTypeLayoutCallbackList* LayoutCallbacks = PropertyTypeToLayoutMap.Find(StructTypeName);
+		if (LayoutCallbacks)
+		{
+			LayoutCallbacks->Add(Callback);
+		}
+		else
+		{
+			FPropertyTypeLayoutCallbackList NewLayoutCallbacks;
+			NewLayoutCallbacks.Add(Callback);
+			PropertyTypeToLayoutMap.Add(StructTypeName, NewLayoutCallbacks);
+		}
 	}
 }
 
 void FPropertyEditorModule::UnregisterStructPropertyLayout( FName StructTypeName )
 {
-	if( StructTypeName != NAME_None )
+	UnregisterCustomPropertyTypeLayout( StructTypeName );
+}
+
+
+void FPropertyEditorModule::RegisterCustomPropertyTypeLayout( FName PropertyTypeName, FOnGetPropertyTypeCustomizationInstance PropertyTypeLayoutDelegate, TSharedPtr<IPropertyTypeIdentifier> Identifier )
+{
+	if( PropertyTypeName != NAME_None )
 	{
-		StructTypeToLayoutMap.Remove( StructTypeName );
+		FPropertyTypeLayoutCallback Callback;
+		Callback.PropertyTypeLayoutDelegate = PropertyTypeLayoutDelegate;
+		Callback.PropertyTypeIdentifier = Identifier;
+
+		FPropertyTypeLayoutCallbackList* LayoutCallbacks = PropertyTypeToLayoutMap.Find( PropertyTypeName );
+		if( LayoutCallbacks )
+		{
+			LayoutCallbacks->Add( Callback );
+		}
+		else
+		{
+			FPropertyTypeLayoutCallbackList NewLayoutCallbacks;
+			NewLayoutCallbacks.Add( Callback );
+			PropertyTypeToLayoutMap.Add( PropertyTypeName, NewLayoutCallbacks );
+		}
+	}
+}
+
+void FPropertyEditorModule::UnregisterCustomPropertyTypeLayout( FName PropertyTypeName, TSharedPtr<IPropertyTypeIdentifier> Identifier )
+{
+	if( PropertyTypeName != NAME_None )
+	{
+		FPropertyTypeLayoutCallbackList* LayoutCallbacks = PropertyTypeToLayoutMap.Find( PropertyTypeName );
+
+		if( LayoutCallbacks )
+		{
+			LayoutCallbacks->Remove( Identifier );
+		}
 	}
 }
 
@@ -473,16 +566,48 @@ bool FPropertyEditorModule::IsCustomizedStruct(const UStruct* Struct) const
 {
 	if (Struct && !Struct->IsA<UUserDefinedStruct>())
 	{
-		return StructTypeToLayoutMap.Contains(Struct->GetFName());
+		return PropertyTypeToLayoutMap.Contains(Struct->GetFName());
 	}
 	return false;
 }
 
-FOnGetStructCustomizationInstance FPropertyEditorModule::GetStructCustomizaton(const UStruct* Struct)
+FPropertyTypeLayoutCallback FPropertyEditorModule::GetPropertyTypeCustomization(const UProperty* Property)
 {
-	if (Struct && !Struct->IsA<UUserDefinedStruct>())
+	if( Property )
 	{
-		return StructTypeToLayoutMap.FindRef(Struct->GetFName());
+		const UStructProperty* StructProperty = Cast<UStructProperty>(Property);
+		bool bStructProperty = StructProperty && StructProperty->Struct;
+		const bool bUserDefinedStruct = bStructProperty && StructProperty->Struct->IsA<UUserDefinedStruct>();
+		bStructProperty &= !bUserDefinedStruct;
+
+		const UByteProperty* ByteProperty = Cast<UByteProperty>(Property);
+		const bool bEnumProperty = ByteProperty && ByteProperty->Enum;
+
+		FName PropertyTypeName;
+		if( bStructProperty )
+		{
+			PropertyTypeName = StructProperty->Struct->GetFName();
+		}
+		else if( bEnumProperty )
+		{
+			PropertyTypeName = ByteProperty->Enum->GetFName();
+		}
+		else
+		{
+			PropertyTypeName = Property->GetFName();
+		}
+
+
+		if (PropertyTypeName != NAME_None)
+		{
+			FPropertyTypeLayoutCallbackList* LayoutCallbacks = PropertyTypeToLayoutMap.Find(PropertyTypeName);
+			if (LayoutCallbacks)
+			{
+				const FPropertyTypeLayoutCallback& Callback = LayoutCallbacks->Find(*Property);
+				return Callback;
+			}
+		}
 	}
-	return FOnGetStructCustomizationInstance();
+
+	return FPropertyTypeLayoutCallback();
 }
\ No newline at end of file
diff --git a/Engine/Source/Editor/PropertyEditor/Public/PropertyEditorDelegates.h b/Engine/Source/Editor/PropertyEditor/Public/PropertyEditorDelegates.h
--- a/Engine/Source/Editor/PropertyEditor/Public/PropertyEditorDelegates.h
+++ b/Engine/Source/Editor/PropertyEditor/Public/PropertyEditorDelegates.h
@@ -9,8 +9,8 @@ DECLARE_DELEGATE_RetVal_OneParam( bool, FIsPropertyVisible, const UProperty* con
 /** Delegate called to get a detail layout for a specific object class */
 DECLARE_DELEGATE_RetVal( TSharedRef<class IDetailCustomization>, FOnGetDetailCustomizationInstance );
 
-/** Delegate called to get a struct property layout for a specific struct class */
-DECLARE_DELEGATE_RetVal( TSharedRef<class IStructCustomization>, FOnGetStructCustomizationInstance );
+/** Delegate called to get a property layout for a specific property type */
+DECLARE_DELEGATE_RetVal( TSharedRef<class IPropertyTypeCustomization>, FOnGetPropertyTypeCustomizationInstance );
 
 /** Notification for when a property view changes */
 DECLARE_DELEGATE_TwoParams( FOnObjectArrayChanged, const FString&, const TArray< TWeakObjectPtr< UObject > >& );
@@ -37,4 +37,8 @@ DECLARE_DELEGATE_RetVal(bool, FIsPropertyEditingEnabled );
  * This can be used to safely make changes to data that the details panel is observing instead of during PostEditChange (which is
  * unsafe)
  */
-DECLARE_MULTICAST_DELEGATE_OneParam(FOnFinishedChangingProperties, const FPropertyChangedEvent&);
\ No newline at end of file
+DECLARE_MULTICAST_DELEGATE_OneParam(FOnFinishedChangingProperties, const FPropertyChangedEvent&);
+
+/* Deprecated customization delegate support*/
+DECLARE_DELEGATE_RetVal( TSharedRef<class IStructCustomization>, FOnGetStructCustomizationInstance );
+
diff --git a/Engine/Source/Editor/PropertyEditor/Public/PropertyEditorModule.h b/Engine/Source/Editor/PropertyEditor/Public/PropertyEditorModule.h
--- a/Engine/Source/Editor/PropertyEditor/Public/PropertyEditorModule.h
+++ b/Engine/Source/Editor/PropertyEditor/Public/PropertyEditorModule.h
@@ -48,7 +48,8 @@ namespace FPropertyAccess
 class IDetailCustomization : public TSharedFromThis<IDetailCustomization>
 {
 public:
-  virtual ~IDetailCustomization() {}
+	virtual ~IDetailCustomization() {}
+
 	/** Called when details should be customized */
 	virtual void CustomizeDetails( class IDetailLayoutBuilder& DetailBuilder ) = 0;
 };
@@ -60,12 +61,12 @@ class IPropertyTableCellPresenter;
 
 
 /**
- * Utilities for struct customization
+ * Utilities for property type customization
  */
-class IStructCustomizationUtils
+class IPropertyTypeCustomizationUtils
 {
 public:
-	virtual ~IStructCustomizationUtils(){};
+	virtual ~IPropertyTypeCustomizationUtils(){};
 
 	/**
 	 * @return the font used for properties and details
@@ -88,6 +89,9 @@ public:
 	virtual TSharedPtr<class IPropertyUtilities> GetPropertyUtilities() const { return NULL; }
 };
 
+/** Deprecated IStructCustomizationUtils interface */
+typedef IPropertyTypeCustomizationUtils IStructCustomizationUtils;
+
 /**
  * Builder for adding children to a detail customization
  */
@@ -139,12 +143,46 @@ public:
 };
 
 /**
- * Base class for struct customizations
+ * Base class for property type customizations
  */
-class IStructCustomization : public TSharedFromThis<IStructCustomization>
+class IPropertyTypeCustomization : public TSharedFromThis<IPropertyTypeCustomization>
 {
 public:
-	virtual ~IStructCustomization() {}
+	virtual ~IPropertyTypeCustomization() {}
+
+	/**
+	 * Called when the header of the property (the row in the details panel where the property is shown)
+	 * If nothing is added to the row, the header is not displayed
+	 *
+	 * @param PropertyHandle			Handle to the property being customized
+	 * @param HeaderRow					A row that widgets can be added to
+	 * @param StructCustomizationUtils	Utilities for customization
+	 */
+	virtual void CustomizeHeader( TSharedRef<IPropertyHandle> PropertyHandle, class FDetailWidgetRow& HeaderRow, IPropertyTypeCustomizationUtils& CustomizationUtils ) = 0;
+
+	/**
+	 * Called when the children of the property should be customized or extra rows added
+	 *
+	 * @param PropertyHandle			Handle to the property being customized
+	 * @param StructBuilder				A builder for adding children
+	 * @param StructCustomizationUtils	Utilities for customization
+	 */
+	virtual void CustomizeChildren( TSharedRef<IPropertyHandle> PropertyHandle, class IDetailChildrenBuilder& ChildBuilder, IPropertyTypeCustomizationUtils& CustomizationUtils ) = 0;
+};
+
+/** Deprecated IStructCustomization interface */
+class IStructCustomization : public IPropertyTypeCustomization
+{
+public:
+	virtual void CustomizeHeader( TSharedRef<IPropertyHandle> PropertyHandle, class FDetailWidgetRow& HeaderRow, IPropertyTypeCustomizationUtils& CustomizationUtils )
+	{
+		CustomizeStructHeader( PropertyHandle, HeaderRow, CustomizationUtils );
+	}
+
+	virtual void CustomizeChildren( TSharedRef<IPropertyHandle> PropertyHandle, class IDetailChildrenBuilder& ChildBuilder, IPropertyTypeCustomizationUtils& CustomizationUtils ) 
+	{
+		CustomizeStructChildren( PropertyHandle, ChildBuilder, CustomizationUtils );
+	}
 
 	/**
 	 * Called when the header of the struct (usually where the name of the struct and information about the struct as a whole is added)
@@ -166,6 +204,18 @@ public:
 	virtual void CustomizeStructChildren( TSharedRef<IPropertyHandle> StructPropertyHandle, class IDetailChildrenBuilder& ChildBuilder, IStructCustomizationUtils& StructCustomizationUtils ) = 0;
 };
 
+/** 
+ * Base class for adding an extra data to identify a custom property type
+ */
+class IPropertyTypeIdentifier
+{
+public:
+	virtual ~IPropertyTypeIdentifier() {}
+
+	virtual bool IsPropertyTypeCustomized( const UProperty& InProperty ) const = 0;
+};
+
+
 
 /**
  * Callback executed to query the custom layout of details
@@ -178,9 +228,43 @@ struct FDetailLayoutCallback
 	int32 Order;
 };
 
+struct FPropertyTypeLayoutCallback
+{
+	FOnGetPropertyTypeCustomizationInstance PropertyTypeLayoutDelegate;
+
+	FOnGetStructCustomizationInstance DeprecatedLayoutDelegate;
+
+	TSharedPtr<IPropertyTypeIdentifier> PropertyTypeIdentifier;
+
+	bool IsValid() { return PropertyTypeLayoutDelegate.IsBound() || DeprecatedLayoutDelegate.IsBound(); }
+
+	TSharedRef<IPropertyTypeCustomization> GetCustomizationInstance() const
+	{
+		return PropertyTypeLayoutDelegate.IsBound() ? PropertyTypeLayoutDelegate.Execute() : DeprecatedLayoutDelegate.Execute();
+	}
+};
+
+
+struct FPropertyTypeLayoutCallbackList
+{
+	/** The base callback is a registered callback with a null identifier */
+	FPropertyTypeLayoutCallback BaseCallback;
+
+	/** List of registered callbacks with a non null identifier */
+	TArray< FPropertyTypeLayoutCallback > IdentifierList;
+
+	void Add( const FPropertyTypeLayoutCallback& NewCallback );
+	
+	void Remove( const TSharedPtr<IPropertyTypeIdentifier>& InIdentifier );
+
+	const FPropertyTypeLayoutCallback& Find( const UProperty& Property );
+};
+
 typedef TMap< TWeakObjectPtr<UStruct>, FDetailLayoutCallback > FCustomDetailLayoutMap;
 typedef TMap< FName, FDetailLayoutCallback > FCustomDetailLayoutNameMap;
-typedef TMap< FName, FOnGetStructCustomizationInstance > FCustomStructLayoutMap;
+
+/** This is a multimap as there many be more than one customization per property type */
+typedef TMap< FName, FPropertyTypeLayoutCallbackList > FCustomPropertyTypeLayoutMap;
 
 class FPropertyEditorModule : public IModuleInterface
 {
@@ -225,27 +309,45 @@ public:
 	virtual bool HasUnlockedDetailViews() const;
 
 	/**
-	 * Registers a custom detail layout delegate for a specific class name
+	 * Registers a custom detail layout delegate for a specific class
 	 *
-	 * @param ClassName	The class name the custom detail layout is for
+	 * @param ClassName	The name of the class that the custom detail layout is for
 	 * @param DetailLayoutDelegate	The delegate to call when querying for custom detail layouts for the classes properties
 	 */
 	virtual void RegisterCustomPropertyLayout( FName ClassName, FOnGetDetailCustomizationInstance DetailLayoutDelegate );
 
 	/**
-	 * Registers a custom detail layout delegate for a specific class
+	 * Unregisters a custom detail layout delegate for a specific class name
 	 *
-	 * @param Class	The class the custom detail layout is for
-	 * @param DetailLayoutDelegate	The delegate to call when querying for custom detail layouts for the classes properties
+	 * @param ClassName	The class name with the custom detail layout delegate to remove
 	 */
-	virtual void RegisterStructPropertyLayout( FName StructTypeName, FOnGetStructCustomizationInstance StructLayoutDelegate );
+	virtual void UnregisterCustomPropertyLayout( FName ClassName );
 
 	/**
-	 * Unregisters a custom detail layout delegate for a specific class name
+	 * Registers a property type customization
+	 * A property type is a specific UProperty type, a struct, or enum type
 	 *
-	 * @param ClassName	The class name with the custom detail layout delegate to remove
+	 * @param PropertyTypeName		The name of the property type to customize.  For structs and enums this is the name of the struct class or enum	(not StructProperty or ByteProperty) 
+	 * @param PropertyTypeLayoutDelegate	The delegate to call when querying for a custom layout of the property type
+	 * @param Identifier			An identifier to use to differentiate between two customizations on the same type
 	 */
-	virtual void UnregisterCustomPropertyLayout( FName ClassName );
+	virtual void RegisterCustomPropertyTypeLayout( FName PropertyTypeName, FOnGetPropertyTypeCustomizationInstance PropertyTypeLayoutDelegate, TSharedPtr<IPropertyTypeIdentifier> Identifier = NULL );
+
+	/**
+	 * Unregisters a custom detail layout for a properrty type
+	 *
+	 * @param PropertyTypeName 	The name of the property type that was registered
+	 * @param Identifier 		An identifier to use to differentiate between two customizations on the same type
+	 */
+	virtual void UnregisterCustomPropertyTypeLayout( FName PropertyTypeName, TSharedPtr<IPropertyTypeIdentifier> InIdentifier = NULL );
+
+	/**
+	 * Registers a custom detail layout delegate for a specific class
+	 *
+	 * @param Class	The class the custom detail layout is for
+	 * @param StructLayoutDelegate	The delegate to call when querying for custom detail layouts for the struct properties
+	 */
+	virtual void RegisterStructPropertyLayout( FName StructTypeName, FOnGetStructCustomizationInstance StructLayoutDelegate );
 
 	/**
 	 * Unregisters a custom detail layout delegate for a specific structure
@@ -320,7 +422,7 @@ public:
 	virtual TSharedRef< FAssetEditorToolkit > CreatePropertyEditorToolkit( const EToolkitMode::Type Mode, const TSharedPtr< IToolkitHost >& InitToolkitHost, const TArray< UObject* >& ObjectsToEdit );
 	virtual TSharedRef< FAssetEditorToolkit > CreatePropertyEditorToolkit( const EToolkitMode::Type Mode, const TSharedPtr< IToolkitHost >& InitToolkitHost, const TArray< TWeakObjectPtr< UObject > >& ObjectsToEdit );
 
-	FOnGetStructCustomizationInstance GetStructCustomizaton(const UStruct* Struct);
+	FPropertyTypeLayoutCallback GetPropertyTypeCustomization(const UProperty* InProperty);
 	bool IsCustomizedStruct(const UStruct* Struct) const;
 
 private:
@@ -347,6 +449,6 @@ private:
 	TArray< TWeakPtr<class SSingleProperty> > AllSinglePropertyViews;
 	/** A mapping of class names to detail layout delegates, called when querying for custom detail layouts */
 	FCustomDetailLayoutNameMap ClassNameToDetailLayoutNameMap;
-	/** A mapping of structs to struct layout delegates, called when querying for custom struct layouts */
-	FCustomStructLayoutMap StructTypeToLayoutMap;
+	/** A mapping of property names to property type layout delegates, called when querying for custom property layouts */
+	FCustomPropertyTypeLayoutMap PropertyTypeToLayoutMap;
 };
