Deprecated IStructCustomization and struct property customization registration.  Use IPropertyTypeCustomization and RegisterPropertyTypeCustomization instead

diff --git a/Engine/Plugins/Experimental/Paper2D/Source/Paper2DEditor/Private/Paper2DEditorModule.cpp b/Engine/Plugins/Experimental/Paper2D/Source/Paper2DEditor/Private/Paper2DEditorModule.cpp
--- a/Engine/Plugins/Experimental/Paper2D/Source/Paper2DEditor/Private/Paper2DEditorModule.cpp
+++ b/Engine/Plugins/Experimental/Paper2D/Source/Paper2DEditor/Private/Paper2DEditorModule.cpp
@@ -99,7 +99,7 @@ public:
 			PropertyModule.RegisterCustomPropertyLayout("PaperSprite", FOnGetDetailCustomizationInstance::CreateStatic(&FSpriteDetailsCustomization::MakeInstance));
 
 			//@TODO: Struct registration should happen using ::StaticStruct, not by string!!!
-			//PropertyModule.RegisterStructPropertyLayout( "SpritePolygonCollection", FOnGetStructCustomizationInstance::CreateStatic( &FSpritePolygonCollectionCustomization::MakeInstance ) );
+			//PropertyModule.RegisterCustomPropertyTypeLayout( "SpritePolygonCollection", FOnGetPropertyTypeCustomizationInstance::CreateStatic( &FSpritePolygonCollectionCustomization::MakeInstance ) );
 
 			PropertyModule.NotifyCustomizationModuleChanged();
 		}
diff --git a/Engine/Plugins/Experimental/Paper2D/Source/Paper2DEditor/Private/SpriteEditor/SpritePolygonCollectionCustomization.cpp b/Engine/Plugins/Experimental/Paper2D/Source/Paper2DEditor/Private/SpriteEditor/SpritePolygonCollectionCustomization.cpp
--- a/Engine/Plugins/Experimental/Paper2D/Source/Paper2DEditor/Private/SpriteEditor/SpritePolygonCollectionCustomization.cpp
+++ b/Engine/Plugins/Experimental/Paper2D/Source/Paper2DEditor/Private/SpriteEditor/SpritePolygonCollectionCustomization.cpp
@@ -6,16 +6,16 @@
 //////////////////////////////////////////////////////////////////////////
 // FSpritePolygonCollectionCustomization
 
-TSharedRef<IStructCustomization> FSpritePolygonCollectionCustomization::MakeInstance()
+TSharedRef<IPropertyTypeCustomization> FSpritePolygonCollectionCustomization::MakeInstance()
 {
 	return MakeShareable(new FSpritePolygonCollectionCustomization);
 }
 
-void FSpritePolygonCollectionCustomization::CustomizeStructHeader(TSharedRef<IPropertyHandle> StructPropertyHandle, class FDetailWidgetRow& HeaderRow, IStructCustomizationUtils& StructCustomizationUtils)
+void FSpritePolygonCollectionCustomization::CustomizeHeader(TSharedRef<IPropertyHandle> StructPropertyHandle, class FDetailWidgetRow& HeaderRow, IStructCustomizationUtils& StructCustomizationUtils)
 {
 }
 
-void FSpritePolygonCollectionCustomization::CustomizeStructChildren(TSharedRef<class IPropertyHandle> StructPropertyHandle, class IDetailChildrenBuilder& StructBuilder, IStructCustomizationUtils& StructCustomizationUtils)
+void FSpritePolygonCollectionCustomization::CustomizeChildren(TSharedRef<class IPropertyHandle> StructPropertyHandle, class IDetailChildrenBuilder& StructBuilder, IStructCustomizationUtils& StructCustomizationUtils)
 {
 }
 #if 0
diff --git a/Engine/Plugins/Experimental/Paper2D/Source/Paper2DEditor/Private/SpriteEditor/SpritePolygonCollectionCustomization.h b/Engine/Plugins/Experimental/Paper2D/Source/Paper2DEditor/Private/SpriteEditor/SpritePolygonCollectionCustomization.h
--- a/Engine/Plugins/Experimental/Paper2D/Source/Paper2DEditor/Private/SpriteEditor/SpritePolygonCollectionCustomization.h
+++ b/Engine/Plugins/Experimental/Paper2D/Source/Paper2DEditor/Private/SpriteEditor/SpritePolygonCollectionCustomization.h
@@ -7,16 +7,16 @@
 //////////////////////////////////////////////////////////////////////////
 // FSpritePolygonCollectionCustomization
 
-class FSpritePolygonCollectionCustomization : public IStructCustomization
+class FSpritePolygonCollectionCustomization : public IPropertyTypeCustomization
 {
 public:
 	// Makes a new instance of this customization
-	static TSharedRef<IStructCustomization> MakeInstance();
+	static TSharedRef<IPropertyTypeCustomization> MakeInstance();
 
-	// IStructCustomization interface
-	virtual void CustomizeStructHeader(TSharedRef<IPropertyHandle> StructPropertyHandle, class FDetailWidgetRow& HeaderRow, IStructCustomizationUtils& StructCustomizationUtils) OVERRIDE;
-	virtual void CustomizeStructChildren(TSharedRef<class IPropertyHandle> StructPropertyHandle, class IDetailChildrenBuilder& StructBuilder, IStructCustomizationUtils& StructCustomizationUtils) OVERRIDE;
-	// End of IStructCustomization interface
+	// IPropertyTypeCustomization interface
+	virtual void CustomizeHeader(TSharedRef<IPropertyHandle> StructPropertyHandle, class FDetailWidgetRow& HeaderRow, IStructCustomizationUtils& StructCustomizationUtils) OVERRIDE;
+	virtual void CustomizeChildren(TSharedRef<class IPropertyHandle> StructPropertyHandle, class IDetailChildrenBuilder& StructBuilder, IStructCustomizationUtils& StructCustomizationUtils) OVERRIDE;
+	// End of IPropertyTypeCustomization interface
 
 protected:
 	EVisibility PolygonModeMatches(TSharedPtr<IPropertyHandle> Property, ESpritePolygonMode::Type DesiredMode) const;
diff --git a/Engine/Source/Editor/BehaviorTreeEditor/Private/BehaviorTreeEditorModule.cpp b/Engine/Source/Editor/BehaviorTreeEditor/Private/BehaviorTreeEditorModule.cpp
--- a/Engine/Source/Editor/BehaviorTreeEditor/Private/BehaviorTreeEditorModule.cpp
+++ b/Engine/Source/Editor/BehaviorTreeEditor/Private/BehaviorTreeEditorModule.cpp
@@ -50,7 +50,7 @@ void FBehaviorTreeEditorModule::StartupModule()
 
 	// Register the details customizer
 	FPropertyEditorModule& PropertyModule = FModuleManager::LoadModuleChecked<FPropertyEditorModule>("PropertyEditor");
-	PropertyModule.RegisterStructPropertyLayout( "BlackboardKeySelector", FOnGetStructCustomizationInstance::CreateStatic( &FBlackboardSelectorDetails::MakeInstance ) );
+	PropertyModule.RegisterCustomPropertyTypeLayout( "BlackboardKeySelector", FOnGetPropertyTypeCustomizationInstance::CreateStatic( &FBlackboardSelectorDetails::MakeInstance ) );
 	PropertyModule.RegisterCustomPropertyLayout( "BTDecorator_Blackboard", FOnGetDetailCustomizationInstance::CreateStatic( &FBlackboardDecoratorDetails::MakeInstance ) );
 	PropertyModule.RegisterCustomPropertyLayout( "BTDecorator", FOnGetDetailCustomizationInstance::CreateStatic( &FBehaviorDecoratorDetails::MakeInstance ) );
 	PropertyModule.NotifyCustomizationModuleChanged();
diff --git a/Engine/Source/Editor/BehaviorTreeEditor/Private/DetailCustomizations/BlackboardSelectorDetails.cpp b/Engine/Source/Editor/BehaviorTreeEditor/Private/DetailCustomizations/BlackboardSelectorDetails.cpp
--- a/Engine/Source/Editor/BehaviorTreeEditor/Private/DetailCustomizations/BlackboardSelectorDetails.cpp
+++ b/Engine/Source/Editor/BehaviorTreeEditor/Private/DetailCustomizations/BlackboardSelectorDetails.cpp
@@ -8,13 +8,12 @@
 
 #define LOCTEXT_NAMESPACE "BlackboardSelectorDetails"
 
-TSharedRef<IStructCustomization> FBlackboardSelectorDetails::MakeInstance()
+TSharedRef<IPropertyTypeCustomization> FBlackboardSelectorDetails::MakeInstance()
 {
 	return MakeShareable( new FBlackboardSelectorDetails );
 }
 
-BEGIN_SLATE_FUNCTION_BUILD_OPTIMIZATION
-void FBlackboardSelectorDetails::CustomizeStructHeader( TSharedRef<class IPropertyHandle> StructPropertyHandle, class FDetailWidgetRow& HeaderRow, IStructCustomizationUtils& StructCustomizationUtils )
+void FBlackboardSelectorDetails::CustomizeHeader( TSharedRef<class IPropertyHandle> StructPropertyHandle, class FDetailWidgetRow& HeaderRow, IStructCustomizationUtils& StructCustomizationUtils )
 {
 	MyStructProperty = StructPropertyHandle;
 	PropUtils = StructCustomizationUtils.GetPropertyUtilities().Get();
@@ -42,9 +41,8 @@ void FBlackboardSelectorDetails::CustomizeStructHeader( TSharedRef<class IProper
 
 	InitKeyFromProperty();
 }
-END_SLATE_FUNCTION_BUILD_OPTIMIZATION
 
-void FBlackboardSelectorDetails::CustomizeStructChildren( TSharedRef<class IPropertyHandle> StructPropertyHandle, class IDetailChildrenBuilder& StructBuilder, IStructCustomizationUtils& StructCustomizationUtils )
+void FBlackboardSelectorDetails::CustomizeChildren( TSharedRef<class IPropertyHandle> StructPropertyHandle, class IDetailChildrenBuilder& StructBuilder, IStructCustomizationUtils& StructCustomizationUtils )
 {
 }
 
diff --git a/Engine/Source/Editor/BehaviorTreeEditor/Private/DetailCustomizations/BlackboardSelectorDetails.h b/Engine/Source/Editor/BehaviorTreeEditor/Private/DetailCustomizations/BlackboardSelectorDetails.h
--- a/Engine/Source/Editor/BehaviorTreeEditor/Private/DetailCustomizations/BlackboardSelectorDetails.h
+++ b/Engine/Source/Editor/BehaviorTreeEditor/Private/DetailCustomizations/BlackboardSelectorDetails.h
@@ -5,15 +5,15 @@
 #include "PropertyEditing.h"
 #include "PropertyCustomizationHelpers.h"
 
-class FBlackboardSelectorDetails : public IStructCustomization
+class FBlackboardSelectorDetails : public IPropertyTypeCustomization
 {
 public:
 	/** Makes a new instance of this detail layout class for a specific detail view requesting it */
-	static TSharedRef<IStructCustomization> MakeInstance();
+	static TSharedRef<IPropertyTypeCustomization> MakeInstance();
 
-	/** IStructCustomization interface */
-	virtual void CustomizeStructHeader( TSharedRef<class IPropertyHandle> StructPropertyHandle, class FDetailWidgetRow& HeaderRow, IStructCustomizationUtils& StructCustomizationUtils ) OVERRIDE;
-	virtual void CustomizeStructChildren( TSharedRef<class IPropertyHandle> StructPropertyHandle, class IDetailChildrenBuilder& StructBuilder, IStructCustomizationUtils& StructCustomizationUtils ) OVERRIDE;
+	/** IPropertyTypeCustomization interface */
+	virtual void CustomizeHeader( TSharedRef<class IPropertyHandle> StructPropertyHandle, class FDetailWidgetRow& HeaderRow, IStructCustomizationUtils& StructCustomizationUtils ) OVERRIDE;
+	virtual void CustomizeChildren( TSharedRef<class IPropertyHandle> StructPropertyHandle, class IDetailChildrenBuilder& StructBuilder, IStructCustomizationUtils& StructCustomizationUtils ) OVERRIDE;
 
 private:
 
diff --git a/Engine/Source/Editor/DetailCustomizations/Private/AttenuationSettingsCustomizations.cpp b/Engine/Source/Editor/DetailCustomizations/Private/AttenuationSettingsCustomizations.cpp
--- a/Engine/Source/Editor/DetailCustomizations/Private/AttenuationSettingsCustomizations.cpp
+++ b/Engine/Source/Editor/DetailCustomizations/Private/AttenuationSettingsCustomizations.cpp
@@ -3,12 +3,12 @@
 #include "DetailCustomizationsPrivatePCH.h"
 #include "AttenuationSettingsCustomizations.h"
 
-TSharedRef<IStructCustomization> FAttenuationSettingsCustomization::MakeInstance() 
+TSharedRef<IPropertyTypeCustomization> FAttenuationSettingsCustomization::MakeInstance() 
 {
 	return MakeShareable( new FAttenuationSettingsCustomization );
 }
 
-void FAttenuationSettingsCustomization::CustomizeStructHeader( TSharedRef<class IPropertyHandle> StructPropertyHandle, class FDetailWidgetRow& HeaderRow, IStructCustomizationUtils& StructCustomizationUtils )
+void FAttenuationSettingsCustomization::CustomizeHeader( TSharedRef<class IPropertyHandle> StructPropertyHandle, class FDetailWidgetRow& HeaderRow, IStructCustomizationUtils& StructCustomizationUtils )
 {
 	// We'll set up reset to default ourselves
 	const bool bDisplayResetToDefault = false;
@@ -21,7 +21,7 @@ void FAttenuationSettingsCustomization::CustomizeStructHeader( TSharedRef<class
 		];
 }
 
-void FAttenuationSettingsCustomization::CustomizeStructChildren( TSharedRef<IPropertyHandle> StructPropertyHandle, class IDetailChildrenBuilder& ChildBuilder, IStructCustomizationUtils& StructCustomizationUtils )
+void FAttenuationSettingsCustomization::CustomizeChildren( TSharedRef<IPropertyHandle> StructPropertyHandle, class IDetailChildrenBuilder& ChildBuilder, IStructCustomizationUtils& StructCustomizationUtils )
 {
 	uint32 NumChildren;
 	StructPropertyHandle->GetNumChildren( NumChildren );
diff --git a/Engine/Source/Editor/DetailCustomizations/Private/AttenuationSettingsCustomizations.h b/Engine/Source/Editor/DetailCustomizations/Private/AttenuationSettingsCustomizations.h
--- a/Engine/Source/Editor/DetailCustomizations/Private/AttenuationSettingsCustomizations.h
+++ b/Engine/Source/Editor/DetailCustomizations/Private/AttenuationSettingsCustomizations.h
@@ -2,14 +2,14 @@
 
 #pragma once
 
-class FAttenuationSettingsCustomization : public IStructCustomization
+class FAttenuationSettingsCustomization : public IPropertyTypeCustomization
 {
 public:
-	static TSharedRef<IStructCustomization> MakeInstance();
+	static TSharedRef<IPropertyTypeCustomization> MakeInstance();
 
-	/** IStructCustomization instance */
-	virtual void CustomizeStructHeader( TSharedRef<class IPropertyHandle> StructPropertyHandle, class FDetailWidgetRow& HeaderRow, IStructCustomizationUtils& StructCustomizationUtils ) OVERRIDE;
-	virtual void CustomizeStructChildren( TSharedRef<IPropertyHandle> StructPropertyHandle, class IDetailChildrenBuilder& ChildBuilder, IStructCustomizationUtils& StructCustomizationUtils ) OVERRIDE;
+	/** IPropertyTypeCustomization instance */
+	virtual void CustomizeHeader( TSharedRef<class IPropertyHandle> StructPropertyHandle, class FDetailWidgetRow& HeaderRow, IStructCustomizationUtils& StructCustomizationUtils ) OVERRIDE;
+	virtual void CustomizeChildren( TSharedRef<IPropertyHandle> StructPropertyHandle, class IDetailChildrenBuilder& ChildBuilder, IStructCustomizationUtils& StructCustomizationUtils ) OVERRIDE;
 
 protected:
 
diff --git a/Engine/Source/Editor/DetailCustomizations/Private/BlackboardEntryDetails.cpp b/Engine/Source/Editor/DetailCustomizations/Private/BlackboardEntryDetails.cpp
--- a/Engine/Source/Editor/DetailCustomizations/Private/BlackboardEntryDetails.cpp
+++ b/Engine/Source/Editor/DetailCustomizations/Private/BlackboardEntryDetails.cpp
@@ -6,13 +6,13 @@
 
 #define LOCTEXT_NAMESPACE "BlackboardEntryDetails"
 
-TSharedRef<IStructCustomization> FBlackboardEntryDetails::MakeInstance()
+TSharedRef<IPropertyTypeCustomization> FBlackboardEntryDetails::MakeInstance()
 {
 	return MakeShareable( new FBlackboardEntryDetails );
 }
 
 BEGIN_SLATE_FUNCTION_BUILD_OPTIMIZATION
-void FBlackboardEntryDetails::CustomizeStructHeader( TSharedRef<class IPropertyHandle> StructPropertyHandle, class FDetailWidgetRow& HeaderRow, IStructCustomizationUtils& StructCustomizationUtils )
+void FBlackboardEntryDetails::CustomizeHeader( TSharedRef<class IPropertyHandle> StructPropertyHandle, class FDetailWidgetRow& HeaderRow, IStructCustomizationUtils& StructCustomizationUtils )
 {
 	MyStructProperty = StructPropertyHandle;
 	MyNameProperty = MyStructProperty->GetChildHandle(TEXT("EntryName"));
@@ -32,7 +32,7 @@ void FBlackboardEntryDetails::CustomizeStructHeader( TSharedRef<class IPropertyH
 }
 END_SLATE_FUNCTION_BUILD_OPTIMIZATION
 
-void FBlackboardEntryDetails::CustomizeStructChildren( TSharedRef<class IPropertyHandle> StructPropertyHandle, class IDetailChildrenBuilder& StructBuilder, IStructCustomizationUtils& StructCustomizationUtils )
+void FBlackboardEntryDetails::CustomizeChildren( TSharedRef<class IPropertyHandle> StructPropertyHandle, class IDetailChildrenBuilder& StructBuilder, IStructCustomizationUtils& StructCustomizationUtils )
 {
 	if (StructPropertyHandle->IsValidHandle())
 	{
diff --git a/Engine/Source/Editor/DetailCustomizations/Private/BlackboardEntryDetails.h b/Engine/Source/Editor/DetailCustomizations/Private/BlackboardEntryDetails.h
--- a/Engine/Source/Editor/DetailCustomizations/Private/BlackboardEntryDetails.h
+++ b/Engine/Source/Editor/DetailCustomizations/Private/BlackboardEntryDetails.h
@@ -2,15 +2,15 @@
 
 #pragma once
 
-class FBlackboardEntryDetails : public IStructCustomization
+class FBlackboardEntryDetails : public IPropertyTypeCustomization
 {
 public:
 	/** Makes a new instance of this detail layout class for a specific detail view requesting it */
-	static TSharedRef<IStructCustomization> MakeInstance();
+	static TSharedRef<IPropertyTypeCustomization> MakeInstance();
 
-	/** IStructCustomization interface */
-	virtual void CustomizeStructHeader( TSharedRef<class IPropertyHandle> StructPropertyHandle, class FDetailWidgetRow& HeaderRow, IStructCustomizationUtils& StructCustomizationUtils ) OVERRIDE;
-	virtual void CustomizeStructChildren( TSharedRef<class IPropertyHandle> StructPropertyHandle, class IDetailChildrenBuilder& StructBuilder, IStructCustomizationUtils& StructCustomizationUtils ) OVERRIDE;
+	/** IPropertyTypeCustomization interface */
+	virtual void CustomizeHeader( TSharedRef<class IPropertyHandle> StructPropertyHandle, class FDetailWidgetRow& HeaderRow, IStructCustomizationUtils& StructCustomizationUtils ) OVERRIDE;
+	virtual void CustomizeChildren( TSharedRef<class IPropertyHandle> StructPropertyHandle, class IDetailChildrenBuilder& StructBuilder, IStructCustomizationUtils& StructCustomizationUtils ) OVERRIDE;
 
 private:
 
diff --git a/Engine/Source/Editor/DetailCustomizations/Private/BodyInstanceCustomization.cpp b/Engine/Source/Editor/DetailCustomizations/Private/BodyInstanceCustomization.cpp
--- a/Engine/Source/Editor/DetailCustomizations/Private/BodyInstanceCustomization.cpp
+++ b/Engine/Source/Editor/DetailCustomizations/Private/BodyInstanceCustomization.cpp
@@ -34,7 +34,7 @@ void FBodyInstanceCustomization::RefreshCollisionProfiles()
 }
 
 BEGIN_SLATE_FUNCTION_BUILD_OPTIMIZATION
-void FBodyInstanceCustomization::CustomizeStructChildren( TSharedRef<class IPropertyHandle> StructPropertyHandle, class IDetailChildrenBuilder& StructBuilder, IStructCustomizationUtils& StructCustomizationUtils )
+void FBodyInstanceCustomization::CustomizeChildren( TSharedRef<class IPropertyHandle> StructPropertyHandle, class IDetailChildrenBuilder& StructBuilder, IStructCustomizationUtils& StructCustomizationUtils )
 {
 	CollisionProfileNameHandle = StructPropertyHandle->GetChildHandle(TEXT("CollisionProfileName"));
 	CollisionEnabledHandle = StructPropertyHandle->GetChildHandle(TEXT("CollisionEnabled"));
diff --git a/Engine/Source/Editor/DetailCustomizations/Private/BodyInstanceCustomization.h b/Engine/Source/Editor/DetailCustomizations/Private/BodyInstanceCustomization.h
--- a/Engine/Source/Editor/DetailCustomizations/Private/BodyInstanceCustomization.h
+++ b/Engine/Source/Editor/DetailCustomizations/Private/BodyInstanceCustomization.h
@@ -12,19 +12,19 @@ struct FCollisionChannelInfo
 /**
  * Customizes a DataTable asset to use a dropdown
  */
-class FBodyInstanceCustomization : public IStructCustomization
+class FBodyInstanceCustomization : public IPropertyTypeCustomization
 {
 public:
 	FBodyInstanceCustomization();
 
-	static TSharedRef<IStructCustomization> MakeInstance() 
+	static TSharedRef<IPropertyTypeCustomization> MakeInstance() 
 	{
 		return MakeShareable( new FBodyInstanceCustomization );
 	}
 
-	/** IStructCustomization interface */
-	virtual void CustomizeStructHeader( TSharedRef<IPropertyHandle> StructPropertyHandle, class FDetailWidgetRow& HeaderRow, IStructCustomizationUtils& StructCustomizationUtils ) OVERRIDE {};
-	virtual void CustomizeStructChildren( TSharedRef<class IPropertyHandle> StructPropertyHandle, class IDetailChildrenBuilder& StructBuilder, IStructCustomizationUtils& StructCustomizationUtils ) OVERRIDE;
+	/** IPropertyTypeCustomization interface */
+	virtual void CustomizeHeader( TSharedRef<IPropertyHandle> StructPropertyHandle, class FDetailWidgetRow& HeaderRow, IStructCustomizationUtils& StructCustomizationUtils ) OVERRIDE {};
+	virtual void CustomizeChildren( TSharedRef<class IPropertyHandle> StructPropertyHandle, class IDetailChildrenBuilder& StructBuilder, IStructCustomizationUtils& StructCustomizationUtils ) OVERRIDE;
 
 private:
 	// Profile combo related
diff --git a/Engine/Source/Editor/DetailCustomizations/Private/CurveStructCustomization.cpp b/Engine/Source/Editor/DetailCustomizations/Private/CurveStructCustomization.cpp
--- a/Engine/Source/Editor/DetailCustomizations/Private/CurveStructCustomization.cpp
+++ b/Engine/Source/Editor/DetailCustomizations/Private/CurveStructCustomization.cpp
@@ -11,7 +11,7 @@
 
 const FVector2D FCurveStructCustomization::DEFAULT_WINDOW_SIZE = FVector2D(800, 500);
 
-TSharedRef<IStructCustomization> FCurveStructCustomization::MakeInstance() 
+TSharedRef<IPropertyTypeCustomization> FCurveStructCustomization::MakeInstance() 
 {
 	return MakeShareable( new FCurveStructCustomization );
 }
@@ -29,7 +29,7 @@ FCurveStructCustomization::FCurveStructCustomization()
 {
 }
 
-void FCurveStructCustomization::CustomizeStructHeader( TSharedRef<class IPropertyHandle> InStructPropertyHandle, class FDetailWidgetRow& HeaderRow, IStructCustomizationUtils& StructCustomizationUtils )
+void FCurveStructCustomization::CustomizeHeader( TSharedRef<class IPropertyHandle> InStructPropertyHandle, class FDetailWidgetRow& HeaderRow, IStructCustomizationUtils& StructCustomizationUtils )
 {
 	this->StructPropertyHandle = InStructPropertyHandle;
 
@@ -106,7 +106,7 @@ void FCurveStructCustomization::CustomizeStructHeader( TSharedRef<class IPropert
 	}
 }
 
-void FCurveStructCustomization::CustomizeStructChildren( TSharedRef<class IPropertyHandle> InStructPropertyHandle, class IDetailChildrenBuilder& StructBuilder, IStructCustomizationUtils& StructCustomizationUtils )
+void FCurveStructCustomization::CustomizeChildren( TSharedRef<class IPropertyHandle> InStructPropertyHandle, class IDetailChildrenBuilder& StructBuilder, IStructCustomizationUtils& StructCustomizationUtils )
 {
 	uint32 NumChildren = 0;
 	StructPropertyHandle->GetNumChildren(NumChildren);
diff --git a/Engine/Source/Editor/DetailCustomizations/Private/CurveStructCustomization.h b/Engine/Source/Editor/DetailCustomizations/Private/CurveStructCustomization.h
--- a/Engine/Source/Editor/DetailCustomizations/Private/CurveStructCustomization.h
+++ b/Engine/Source/Editor/DetailCustomizations/Private/CurveStructCustomization.h
@@ -7,19 +7,19 @@
 /**
  * Customizes a RuntimeFloatCurve struct to display a Curve Editor
  */
-class FCurveStructCustomization : public IStructCustomization, public FCurveOwnerInterface
+class FCurveStructCustomization : public IPropertyTypeCustomization, public FCurveOwnerInterface
 {
 public:
-	static TSharedRef<IStructCustomization> MakeInstance();
+	static TSharedRef<IPropertyTypeCustomization> MakeInstance();
 
 	/**
 	 * Destructor
 	 */
 	virtual ~FCurveStructCustomization();
 
-	/** IStructCustomization interface */
-	virtual void CustomizeStructHeader( TSharedRef<class IPropertyHandle> InStructPropertyHandle, class FDetailWidgetRow& HeaderRow, IStructCustomizationUtils& StructCustomizationUtils ) OVERRIDE;
-	virtual void CustomizeStructChildren( TSharedRef<class IPropertyHandle> InStructPropertyHandle, class IDetailChildrenBuilder& StructBuilder, IStructCustomizationUtils& StructCustomizationUtils ) OVERRIDE;
+	/** IPropertyTypeCustomization interface */
+	virtual void CustomizeHeader( TSharedRef<class IPropertyHandle> InStructPropertyHandle, class FDetailWidgetRow& HeaderRow, IStructCustomizationUtils& StructCustomizationUtils ) OVERRIDE;
+	virtual void CustomizeChildren( TSharedRef<class IPropertyHandle> InStructPropertyHandle, class IDetailChildrenBuilder& StructBuilder, IStructCustomizationUtils& StructCustomizationUtils ) OVERRIDE;
 
 	/** FCurveOwnerInterface interface */
 	virtual TArray<FRichCurveEditInfoConst> GetCurves() const OVERRIDE;
diff --git a/Engine/Source/Editor/DetailCustomizations/Private/CurveTableCustomization.h b/Engine/Source/Editor/DetailCustomizations/Private/CurveTableCustomization.h
--- a/Engine/Source/Editor/DetailCustomizations/Private/CurveTableCustomization.h
+++ b/Engine/Source/Editor/DetailCustomizations/Private/CurveTableCustomization.h
@@ -5,16 +5,16 @@
 /**
  * Customizes a DataTable asset to use a dropdown
  */
-class FCurveTableCustomizationLayout : public IStructCustomization
+class FCurveTableCustomizationLayout : public IPropertyTypeCustomization
 {
 public:
-	static TSharedRef<IStructCustomization> MakeInstance() 
+	static TSharedRef<IPropertyTypeCustomization> MakeInstance() 
 	{
 		return MakeShareable( new FCurveTableCustomizationLayout );
 	}
 
-	/** IStructCustomization interface */
-	virtual void CustomizeStructHeader( TSharedRef<class IPropertyHandle> InStructPropertyHandle, class FDetailWidgetRow& HeaderRow, IStructCustomizationUtils& StructCustomizationUtils ) OVERRIDE
+	/** IPropertyTypeCustomization interface */
+	virtual void CustomizeHeader( TSharedRef<class IPropertyHandle> InStructPropertyHandle, class FDetailWidgetRow& HeaderRow, IStructCustomizationUtils& StructCustomizationUtils ) OVERRIDE
 	{
 		this->StructPropertyHandle = InStructPropertyHandle;
 
@@ -25,7 +25,7 @@ public:
 			];
 	}
 
-	virtual void CustomizeStructChildren( TSharedRef<class IPropertyHandle> InStructPropertyHandle, class IDetailChildrenBuilder& StructBuilder, IStructCustomizationUtils& StructCustomizationUtils ) OVERRIDE
+	virtual void CustomizeChildren( TSharedRef<class IPropertyHandle> InStructPropertyHandle, class IDetailChildrenBuilder& StructBuilder, IStructCustomizationUtils& StructCustomizationUtils ) OVERRIDE
 	{
 		/** Get all the existing property handles */
 		CurveTablePropertyHandle = InStructPropertyHandle->GetChildHandle( "CurveTable" );
diff --git a/Engine/Source/Editor/DetailCustomizations/Private/DataTableCategoryCustomization.h b/Engine/Source/Editor/DetailCustomizations/Private/DataTableCategoryCustomization.h
--- a/Engine/Source/Editor/DetailCustomizations/Private/DataTableCategoryCustomization.h
+++ b/Engine/Source/Editor/DetailCustomizations/Private/DataTableCategoryCustomization.h
@@ -7,16 +7,16 @@
 /**
  * Customizes a DataTable asset to use a dropdown
  */
-class FDataTableCategoryCustomizationLayout : public IStructCustomization
+class FDataTableCategoryCustomizationLayout : public IPropertyTypeCustomization
 {
 public:
-	static TSharedRef<IStructCustomization> MakeInstance() 
+	static TSharedRef<IPropertyTypeCustomization> MakeInstance() 
 	{
 		return MakeShareable(new FDataTableCategoryCustomizationLayout);
 	}
 
-	/** IStructCustomization interface */
-	virtual void CustomizeStructHeader(TSharedRef<class IPropertyHandle> InStructPropertyHandle, class FDetailWidgetRow& HeaderRow, IStructCustomizationUtils& StructCustomizationUtils) OVERRIDE
+	/** IPropertyTypeCustomization interface */
+	virtual void CustomizeHeader(TSharedRef<class IPropertyHandle> InStructPropertyHandle, class FDetailWidgetRow& HeaderRow, IStructCustomizationUtils& StructCustomizationUtils) OVERRIDE
 	{
 		this->StructPropertyHandle = InStructPropertyHandle;
 
@@ -27,7 +27,7 @@ public:
 			];
 	}
 
-	virtual void CustomizeStructChildren(TSharedRef<class IPropertyHandle> InStructPropertyHandle, class IDetailChildrenBuilder& StructBuilder, IStructCustomizationUtils& StructCustomizationUtils) OVERRIDE
+	virtual void CustomizeChildren(TSharedRef<class IPropertyHandle> InStructPropertyHandle, class IDetailChildrenBuilder& StructBuilder, IStructCustomizationUtils& StructCustomizationUtils) OVERRIDE
 	{
 		/** Get all the existing property handles */
 		DataTablePropertyHandle		= InStructPropertyHandle->GetChildHandle("DataTable");
diff --git a/Engine/Source/Editor/DetailCustomizations/Private/DataTableCustomization.h b/Engine/Source/Editor/DetailCustomizations/Private/DataTableCustomization.h
--- a/Engine/Source/Editor/DetailCustomizations/Private/DataTableCustomization.h
+++ b/Engine/Source/Editor/DetailCustomizations/Private/DataTableCustomization.h
@@ -7,16 +7,16 @@
 /**
  * Customizes a DataTable asset to use a dropdown
  */
-class FDataTableCustomizationLayout : public IStructCustomization
+class FDataTableCustomizationLayout : public IPropertyTypeCustomization
 {
 public:
-	static TSharedRef<IStructCustomization> MakeInstance() 
+	static TSharedRef<IPropertyTypeCustomization> MakeInstance() 
 	{
 		return MakeShareable( new FDataTableCustomizationLayout );
 	}
 
-	/** IStructCustomization interface */
-	virtual void CustomizeStructHeader( TSharedRef<class IPropertyHandle> InStructPropertyHandle, class FDetailWidgetRow& HeaderRow, IStructCustomizationUtils& StructCustomizationUtils ) OVERRIDE
+	/** IPropertyTypeCustomization interface */
+	virtual void CustomizeHeader( TSharedRef<class IPropertyHandle> InStructPropertyHandle, class FDetailWidgetRow& HeaderRow, IStructCustomizationUtils& StructCustomizationUtils ) OVERRIDE
 	{
 		this->StructPropertyHandle = InStructPropertyHandle;
 
@@ -27,7 +27,7 @@ public:
 			];
 	}
 
-	virtual void CustomizeStructChildren( TSharedRef<class IPropertyHandle> InStructPropertyHandle, class IDetailChildrenBuilder& StructBuilder, IStructCustomizationUtils& StructCustomizationUtils ) OVERRIDE
+	virtual void CustomizeChildren( TSharedRef<class IPropertyHandle> InStructPropertyHandle, class IDetailChildrenBuilder& StructBuilder, IStructCustomizationUtils& StructCustomizationUtils ) OVERRIDE
 	{
 		/** Get all the existing property handles */
 		DataTablePropertyHandle = InStructPropertyHandle->GetChildHandle( "DataTable" );
diff --git a/Engine/Source/Editor/DetailCustomizations/Private/DetailCustomizations.cpp b/Engine/Source/Editor/DetailCustomizations/Private/DetailCustomizations.cpp
--- a/Engine/Source/Editor/DetailCustomizations/Private/DetailCustomizations.cpp
+++ b/Engine/Source/Editor/DetailCustomizations/Private/DetailCustomizations.cpp
@@ -124,47 +124,47 @@ void FDetailCustomizationsModule::ShutdownModule()
 void FDetailCustomizationsModule::RegisterPropertyTypeCustomizations()
 {
 	// Structs
-	RegisterStructPropertyLayout("StringAssetReference", FOnGetStructCustomizationInstance::CreateStatic(&FStringAssetReferenceCustomization::MakeInstance));
-	RegisterStructPropertyLayout("StringClassReference", FOnGetStructCustomizationInstance::CreateStatic(&FStringClassReferenceCustomization::MakeInstance));
-	RegisterStructPropertyLayout("DataTableRowHandle", FOnGetStructCustomizationInstance::CreateStatic(&FDataTableCustomizationLayout::MakeInstance));
-	RegisterStructPropertyLayout("DataTableCategoryHandle", FOnGetStructCustomizationInstance::CreateStatic(&FDataTableCategoryCustomizationLayout::MakeInstance));
-	RegisterStructPropertyLayout("CurveTableRowHandle", FOnGetStructCustomizationInstance::CreateStatic(&FCurveTableCustomizationLayout::MakeInstance));
-	RegisterStructPropertyLayout(NAME_Vector, FOnGetStructCustomizationInstance::CreateStatic(&FMathStructCustomization::MakeInstance));
-	RegisterStructPropertyLayout(NAME_Vector4, FOnGetStructCustomizationInstance::CreateStatic(&FMathStructCustomization::MakeInstance));
-	RegisterStructPropertyLayout(NAME_Vector2D, FOnGetStructCustomizationInstance::CreateStatic(&FMathStructCustomization::MakeInstance));
-	RegisterStructPropertyLayout(NAME_IntPoint, FOnGetStructCustomizationInstance::CreateStatic(&FMathStructCustomization::MakeInstance));
-	RegisterStructPropertyLayout(NAME_Rotator, FOnGetStructCustomizationInstance::CreateStatic(&FRotatorStructCustomization::MakeInstance));
-	RegisterStructPropertyLayout(NAME_LinearColor, FOnGetStructCustomizationInstance::CreateStatic(&FColorStructCustomization::MakeInstance));
-	RegisterStructPropertyLayout(NAME_Color, FOnGetStructCustomizationInstance::CreateStatic(&FColorStructCustomization::MakeInstance));
-	RegisterStructPropertyLayout(NAME_Matrix, FOnGetStructCustomizationInstance::CreateStatic(&FMatrixStructCustomization::MakeInstance));
-	RegisterStructPropertyLayout(NAME_Transform, FOnGetStructCustomizationInstance::CreateStatic(&FTransformStructCustomization::MakeInstance));
-	RegisterStructPropertyLayout("SlateColor", FOnGetStructCustomizationInstance::CreateStatic(&FSlateColorCustomization::MakeInstance));
-	RegisterStructPropertyLayout("AttenuationSettings", FOnGetStructCustomizationInstance::CreateStatic(&FAttenuationSettingsCustomization::MakeInstance));
-	RegisterStructPropertyLayout("DialogueContext", FOnGetStructCustomizationInstance::CreateStatic(&FDialogueContextStructCustomization::MakeInstance));
-	RegisterStructPropertyLayout("DialogueWaveParameter", FOnGetStructCustomizationInstance::CreateStatic(&FDialogueWaveParameterStructCustomization::MakeInstance));
-	RegisterStructPropertyLayout("BodyInstance", FOnGetStructCustomizationInstance::CreateStatic(&FBodyInstanceCustomization::MakeInstance));
-	RegisterStructPropertyLayout("VehicleTransmissionData", FOnGetStructCustomizationInstance::CreateStatic(&FVehicleTransmissionDataCustomization::MakeInstance));
-	RegisterStructPropertyLayout("SlateBrush", FOnGetStructCustomizationInstance::CreateStatic(&FSlateBrushStructCustomization::MakeInstance));
-	RegisterStructPropertyLayout("SlateSound", FOnGetStructCustomizationInstance::CreateStatic(&FSlateSoundStructCustomization::MakeInstance));
-	RegisterStructPropertyLayout("Guid", FOnGetStructCustomizationInstance::CreateStatic(&FGuidStructCustomization::MakeInstance));
-	RegisterStructPropertyLayout("Key", FOnGetStructCustomizationInstance::CreateStatic(&FKeyStructCustomization::MakeInstance));
-	RegisterStructPropertyLayout("FloatRange", FOnGetStructCustomizationInstance::CreateStatic(&FRangeStructCustomization<float>::MakeInstance));
-	RegisterStructPropertyLayout("Int32Range", FOnGetStructCustomizationInstance::CreateStatic(&FRangeStructCustomization<int32>::MakeInstance));
-
-	RegisterStructPropertyLayout("BlackboardEntry", FOnGetStructCustomizationInstance::CreateStatic(&FBlackboardEntryDetails::MakeInstance));
-	RegisterStructPropertyLayout("RuntimeFloatCurve", FOnGetStructCustomizationInstance::CreateStatic(&FCurveStructCustomization::MakeInstance));
-	RegisterStructPropertyLayout("EnvNamedValue", FOnGetStructCustomizationInstance::CreateStatic(&FEnvQueryParamInstanceCustomization::MakeInstance));
-	RegisterStructPropertyLayout("NavigationLink", FOnGetStructCustomizationInstance::CreateStatic(&FNavLinkStructCustomization::MakeInstance));
-	RegisterStructPropertyLayout("NavigationSegmentLink", FOnGetStructCustomizationInstance::CreateStatic(&FNavLinkStructCustomization::MakeInstance));
-
-	RegisterStructPropertyLayout("Margin", FOnGetStructCustomizationInstance::CreateStatic(&FMarginStructCustomization::MakeInstance));
-
-	RegisterStructPropertyLayout("DirectoryPath", FOnGetStructCustomizationInstance::CreateStatic(&FDirectoryPathStructCustomization::MakeInstance));
-	RegisterStructPropertyLayout("FilePath", FOnGetStructCustomizationInstance::CreateStatic(&FFilePathStructCustomization::MakeInstance));
-
-	RegisterStructPropertyLayout("InputAxisConfigEntry", FOnGetStructCustomizationInstance::CreateStatic(&FInputAxisConfigCustomization::MakeInstance));
-	RegisterStructPropertyLayout("InputActionKeyMapping", FOnGetStructCustomizationInstance::CreateStatic(&FInputActionMappingCustomization::MakeInstance));
-	RegisterStructPropertyLayout("InputAxisKeyMapping", FOnGetStructCustomizationInstance::CreateStatic(&FInputAxisMappingCustomization::MakeInstance));
+	RegisterCustomPropertyTypeLayout("StringAssetReference", FOnGetPropertyTypeCustomizationInstance::CreateStatic(&FStringAssetReferenceCustomization::MakeInstance));
+	RegisterCustomPropertyTypeLayout("StringClassReference", FOnGetPropertyTypeCustomizationInstance::CreateStatic(&FStringClassReferenceCustomization::MakeInstance));
+	RegisterCustomPropertyTypeLayout("DataTableRowHandle", FOnGetPropertyTypeCustomizationInstance::CreateStatic(&FDataTableCustomizationLayout::MakeInstance));
+	RegisterCustomPropertyTypeLayout("DataTableCategoryHandle", FOnGetPropertyTypeCustomizationInstance::CreateStatic(&FDataTableCategoryCustomizationLayout::MakeInstance));
+	RegisterCustomPropertyTypeLayout("CurveTableRowHandle", FOnGetPropertyTypeCustomizationInstance::CreateStatic(&FCurveTableCustomizationLayout::MakeInstance));
+	RegisterCustomPropertyTypeLayout(NAME_Vector, FOnGetPropertyTypeCustomizationInstance::CreateStatic(&FMathStructCustomization::MakeInstance));
+	RegisterCustomPropertyTypeLayout(NAME_Vector4, FOnGetPropertyTypeCustomizationInstance::CreateStatic(&FMathStructCustomization::MakeInstance));
+	RegisterCustomPropertyTypeLayout(NAME_Vector2D, FOnGetPropertyTypeCustomizationInstance::CreateStatic(&FMathStructCustomization::MakeInstance));
+	RegisterCustomPropertyTypeLayout(NAME_IntPoint, FOnGetPropertyTypeCustomizationInstance::CreateStatic(&FMathStructCustomization::MakeInstance));
+	RegisterCustomPropertyTypeLayout(NAME_Rotator, FOnGetPropertyTypeCustomizationInstance::CreateStatic(&FRotatorStructCustomization::MakeInstance));
+	RegisterCustomPropertyTypeLayout(NAME_LinearColor, FOnGetPropertyTypeCustomizationInstance::CreateStatic(&FColorStructCustomization::MakeInstance));
+	RegisterCustomPropertyTypeLayout(NAME_Color, FOnGetPropertyTypeCustomizationInstance::CreateStatic(&FColorStructCustomization::MakeInstance));
+	RegisterCustomPropertyTypeLayout(NAME_Matrix, FOnGetPropertyTypeCustomizationInstance::CreateStatic(&FMatrixStructCustomization::MakeInstance));
+	RegisterCustomPropertyTypeLayout(NAME_Transform, FOnGetPropertyTypeCustomizationInstance::CreateStatic(&FTransformStructCustomization::MakeInstance));
+	RegisterCustomPropertyTypeLayout("SlateColor", FOnGetPropertyTypeCustomizationInstance::CreateStatic(&FSlateColorCustomization::MakeInstance));
+	RegisterCustomPropertyTypeLayout("AttenuationSettings", FOnGetPropertyTypeCustomizationInstance::CreateStatic(&FAttenuationSettingsCustomization::MakeInstance));
+	RegisterCustomPropertyTypeLayout("DialogueContext", FOnGetPropertyTypeCustomizationInstance::CreateStatic(&FDialogueContextStructCustomization::MakeInstance));
+	RegisterCustomPropertyTypeLayout("DialogueWaveParameter", FOnGetPropertyTypeCustomizationInstance::CreateStatic(&FDialogueWaveParameterStructCustomization::MakeInstance));
+	RegisterCustomPropertyTypeLayout("BodyInstance", FOnGetPropertyTypeCustomizationInstance::CreateStatic(&FBodyInstanceCustomization::MakeInstance));
+	RegisterCustomPropertyTypeLayout("VehicleTransmissionData", FOnGetPropertyTypeCustomizationInstance::CreateStatic(&FVehicleTransmissionDataCustomization::MakeInstance));
+	RegisterCustomPropertyTypeLayout("SlateBrush", FOnGetPropertyTypeCustomizationInstance::CreateStatic(&FSlateBrushStructCustomization::MakeInstance));
+	RegisterCustomPropertyTypeLayout("SlateSound", FOnGetPropertyTypeCustomizationInstance::CreateStatic(&FSlateSoundStructCustomization::MakeInstance));
+	RegisterCustomPropertyTypeLayout("Guid", FOnGetPropertyTypeCustomizationInstance::CreateStatic(&FGuidStructCustomization::MakeInstance));
+	RegisterCustomPropertyTypeLayout("Key", FOnGetPropertyTypeCustomizationInstance::CreateStatic(&FKeyStructCustomization::MakeInstance));
+	RegisterCustomPropertyTypeLayout("FloatRange", FOnGetPropertyTypeCustomizationInstance::CreateStatic(&FRangeStructCustomization<float>::MakeInstance));
+	RegisterCustomPropertyTypeLayout("Int32Range", FOnGetPropertyTypeCustomizationInstance::CreateStatic(&FRangeStructCustomization<int32>::MakeInstance));
+
+	RegisterCustomPropertyTypeLayout("BlackboardEntry", FOnGetPropertyTypeCustomizationInstance::CreateStatic(&FBlackboardEntryDetails::MakeInstance));
+	RegisterCustomPropertyTypeLayout("RuntimeFloatCurve", FOnGetPropertyTypeCustomizationInstance::CreateStatic(&FCurveStructCustomization::MakeInstance));
+	RegisterCustomPropertyTypeLayout("EnvNamedValue", FOnGetPropertyTypeCustomizationInstance::CreateStatic(&FEnvQueryParamInstanceCustomization::MakeInstance));
+	RegisterCustomPropertyTypeLayout("NavigationLink", FOnGetPropertyTypeCustomizationInstance::CreateStatic(&FNavLinkStructCustomization::MakeInstance));
+	RegisterCustomPropertyTypeLayout("NavigationSegmentLink", FOnGetPropertyTypeCustomizationInstance::CreateStatic(&FNavLinkStructCustomization::MakeInstance));
+
+	RegisterCustomPropertyTypeLayout("Margin", FOnGetPropertyTypeCustomizationInstance::CreateStatic(&FMarginStructCustomization::MakeInstance));
+
+	RegisterCustomPropertyTypeLayout("DirectoryPath", FOnGetPropertyTypeCustomizationInstance::CreateStatic(&FDirectoryPathStructCustomization::MakeInstance));
+	RegisterCustomPropertyTypeLayout("FilePath", FOnGetPropertyTypeCustomizationInstance::CreateStatic(&FFilePathStructCustomization::MakeInstance));
+
+	RegisterCustomPropertyTypeLayout("InputAxisConfigEntry", FOnGetPropertyTypeCustomizationInstance::CreateStatic(&FInputAxisConfigCustomization::MakeInstance));
+	RegisterCustomPropertyTypeLayout("InputActionKeyMapping", FOnGetPropertyTypeCustomizationInstance::CreateStatic(&FInputActionMappingCustomization::MakeInstance));
+	RegisterCustomPropertyTypeLayout("InputAxisKeyMapping", FOnGetPropertyTypeCustomizationInstance::CreateStatic(&FInputAxisMappingCustomization::MakeInstance));
 }
 
 void FDetailCustomizationsModule::RegisterObjectCustomizations()
@@ -265,17 +265,4 @@ void FDetailCustomizationsModule::RegisterCustomPropertyTypeLayout(FName Propert
 	FPropertyEditorModule& PropertyModule = FModuleManager::GetModuleChecked<FPropertyEditorModule>(PropertyEditor);
 
 	PropertyModule.RegisterCustomPropertyTypeLayout(PropertyTypeName, PropertyTypeLayoutDelegate);
-}
-
-void FDetailCustomizationsModule::RegisterStructPropertyLayout(FName StructPropertyName, FOnGetStructCustomizationInstance StructLayoutDelegate)
-{
-	check(StructPropertyName != NAME_None);
-
-	static FName PropertyEditor("PropertyEditor");
-
-	RegisteredPropertyTypes.Add(StructPropertyName);
-
-	FPropertyEditorModule& PropertyModule = FModuleManager::GetModuleChecked<FPropertyEditorModule>(PropertyEditor);
-
-	PropertyModule.RegisterStructPropertyLayout(StructPropertyName, StructLayoutDelegate);
-}
+}
\ No newline at end of file
diff --git a/Engine/Source/Editor/DetailCustomizations/Private/DialogueStructsCustomizations.cpp b/Engine/Source/Editor/DetailCustomizations/Private/DialogueStructsCustomizations.cpp
--- a/Engine/Source/Editor/DetailCustomizations/Private/DialogueStructsCustomizations.cpp
+++ b/Engine/Source/Editor/DetailCustomizations/Private/DialogueStructsCustomizations.cpp
@@ -8,12 +8,12 @@
 
 #define LOCTEXT_NAMESPACE "DialogueWaveDetails"
 
-TSharedRef<IStructCustomization> FDialogueContextStructCustomization::MakeInstance() 
+TSharedRef<IPropertyTypeCustomization> FDialogueContextStructCustomization::MakeInstance() 
 {
 	return MakeShareable( new FDialogueContextStructCustomization );
 }
 
-void FDialogueContextStructCustomization::CustomizeStructHeader( TSharedRef<IPropertyHandle> StructPropertyHandle, class FDetailWidgetRow& HeaderRow, IStructCustomizationUtils& StructCustomizationUtils )
+void FDialogueContextStructCustomization::CustomizeHeader( TSharedRef<IPropertyHandle> StructPropertyHandle, class FDetailWidgetRow& HeaderRow, IStructCustomizationUtils& StructCustomizationUtils )
 {
 	if( StructPropertyHandle->IsValidHandle() )
 	{
@@ -39,7 +39,7 @@ void FDialogueContextStructCustomization::CustomizeStructHeader( TSharedRef<IPro
 	}
 }
 
-void FDialogueContextStructCustomization::CustomizeStructChildren( TSharedRef<IPropertyHandle> StructPropertyHandle, class IDetailChildrenBuilder& ChildBuilder, IStructCustomizationUtils& StructCustomizationUtils )
+void FDialogueContextStructCustomization::CustomizeChildren( TSharedRef<IPropertyHandle> StructPropertyHandle, class IDetailChildrenBuilder& ChildBuilder, IStructCustomizationUtils& StructCustomizationUtils )
 {
 	if( StructPropertyHandle->IsValidHandle() )
 	{
@@ -1004,16 +1004,16 @@ EVisibility SValidatedDialogueContextHeaderWidget::GetTargetsErrorVisibility() c
 	return TargetsErrorText.IsValid() && TargetsErrorText->HasError() ? EVisibility::Visible : EVisibility::Collapsed;
 }
 
-TSharedRef<IStructCustomization> FDialogueWaveParameterStructCustomization::MakeInstance() 
+TSharedRef<IPropertyTypeCustomization> FDialogueWaveParameterStructCustomization::MakeInstance() 
 {
 	return MakeShareable( new FDialogueWaveParameterStructCustomization );
 }
 
-void FDialogueWaveParameterStructCustomization::CustomizeStructHeader( TSharedRef<IPropertyHandle> StructPropertyHandle, class FDetailWidgetRow& HeaderRow, IStructCustomizationUtils& StructCustomizationUtils )
+void FDialogueWaveParameterStructCustomization::CustomizeHeader( TSharedRef<IPropertyHandle> StructPropertyHandle, class FDetailWidgetRow& HeaderRow, IStructCustomizationUtils& StructCustomizationUtils )
 {
 }
 
-void FDialogueWaveParameterStructCustomization::CustomizeStructChildren( TSharedRef<IPropertyHandle> StructPropertyHandle, class IDetailChildrenBuilder& ChildBuilder, IStructCustomizationUtils& StructCustomizationUtils )
+void FDialogueWaveParameterStructCustomization::CustomizeChildren( TSharedRef<IPropertyHandle> StructPropertyHandle, class IDetailChildrenBuilder& ChildBuilder, IStructCustomizationUtils& StructCustomizationUtils )
 {
 	if( StructPropertyHandle->IsValidHandle() )
 	{
diff --git a/Engine/Source/Editor/DetailCustomizations/Private/DialogueStructsCustomizations.h b/Engine/Source/Editor/DetailCustomizations/Private/DialogueStructsCustomizations.h
--- a/Engine/Source/Editor/DetailCustomizations/Private/DialogueStructsCustomizations.h
+++ b/Engine/Source/Editor/DetailCustomizations/Private/DialogueStructsCustomizations.h
@@ -2,26 +2,26 @@
 
 #pragma once
 
-class FDialogueContextStructCustomization : public IStructCustomization
+class FDialogueContextStructCustomization : public IPropertyTypeCustomization
 {
 public:
 	/** @return A new instance of this class */
-	static TSharedRef<IStructCustomization> MakeInstance();
+	static TSharedRef<IPropertyTypeCustomization> MakeInstance();
 
-	/** IStructCustomization interface begin */
-	virtual void CustomizeStructHeader( TSharedRef<IPropertyHandle> StructPropertyHandle, class FDetailWidgetRow& HeaderRow, IStructCustomizationUtils& StructCustomizationUtils ) OVERRIDE;
-	virtual void CustomizeStructChildren( TSharedRef<IPropertyHandle> StructPropertyHandle, class IDetailChildrenBuilder& ChildBuilder, IStructCustomizationUtils& StructCustomizationUtils ) OVERRIDE;
-	/** IStructCustomization interface end */
+	/** IPropertyTypeCustomization interface begin */
+	virtual void CustomizeHeader( TSharedRef<IPropertyHandle> StructPropertyHandle, class FDetailWidgetRow& HeaderRow, IStructCustomizationUtils& StructCustomizationUtils ) OVERRIDE;
+	virtual void CustomizeChildren( TSharedRef<IPropertyHandle> StructPropertyHandle, class IDetailChildrenBuilder& ChildBuilder, IStructCustomizationUtils& StructCustomizationUtils ) OVERRIDE;
+	/** IPropertyTypeCustomization interface end */
 };
 
-class FDialogueWaveParameterStructCustomization : public IStructCustomization
+class FDialogueWaveParameterStructCustomization : public IPropertyTypeCustomization
 {
 public:
 	/** @return A new instance of this class */
-	static TSharedRef<IStructCustomization> MakeInstance();
+	static TSharedRef<IPropertyTypeCustomization> MakeInstance();
 
-	/** IStructCustomization interface begin */
-	virtual void CustomizeStructHeader( TSharedRef<IPropertyHandle> StructPropertyHandle, class FDetailWidgetRow& HeaderRow, IStructCustomizationUtils& StructCustomizationUtils ) OVERRIDE;
-	virtual void CustomizeStructChildren( TSharedRef<IPropertyHandle> StructPropertyHandle, class IDetailChildrenBuilder& ChildBuilder, IStructCustomizationUtils& StructCustomizationUtils ) OVERRIDE;
-	/** IStructCustomization interface end */
+	/** IPropertyTypeCustomization interface begin */
+	virtual void CustomizeHeader( TSharedRef<IPropertyHandle> StructPropertyHandle, class FDetailWidgetRow& HeaderRow, IStructCustomizationUtils& StructCustomizationUtils ) OVERRIDE;
+	virtual void CustomizeChildren( TSharedRef<IPropertyHandle> StructPropertyHandle, class IDetailChildrenBuilder& ChildBuilder, IStructCustomizationUtils& StructCustomizationUtils ) OVERRIDE;
+	/** IPropertyTypeCustomization interface end */
 };
\ No newline at end of file
diff --git a/Engine/Source/Editor/DetailCustomizations/Private/DirectoryPathStructCustomization.cpp b/Engine/Source/Editor/DetailCustomizations/Private/DirectoryPathStructCustomization.cpp
--- a/Engine/Source/Editor/DetailCustomizations/Private/DirectoryPathStructCustomization.cpp
+++ b/Engine/Source/Editor/DetailCustomizations/Private/DirectoryPathStructCustomization.cpp
@@ -7,12 +7,12 @@
 
 #define LOCTEXT_NAMESPACE "DirectoryPathStructCustomization"
 
-TSharedRef<IStructCustomization> FDirectoryPathStructCustomization::MakeInstance()
+TSharedRef<IPropertyTypeCustomization> FDirectoryPathStructCustomization::MakeInstance()
 {
 	return MakeShareable(new FDirectoryPathStructCustomization());
 }
 
-void FDirectoryPathStructCustomization::CustomizeStructHeader( TSharedRef<IPropertyHandle> StructPropertyHandle, class FDetailWidgetRow& HeaderRow, IStructCustomizationUtils& StructCustomizationUtils )
+void FDirectoryPathStructCustomization::CustomizeHeader( TSharedRef<IPropertyHandle> StructPropertyHandle, class FDetailWidgetRow& HeaderRow, IStructCustomizationUtils& StructCustomizationUtils )
 {
 	TSharedPtr<IPropertyHandle> PathProperty = StructPropertyHandle->GetChildHandle("Path");
 
@@ -60,7 +60,7 @@ void FDirectoryPathStructCustomization::CustomizeStructHeader( TSharedRef<IPrope
 	}
 }
 
-void FDirectoryPathStructCustomization::CustomizeStructChildren( TSharedRef<IPropertyHandle> StructPropertyHandle, class IDetailChildrenBuilder& StructBuilder, IStructCustomizationUtils& StructCustomizationUtils )
+void FDirectoryPathStructCustomization::CustomizeChildren( TSharedRef<IPropertyHandle> StructPropertyHandle, class IDetailChildrenBuilder& StructBuilder, IStructCustomizationUtils& StructCustomizationUtils )
 {
 }
 
diff --git a/Engine/Source/Editor/DetailCustomizations/Private/DirectoryPathStructCustomization.h b/Engine/Source/Editor/DetailCustomizations/Private/DirectoryPathStructCustomization.h
--- a/Engine/Source/Editor/DetailCustomizations/Private/DirectoryPathStructCustomization.h
+++ b/Engine/Source/Editor/DetailCustomizations/Private/DirectoryPathStructCustomization.h
@@ -2,15 +2,15 @@
 
 #pragma once
 
-class FDirectoryPathStructCustomization : public IStructCustomization
+class FDirectoryPathStructCustomization : public IPropertyTypeCustomization
 {
 public:
-	static TSharedRef<IStructCustomization> MakeInstance();
+	static TSharedRef<IPropertyTypeCustomization> MakeInstance();
 
-	/** IStructCustomization interface */
-	virtual void CustomizeStructHeader( TSharedRef<class IPropertyHandle> StructPropertyHandle, class FDetailWidgetRow& HeaderRow, IStructCustomizationUtils& StructCustomizationUtils ) OVERRIDE;
+	/** IPropertyTypeCustomization interface */
+	virtual void CustomizeHeader( TSharedRef<class IPropertyHandle> StructPropertyHandle, class FDetailWidgetRow& HeaderRow, IStructCustomizationUtils& StructCustomizationUtils ) OVERRIDE;
 
-	virtual void CustomizeStructChildren( TSharedRef<class IPropertyHandle> StructPropertyHandle, class IDetailChildrenBuilder& StructBuilder, IStructCustomizationUtils& StructCustomizationUtils ) OVERRIDE;              
+	virtual void CustomizeChildren( TSharedRef<class IPropertyHandle> StructPropertyHandle, class IDetailChildrenBuilder& StructBuilder, IStructCustomizationUtils& StructCustomizationUtils ) OVERRIDE;              
 
 private:
 
diff --git a/Engine/Source/Editor/DetailCustomizations/Private/EnvQueryParamInstanceCustomization.cpp b/Engine/Source/Editor/DetailCustomizations/Private/EnvQueryParamInstanceCustomization.cpp
--- a/Engine/Source/Editor/DetailCustomizations/Private/EnvQueryParamInstanceCustomization.cpp
+++ b/Engine/Source/Editor/DetailCustomizations/Private/EnvQueryParamInstanceCustomization.cpp
@@ -5,12 +5,12 @@
 
 #define LOCTEXT_NAMESPACE "FEnvQueryCustomization"
 
-TSharedRef<IStructCustomization> FEnvQueryParamInstanceCustomization::MakeInstance( )
+TSharedRef<IPropertyTypeCustomization> FEnvQueryParamInstanceCustomization::MakeInstance( )
 {
 	return MakeShareable(new FEnvQueryParamInstanceCustomization);
 }
 
-void FEnvQueryParamInstanceCustomization::CustomizeStructHeader( TSharedRef<class IPropertyHandle> StructPropertyHandle, class FDetailWidgetRow& HeaderRow, IStructCustomizationUtils& StructCustomizationUtils )
+void FEnvQueryParamInstanceCustomization::CustomizeHeader( TSharedRef<class IPropertyHandle> StructPropertyHandle, class FDetailWidgetRow& HeaderRow, IStructCustomizationUtils& StructCustomizationUtils )
 {
 	NameProp = StructPropertyHandle->GetChildHandle(GET_MEMBER_NAME_CHECKED(FEnvNamedValue,ParamName));
 	TypeProp = StructPropertyHandle->GetChildHandle(GET_MEMBER_NAME_CHECKED(FEnvNamedValue,ParamType));
@@ -34,7 +34,7 @@ void FEnvQueryParamInstanceCustomization::CustomizeStructHeader( TSharedRef<clas
 	];
 }
 
-void FEnvQueryParamInstanceCustomization::CustomizeStructChildren( TSharedRef<class IPropertyHandle> StructPropertyHandle, class IDetailChildrenBuilder& StructBuilder, IStructCustomizationUtils& StructCustomizationUtils )
+void FEnvQueryParamInstanceCustomization::CustomizeChildren( TSharedRef<class IPropertyHandle> StructPropertyHandle, class IDetailChildrenBuilder& StructBuilder, IStructCustomizationUtils& StructCustomizationUtils )
 {
 	StructBuilder.AddChildProperty(NameProp.ToSharedRef());
 	StructBuilder.AddChildProperty(TypeProp.ToSharedRef());
diff --git a/Engine/Source/Editor/DetailCustomizations/Private/EnvQueryParamInstanceCustomization.h b/Engine/Source/Editor/DetailCustomizations/Private/EnvQueryParamInstanceCustomization.h
--- a/Engine/Source/Editor/DetailCustomizations/Private/EnvQueryParamInstanceCustomization.h
+++ b/Engine/Source/Editor/DetailCustomizations/Private/EnvQueryParamInstanceCustomization.h
@@ -4,16 +4,16 @@
 
 #include "EnvironmentQuery/EnvQueryTypes.h"
 
-class FEnvQueryParamInstanceCustomization : public IStructCustomization
+class FEnvQueryParamInstanceCustomization : public IPropertyTypeCustomization
 {
 public:
 
-	// Begin IStructCustomization interface
-	virtual void CustomizeStructHeader( TSharedRef<class IPropertyHandle> StructPropertyHandle, class FDetailWidgetRow& HeaderRow, IStructCustomizationUtils& StructCustomizationUtils ) OVERRIDE;
-	virtual void CustomizeStructChildren( TSharedRef<class IPropertyHandle> StructPropertyHandle, class IDetailChildrenBuilder& StructBuilder, IStructCustomizationUtils& StructCustomizationUtils ) OVERRIDE;
-	// End IStructCustomization interface
+	// Begin IPropertyTypeCustomization interface
+	virtual void CustomizeHeader( TSharedRef<class IPropertyHandle> StructPropertyHandle, class FDetailWidgetRow& HeaderRow, IStructCustomizationUtils& StructCustomizationUtils ) OVERRIDE;
+	virtual void CustomizeChildren( TSharedRef<class IPropertyHandle> StructPropertyHandle, class IDetailChildrenBuilder& StructBuilder, IStructCustomizationUtils& StructCustomizationUtils ) OVERRIDE;
+	// End IPropertyTypeCustomization interface
 
-	static TSharedRef<IStructCustomization> MakeInstance( );
+	static TSharedRef<IPropertyTypeCustomization> MakeInstance( );
 
 protected:
 
diff --git a/Engine/Source/Editor/DetailCustomizations/Private/FilePathStructCustomization.cpp b/Engine/Source/Editor/DetailCustomizations/Private/FilePathStructCustomization.cpp
--- a/Engine/Source/Editor/DetailCustomizations/Private/FilePathStructCustomization.cpp
+++ b/Engine/Source/Editor/DetailCustomizations/Private/FilePathStructCustomization.cpp
@@ -100,12 +100,12 @@ private:
 	FString FileFilterExtension;
 };
 
-TSharedRef<IStructCustomization> FFilePathStructCustomization::MakeInstance()
+TSharedRef<IPropertyTypeCustomization> FFilePathStructCustomization::MakeInstance()
 {
 	return MakeShareable(new FFilePathStructCustomization());
 }
 
-void FFilePathStructCustomization::CustomizeStructHeader( TSharedRef<IPropertyHandle> StructPropertyHandle, class FDetailWidgetRow& HeaderRow, IStructCustomizationUtils& StructCustomizationUtils )
+void FFilePathStructCustomization::CustomizeHeader( TSharedRef<IPropertyHandle> StructPropertyHandle, class FDetailWidgetRow& HeaderRow, IStructCustomizationUtils& StructCustomizationUtils )
 {
 	const UProperty* Property = StructPropertyHandle->GetProperty();
 	checkSlow(Property);
@@ -140,7 +140,7 @@ TSharedPtr<SWidget> FFilePathStructCustomization::CreatePickerWidget(TSharedRef<
 	return TSharedPtr<SWidget>();
 }
 
-void FFilePathStructCustomization::CustomizeStructChildren( TSharedRef<IPropertyHandle> StructPropertyHandle, class IDetailChildrenBuilder& StructBuilder, IStructCustomizationUtils& StructCustomizationUtils )
+void FFilePathStructCustomization::CustomizeChildren( TSharedRef<IPropertyHandle> StructPropertyHandle, class IDetailChildrenBuilder& StructBuilder, IStructCustomizationUtils& StructCustomizationUtils )
 {
 }
 
diff --git a/Engine/Source/Editor/DetailCustomizations/Private/FilePathStructCustomization.h b/Engine/Source/Editor/DetailCustomizations/Private/FilePathStructCustomization.h
--- a/Engine/Source/Editor/DetailCustomizations/Private/FilePathStructCustomization.h
+++ b/Engine/Source/Editor/DetailCustomizations/Private/FilePathStructCustomization.h
@@ -5,15 +5,15 @@
 /** Delegate fired when a path is picked */
 DECLARE_DELEGATE_RetVal_OneParam(bool, FOnPathPicked, FString& /* InOutPath */);
 
-class FFilePathStructCustomization : public IStructCustomization
+class FFilePathStructCustomization : public IPropertyTypeCustomization
 {
 public:
-	static TSharedRef<IStructCustomization> MakeInstance();
+	static TSharedRef<IPropertyTypeCustomization> MakeInstance();
 
-	/** IStructCustomization interface */
-	virtual void CustomizeStructHeader( TSharedRef<class IPropertyHandle> StructPropertyHandle, class FDetailWidgetRow& HeaderRow, IStructCustomizationUtils& StructCustomizationUtils ) OVERRIDE;
+	/** IPropertyTypeCustomization interface */
+	virtual void CustomizeHeader( TSharedRef<class IPropertyHandle> StructPropertyHandle, class FDetailWidgetRow& HeaderRow, IStructCustomizationUtils& StructCustomizationUtils ) OVERRIDE;
 
-	virtual void CustomizeStructChildren( TSharedRef<class IPropertyHandle> StructPropertyHandle, class IDetailChildrenBuilder& StructBuilder, IStructCustomizationUtils& StructCustomizationUtils ) OVERRIDE;              
+	virtual void CustomizeChildren( TSharedRef<class IPropertyHandle> StructPropertyHandle, class IDetailChildrenBuilder& StructBuilder, IStructCustomizationUtils& StructCustomizationUtils ) OVERRIDE;              
 
 	static TSharedPtr<SWidget> CreatePickerWidget(TSharedRef<class IPropertyHandle> StructPropertyHandle,  const FString& FileFilterExtension = FString(TEXT("*")), const FOnPathPicked& OnPathPicked = FOnPathPicked());
 
diff --git a/Engine/Source/Editor/DetailCustomizations/Private/GuidStructCustomization.cpp b/Engine/Source/Editor/DetailCustomizations/Private/GuidStructCustomization.cpp
--- a/Engine/Source/Editor/DetailCustomizations/Private/GuidStructCustomization.cpp
+++ b/Engine/Source/Editor/DetailCustomizations/Private/GuidStructCustomization.cpp
@@ -14,16 +14,16 @@
 /* FGuidStructCustomization static interface
  *****************************************************************************/
 
-TSharedRef<IStructCustomization> FGuidStructCustomization::MakeInstance( )
+TSharedRef<IPropertyTypeCustomization> FGuidStructCustomization::MakeInstance( )
 {
 	return MakeShareable(new FGuidStructCustomization);
 }
 
 
-/* IStructCustomization interface
+/* IPropertyTypeCustomization interface
  *****************************************************************************/
 
-void FGuidStructCustomization::CustomizeStructHeader( TSharedRef<class IPropertyHandle> StructPropertyHandle, class FDetailWidgetRow& HeaderRow, IStructCustomizationUtils& StructCustomizationUtils )
+void FGuidStructCustomization::CustomizeHeader( TSharedRef<class IPropertyHandle> StructPropertyHandle, class FDetailWidgetRow& HeaderRow, IStructCustomizationUtils& StructCustomizationUtils )
 {
 	PropertyHandle = StructPropertyHandle;
 	InputValid = true;
@@ -90,7 +90,7 @@ void FGuidStructCustomization::CustomizeStructHeader( TSharedRef<class IProperty
 }
 
 
-void FGuidStructCustomization::CustomizeStructChildren( TSharedRef<class IPropertyHandle> StructPropertyHandle, class IDetailChildrenBuilder& StructBuilder, IStructCustomizationUtils& StructCustomizationUtils )
+void FGuidStructCustomization::CustomizeChildren( TSharedRef<class IPropertyHandle> StructPropertyHandle, class IDetailChildrenBuilder& StructBuilder, IStructCustomizationUtils& StructCustomizationUtils )
 {
 	// do nothing
 }
diff --git a/Engine/Source/Editor/DetailCustomizations/Private/GuidStructCustomization.h b/Engine/Source/Editor/DetailCustomizations/Private/GuidStructCustomization.h
--- a/Engine/Source/Editor/DetailCustomizations/Private/GuidStructCustomization.h
+++ b/Engine/Source/Editor/DetailCustomizations/Private/GuidStructCustomization.h
@@ -27,17 +27,17 @@ namespace EPropertyEditorGuidActions
  * Implements a details panel customization for  FGuid structures.
  */
 class FGuidStructCustomization
-	: public IStructCustomization
+	: public IPropertyTypeCustomization
 {
 public:
 
-	// Begin IStructCustomization interface
+	// Begin IPropertyTypeCustomization interface
 
-	virtual void CustomizeStructHeader( TSharedRef<class IPropertyHandle> StructPropertyHandle, class FDetailWidgetRow& HeaderRow, IStructCustomizationUtils& StructCustomizationUtils ) OVERRIDE;
+	virtual void CustomizeHeader( TSharedRef<class IPropertyHandle> StructPropertyHandle, class FDetailWidgetRow& HeaderRow, IStructCustomizationUtils& StructCustomizationUtils ) OVERRIDE;
 
-	virtual void CustomizeStructChildren( TSharedRef<class IPropertyHandle> StructPropertyHandle, class IDetailChildrenBuilder& StructBuilder, IStructCustomizationUtils& StructCustomizationUtils ) OVERRIDE;
+	virtual void CustomizeChildren( TSharedRef<class IPropertyHandle> StructPropertyHandle, class IDetailChildrenBuilder& StructBuilder, IStructCustomizationUtils& StructCustomizationUtils ) OVERRIDE;
 
-	// End IStructCustomization interface
+	// End IPropertyTypeCustomization interface
 
 public:
 
@@ -46,7 +46,7 @@ public:
 	 *
 	 * @return A new struct customization for Guids.
 	 */
-	static TSharedRef<IStructCustomization> MakeInstance( );
+	static TSharedRef<IPropertyTypeCustomization> MakeInstance( );
 
 protected:
 
diff --git a/Engine/Source/Editor/DetailCustomizations/Private/InputStructCustomization.cpp b/Engine/Source/Editor/DetailCustomizations/Private/InputStructCustomization.cpp
--- a/Engine/Source/Editor/DetailCustomizations/Private/InputStructCustomization.cpp
+++ b/Engine/Source/Editor/DetailCustomizations/Private/InputStructCustomization.cpp
@@ -10,12 +10,12 @@
 // FInputAxisConfigCustomization //
 ///////////////////////////////////
 
-TSharedRef<IStructCustomization> FInputAxisConfigCustomization::MakeInstance() 
+TSharedRef<IPropertyTypeCustomization> FInputAxisConfigCustomization::MakeInstance() 
 {
 	return MakeShareable( new FInputAxisConfigCustomization );
 }
 
-void FInputAxisConfigCustomization::CustomizeStructHeader( TSharedRef<class IPropertyHandle> InStructPropertyHandle, class FDetailWidgetRow& HeaderRow, IStructCustomizationUtils& StructCustomizationUtils )
+void FInputAxisConfigCustomization::CustomizeHeader( TSharedRef<class IPropertyHandle> InStructPropertyHandle, class FDetailWidgetRow& HeaderRow, IStructCustomizationUtils& StructCustomizationUtils )
 {
 	FString AxisKeyName;
 	InStructPropertyHandle->GetChildHandle(GET_MEMBER_NAME_CHECKED(FInputAxisConfigEntry, AxisKeyName))->GetValue(AxisKeyName);
@@ -26,7 +26,7 @@ void FInputAxisConfigCustomization::CustomizeStructHeader( TSharedRef<class IPro
 	];
 }
 
-void FInputAxisConfigCustomization::CustomizeStructChildren( TSharedRef<class IPropertyHandle> InStructPropertyHandle, class IDetailChildrenBuilder& StructBuilder, IStructCustomizationUtils& StructCustomizationUtils )
+void FInputAxisConfigCustomization::CustomizeChildren( TSharedRef<class IPropertyHandle> InStructPropertyHandle, class IDetailChildrenBuilder& StructBuilder, IStructCustomizationUtils& StructCustomizationUtils )
 {
 	TSharedPtr<IPropertyHandle> AxisProperties = InStructPropertyHandle->GetChildHandle(GET_MEMBER_NAME_CHECKED(FInputAxisConfigEntry, AxisProperties));
 
@@ -43,17 +43,17 @@ void FInputAxisConfigCustomization::CustomizeStructChildren( TSharedRef<class IP
 // FInputActionMappingCustomization //
 //////////////////////////////////////
 
-TSharedRef<IStructCustomization> FInputActionMappingCustomization::MakeInstance() 
+TSharedRef<IPropertyTypeCustomization> FInputActionMappingCustomization::MakeInstance() 
 {
 	return MakeShareable( new FInputActionMappingCustomization );
 }
 
-void FInputActionMappingCustomization::CustomizeStructHeader( TSharedRef<class IPropertyHandle> InStructPropertyHandle, class FDetailWidgetRow& HeaderRow, IStructCustomizationUtils& StructCustomizationUtils )
+void FInputActionMappingCustomization::CustomizeHeader( TSharedRef<class IPropertyHandle> InStructPropertyHandle, class FDetailWidgetRow& HeaderRow, IStructCustomizationUtils& StructCustomizationUtils )
 {
 	ActionMappingHandle = InStructPropertyHandle;
 }
 
-void FInputActionMappingCustomization::CustomizeStructChildren( TSharedRef<class IPropertyHandle> InStructPropertyHandle, class IDetailChildrenBuilder& StructBuilder, IStructCustomizationUtils& StructCustomizationUtils )
+void FInputActionMappingCustomization::CustomizeChildren( TSharedRef<class IPropertyHandle> InStructPropertyHandle, class IDetailChildrenBuilder& StructBuilder, IStructCustomizationUtils& StructCustomizationUtils )
 {
 	TSharedPtr<IPropertyHandle> KeyHandle = InStructPropertyHandle->GetChildHandle(GET_MEMBER_NAME_CHECKED(FInputActionKeyMapping, Key));
 	TSharedPtr<IPropertyHandle> ShiftHandle = InStructPropertyHandle->GetChildHandle(GET_MEMBER_NAME_CHECKED(FInputActionKeyMapping, bShift));
@@ -150,17 +150,17 @@ void FInputActionMappingCustomization::RemoveActionMappingButton_OnClick()
 // FInputAxisMappingCustomization //
 //////////////////////////////////////
 
-TSharedRef<IStructCustomization> FInputAxisMappingCustomization::MakeInstance() 
+TSharedRef<IPropertyTypeCustomization> FInputAxisMappingCustomization::MakeInstance() 
 {
 	return MakeShareable( new FInputAxisMappingCustomization );
 }
 
-void FInputAxisMappingCustomization::CustomizeStructHeader( TSharedRef<class IPropertyHandle> InStructPropertyHandle, class FDetailWidgetRow& HeaderRow, IStructCustomizationUtils& StructCustomizationUtils )
+void FInputAxisMappingCustomization::CustomizeHeader( TSharedRef<class IPropertyHandle> InStructPropertyHandle, class FDetailWidgetRow& HeaderRow, IStructCustomizationUtils& StructCustomizationUtils )
 {
 	AxisMappingHandle = InStructPropertyHandle;
 }
 
-void FInputAxisMappingCustomization::CustomizeStructChildren( TSharedRef<class IPropertyHandle> InStructPropertyHandle, class IDetailChildrenBuilder& StructBuilder, IStructCustomizationUtils& StructCustomizationUtils )
+void FInputAxisMappingCustomization::CustomizeChildren( TSharedRef<class IPropertyHandle> InStructPropertyHandle, class IDetailChildrenBuilder& StructBuilder, IStructCustomizationUtils& StructCustomizationUtils )
 {
 	TSharedPtr<IPropertyHandle> KeyHandle = InStructPropertyHandle->GetChildHandle(GET_MEMBER_NAME_CHECKED(FInputAxisKeyMapping, Key));
 	TSharedPtr<IPropertyHandle> ScaleHandle = InStructPropertyHandle->GetChildHandle(GET_MEMBER_NAME_CHECKED(FInputAxisKeyMapping, Scale));
diff --git a/Engine/Source/Editor/DetailCustomizations/Private/InputStructCustomization.h b/Engine/Source/Editor/DetailCustomizations/Private/InputStructCustomization.h
--- a/Engine/Source/Editor/DetailCustomizations/Private/InputStructCustomization.h
+++ b/Engine/Source/Editor/DetailCustomizations/Private/InputStructCustomization.h
@@ -5,27 +5,27 @@
 /**
  * Customizes an InputAxisConfigEntry struct to display it more simply
  */
-class FInputAxisConfigCustomization : public IStructCustomization
+class FInputAxisConfigCustomization : public IPropertyTypeCustomization
 {
 public:
-	static TSharedRef<IStructCustomization> MakeInstance();
+	static TSharedRef<IPropertyTypeCustomization> MakeInstance();
 
-	/** IStructCustomization interface */
-	virtual void CustomizeStructHeader( TSharedRef<class IPropertyHandle> InStructPropertyHandle, class FDetailWidgetRow& HeaderRow, IStructCustomizationUtils& StructCustomizationUtils ) OVERRIDE;
-	virtual void CustomizeStructChildren( TSharedRef<class IPropertyHandle> InStructPropertyHandle, class IDetailChildrenBuilder& StructBuilder, IStructCustomizationUtils& StructCustomizationUtils ) OVERRIDE;
+	/** IPropertyTypeCustomization interface */
+	virtual void CustomizeHeader( TSharedRef<class IPropertyHandle> InStructPropertyHandle, class FDetailWidgetRow& HeaderRow, IStructCustomizationUtils& StructCustomizationUtils ) OVERRIDE;
+	virtual void CustomizeChildren( TSharedRef<class IPropertyHandle> InStructPropertyHandle, class IDetailChildrenBuilder& StructBuilder, IStructCustomizationUtils& StructCustomizationUtils ) OVERRIDE;
 };
 
 /**
  * Customizes an InputActionKeyMapping struct to display it more simply
  */
-class FInputActionMappingCustomization : public IStructCustomization
+class FInputActionMappingCustomization : public IPropertyTypeCustomization
 {
 public:
-	static TSharedRef<IStructCustomization> MakeInstance();
+	static TSharedRef<IPropertyTypeCustomization> MakeInstance();
 
-	/** IStructCustomization interface */
-	virtual void CustomizeStructHeader( TSharedRef<class IPropertyHandle> InStructPropertyHandle, class FDetailWidgetRow& HeaderRow, IStructCustomizationUtils& StructCustomizationUtils ) OVERRIDE;
-	virtual void CustomizeStructChildren( TSharedRef<class IPropertyHandle> InStructPropertyHandle, class IDetailChildrenBuilder& StructBuilder, IStructCustomizationUtils& StructCustomizationUtils ) OVERRIDE;
+	/** IPropertyTypeCustomization interface */
+	virtual void CustomizeHeader( TSharedRef<class IPropertyHandle> InStructPropertyHandle, class FDetailWidgetRow& HeaderRow, IStructCustomizationUtils& StructCustomizationUtils ) OVERRIDE;
+	virtual void CustomizeChildren( TSharedRef<class IPropertyHandle> InStructPropertyHandle, class IDetailChildrenBuilder& StructBuilder, IStructCustomizationUtils& StructCustomizationUtils ) OVERRIDE;
 
 private:
 	void RemoveActionMappingButton_OnClick();
@@ -37,14 +37,14 @@ private:
 /**
  * Customizes an InputAxisKeyMapping struct to display it more simply
  */
-class FInputAxisMappingCustomization : public IStructCustomization
+class FInputAxisMappingCustomization : public IPropertyTypeCustomization
 {
 public:
-	static TSharedRef<IStructCustomization> MakeInstance();
+	static TSharedRef<IPropertyTypeCustomization> MakeInstance();
 
-	/** IStructCustomization interface */
-	virtual void CustomizeStructHeader( TSharedRef<class IPropertyHandle> InStructPropertyHandle, class FDetailWidgetRow& HeaderRow, IStructCustomizationUtils& StructCustomizationUtils ) OVERRIDE;
-	virtual void CustomizeStructChildren( TSharedRef<class IPropertyHandle> InStructPropertyHandle, class IDetailChildrenBuilder& StructBuilder, IStructCustomizationUtils& StructCustomizationUtils ) OVERRIDE;
+	/** IPropertyTypeCustomization interface */
+	virtual void CustomizeHeader( TSharedRef<class IPropertyHandle> InStructPropertyHandle, class FDetailWidgetRow& HeaderRow, IStructCustomizationUtils& StructCustomizationUtils ) OVERRIDE;
+	virtual void CustomizeChildren( TSharedRef<class IPropertyHandle> InStructPropertyHandle, class IDetailChildrenBuilder& StructBuilder, IStructCustomizationUtils& StructCustomizationUtils ) OVERRIDE;
 
 private:
 	void RemoveAxisMappingButton_OnClick();
diff --git a/Engine/Source/Editor/DetailCustomizations/Private/KeyStructCustomization.cpp b/Engine/Source/Editor/DetailCustomizations/Private/KeyStructCustomization.cpp
--- a/Engine/Source/Editor/DetailCustomizations/Private/KeyStructCustomization.cpp
+++ b/Engine/Source/Editor/DetailCustomizations/Private/KeyStructCustomization.cpp
@@ -14,16 +14,16 @@
 /* FKeyStructCustomization static interface
  *****************************************************************************/
 
-TSharedRef<IStructCustomization> FKeyStructCustomization::MakeInstance( )
+TSharedRef<IPropertyTypeCustomization> FKeyStructCustomization::MakeInstance( )
 {
 	return MakeShareable(new FKeyStructCustomization);
 }
 
 
-/* IStructCustomization interface
+/* IPropertyTypeCustomization interface
  *****************************************************************************/
 
-void FKeyStructCustomization::CustomizeStructHeader( TSharedRef<class IPropertyHandle> StructPropertyHandle, class FDetailWidgetRow& HeaderRow, IStructCustomizationUtils& StructCustomizationUtils )
+void FKeyStructCustomization::CustomizeHeader( TSharedRef<class IPropertyHandle> StructPropertyHandle, class FDetailWidgetRow& HeaderRow, IStructCustomizationUtils& StructCustomizationUtils )
 {
 	PropertyHandle = StructPropertyHandle;
 
diff --git a/Engine/Source/Editor/DetailCustomizations/Private/KeyStructCustomization.h b/Engine/Source/Editor/DetailCustomizations/Private/KeyStructCustomization.h
--- a/Engine/Source/Editor/DetailCustomizations/Private/KeyStructCustomization.h
+++ b/Engine/Source/Editor/DetailCustomizations/Private/KeyStructCustomization.h
@@ -10,23 +10,23 @@
  * Implements a details panel customization for FKey structures.
  */
 class FKeyStructCustomization
-	: public IStructCustomization
+	: public IPropertyTypeCustomization
 {
 public:
 
-	// Begin IStructCustomization interface
+	// Begin IPropertyTypeCustomization interface
 
-	virtual void CustomizeStructHeader(TSharedRef<class IPropertyHandle> StructPropertyHandle, class FDetailWidgetRow& HeaderRow, IStructCustomizationUtils& StructCustomizationUtils) OVERRIDE;
-	virtual void CustomizeStructChildren(TSharedRef<IPropertyHandle> StructPropertyHandle, class IDetailChildrenBuilder& ChildBuilder, IStructCustomizationUtils& StructCustomizationUtils) OVERRIDE {};
+	virtual void CustomizeHeader(TSharedRef<class IPropertyHandle> StructPropertyHandle, class FDetailWidgetRow& HeaderRow, IStructCustomizationUtils& StructCustomizationUtils) OVERRIDE;
+	virtual void CustomizeChildren(TSharedRef<IPropertyHandle> StructPropertyHandle, class IDetailChildrenBuilder& ChildBuilder, IStructCustomizationUtils& StructCustomizationUtils) OVERRIDE {};
 
-	// End IStructCustomization interface
+	// End IPropertyTypeCustomization interface
 
 	/**
 	 * Creates a new instance.
 	 *
 	 * @return A new struct customization for Keys.
 	 */
-	static TSharedRef<IStructCustomization> MakeInstance( );
+	static TSharedRef<IPropertyTypeCustomization> MakeInstance( );
 
 private:
 
diff --git a/Engine/Source/Editor/DetailCustomizations/Private/MarginCustomization.cpp b/Engine/Source/Editor/DetailCustomizations/Private/MarginCustomization.cpp
--- a/Engine/Source/Editor/DetailCustomizations/Private/MarginCustomization.cpp
+++ b/Engine/Source/Editor/DetailCustomizations/Private/MarginCustomization.cpp
@@ -3,12 +3,12 @@
 #include "DetailCustomizationsPrivatePCH.h"
 #include "MarginCustomization.h"
 
-TSharedRef<IStructCustomization> FMarginStructCustomization::MakeInstance() 
+TSharedRef<IPropertyTypeCustomization> FMarginStructCustomization::MakeInstance() 
 {
 	return MakeShareable( new FMarginStructCustomization() );
 }
 
-void FMarginStructCustomization::CustomizeStructHeader( TSharedRef<class IPropertyHandle> StructPropertyHandle, class FDetailWidgetRow& HeaderRow, IStructCustomizationUtils& StructCustomizationUtils )
+void FMarginStructCustomization::CustomizeHeader( TSharedRef<class IPropertyHandle> StructPropertyHandle, class FDetailWidgetRow& HeaderRow, IStructCustomizationUtils& StructCustomizationUtils )
 {
 	const FString UVSpaceString( StructPropertyHandle->GetProperty()->GetMetaData( TEXT( "UVSpace" ) ) );
 	bIsMarginUsingUVSpace = UVSpaceString.Len() > 0 && UVSpaceString == TEXT( "true" );
@@ -42,7 +42,7 @@ void FMarginStructCustomization::CustomizeStructHeader( TSharedRef<class IProper
 	UpdateMarginTextFromProperties();
 }
 
-void FMarginStructCustomization::CustomizeStructChildren( TSharedRef<class IPropertyHandle> StructPropertyHandle, class IDetailChildrenBuilder& StructBuilder, IStructCustomizationUtils& StructCustomizationUtils )
+void FMarginStructCustomization::CustomizeChildren( TSharedRef<class IPropertyHandle> StructPropertyHandle, class IDetailChildrenBuilder& StructBuilder, IStructCustomizationUtils& StructCustomizationUtils )
 {
 	FSimpleDelegate OnPropertyChangedDelegate = FSimpleDelegate::CreateSP( this, &FMarginStructCustomization::OnPropertyChanged );
 	StructPropertyHandle->SetOnPropertyValueChanged( OnPropertyChangedDelegate );
diff --git a/Engine/Source/Editor/DetailCustomizations/Private/MarginCustomization.h b/Engine/Source/Editor/DetailCustomizations/Private/MarginCustomization.h
--- a/Engine/Source/Editor/DetailCustomizations/Private/MarginCustomization.h
+++ b/Engine/Source/Editor/DetailCustomizations/Private/MarginCustomization.h
@@ -2,15 +2,15 @@
 
 #pragma once
 
-class FMarginStructCustomization : public IStructCustomization
+class FMarginStructCustomization : public IPropertyTypeCustomization
 {
 public:
-	static TSharedRef<IStructCustomization> MakeInstance();
+	static TSharedRef<IPropertyTypeCustomization> MakeInstance();
 
-	/** IStructCustomization interface */
-	virtual void CustomizeStructHeader( TSharedRef<class IPropertyHandle> StructPropertyHandle, class FDetailWidgetRow& HeaderRow, IStructCustomizationUtils& StructCustomizationUtils ) OVERRIDE;
+	/** IPropertyTypeCustomization interface */
+	virtual void CustomizeHeader( TSharedRef<class IPropertyHandle> StructPropertyHandle, class FDetailWidgetRow& HeaderRow, IStructCustomizationUtils& StructCustomizationUtils ) OVERRIDE;
 
-	virtual void CustomizeStructChildren( TSharedRef<class IPropertyHandle> StructPropertyHandle, class IDetailChildrenBuilder& StructBuilder, IStructCustomizationUtils& StructCustomizationUtils ) OVERRIDE;              
+	virtual void CustomizeChildren( TSharedRef<class IPropertyHandle> StructPropertyHandle, class IDetailChildrenBuilder& StructBuilder, IStructCustomizationUtils& StructCustomizationUtils ) OVERRIDE;              
 
 private:
 
diff --git a/Engine/Source/Editor/DetailCustomizations/Private/MathStructCustomizations.cpp b/Engine/Source/Editor/DetailCustomizations/Private/MathStructCustomizations.cpp
--- a/Engine/Source/Editor/DetailCustomizations/Private/MathStructCustomizations.cpp
+++ b/Engine/Source/Editor/DetailCustomizations/Private/MathStructCustomizations.cpp
@@ -6,19 +6,19 @@
 
 #define LOCTEXT_NAMESPACE "FMathStructCustomization"
 
-TSharedRef<IStructCustomization> FMathStructCustomization::MakeInstance() 
+TSharedRef<IPropertyTypeCustomization> FMathStructCustomization::MakeInstance() 
 {
 	return MakeShareable( new FMathStructCustomization );
 }
 
-void FMathStructCustomization::CustomizeStructHeader( TSharedRef<class IPropertyHandle> StructPropertyHandle, class FDetailWidgetRow& HeaderRow, IStructCustomizationUtils& StructCustomizationUtils )
+void FMathStructCustomization::CustomizeHeader( TSharedRef<class IPropertyHandle> StructPropertyHandle, class FDetailWidgetRow& HeaderRow, IStructCustomizationUtils& StructCustomizationUtils )
 {
 	GetSortedChildren( StructPropertyHandle, SortedChildHandles );
 
 	MakeHeaderRow( StructPropertyHandle, HeaderRow );
 }
 
-void FMathStructCustomization::CustomizeStructChildren( TSharedRef<class IPropertyHandle> StructPropertyHandle, class IDetailChildrenBuilder& StructBuilder, IStructCustomizationUtils& StructCustomizationUtils )
+void FMathStructCustomization::CustomizeChildren( TSharedRef<class IPropertyHandle> StructPropertyHandle, class IDetailChildrenBuilder& StructBuilder, IStructCustomizationUtils& StructCustomizationUtils )
 {
 	for( int32 ChildIndex = 0; ChildIndex < SortedChildHandles.Num(); ++ChildIndex )
 	{
@@ -288,7 +288,7 @@ void FMathStructCustomization::OnEndSliderMovement( NumericType NewValue )
 	GEditor->EndTransaction();
 }
 
-TSharedRef<IStructCustomization> FRotatorStructCustomization::MakeInstance() 
+TSharedRef<IPropertyTypeCustomization> FRotatorStructCustomization::MakeInstance() 
 {
 	return MakeShareable( new FRotatorStructCustomization );
 }
@@ -330,12 +330,12 @@ void FRotatorStructCustomization::GetSortedChildren( TSharedRef<IPropertyHandle>
 }
 
 
-TSharedRef<IStructCustomization> FColorStructCustomization::MakeInstance() 
+TSharedRef<IPropertyTypeCustomization> FColorStructCustomization::MakeInstance() 
 {
 	return MakeShareable( new FColorStructCustomization );
 }
 
-void FColorStructCustomization::CustomizeStructHeader( TSharedRef<class IPropertyHandle> InStructPropertyHandle, class FDetailWidgetRow& InHeaderRow, IStructCustomizationUtils& StructCustomizationUtils )
+void FColorStructCustomization::CustomizeHeader( TSharedRef<class IPropertyHandle> InStructPropertyHandle, class FDetailWidgetRow& InHeaderRow, IStructCustomizationUtils& StructCustomizationUtils )
 {
 	TSharedPtr<FAssetThumbnailPool> Pool = StructCustomizationUtils.GetThumbnailPool();
 
@@ -344,7 +344,7 @@ void FColorStructCustomization::CustomizeStructHeader( TSharedRef<class IPropert
 	bIsLinearColor = CastChecked<UStructProperty>( StructPropertyHandle->GetProperty() )->Struct->GetFName() == NAME_LinearColor;
 	bIgnoreAlpha = StructPropertyHandle->GetProperty()->HasMetaData(TEXT("HideAlphaChannel"));
 
-	FMathStructCustomization::CustomizeStructHeader( InStructPropertyHandle, InHeaderRow, StructCustomizationUtils );
+	FMathStructCustomization::CustomizeHeader( InStructPropertyHandle, InHeaderRow, StructCustomizationUtils );
 }
 
 void FColorStructCustomization::MakeHeaderRow( TSharedRef<class IPropertyHandle>& InStructPropertyHandle, FDetailWidgetRow& Row )
diff --git a/Engine/Source/Editor/DetailCustomizations/Private/MathStructCustomizations.h b/Engine/Source/Editor/DetailCustomizations/Private/MathStructCustomizations.h
--- a/Engine/Source/Editor/DetailCustomizations/Private/MathStructCustomizations.h
+++ b/Engine/Source/Editor/DetailCustomizations/Private/MathStructCustomizations.h
@@ -5,10 +5,10 @@
 /**
  * Base class for math struct customization (e.g, vector, rotator, color)                                                              
  */
-class FMathStructCustomization : public IStructCustomization
+class FMathStructCustomization : public IPropertyTypeCustomization
 {
 public:
-	static TSharedRef<IStructCustomization> MakeInstance();
+	static TSharedRef<IPropertyTypeCustomization> MakeInstance();
 
 	FMathStructCustomization()
 		: bIsUsingSlider(false)
@@ -17,9 +17,9 @@ public:
 
 	virtual ~FMathStructCustomization() {}
 
-	/** IStructCustomization instance */
-	virtual void CustomizeStructHeader( TSharedRef<class IPropertyHandle> StructPropertyHandle, class FDetailWidgetRow& HeaderRow, IStructCustomizationUtils& StructCustomizationUtils ) OVERRIDE;
-	virtual void CustomizeStructChildren( TSharedRef<class IPropertyHandle> StructPropertyHandle, class IDetailChildrenBuilder& StructBuilder, IStructCustomizationUtils& StructCustomizationUtils ) OVERRIDE;
+	/** IPropertyTypeCustomization instance */
+	virtual void CustomizeHeader( TSharedRef<class IPropertyHandle> StructPropertyHandle, class FDetailWidgetRow& HeaderRow, IStructCustomizationUtils& StructCustomizationUtils ) OVERRIDE;
+	virtual void CustomizeChildren( TSharedRef<class IPropertyHandle> StructPropertyHandle, class IDetailChildrenBuilder& StructBuilder, IStructCustomizationUtils& StructCustomizationUtils ) OVERRIDE;
 	
 protected:
 	/**
@@ -129,7 +129,7 @@ class FRotatorStructCustomization : public FMathStructCustomization
 {
 public:
 	/** @return A new instance of this class */
-	static TSharedRef<IStructCustomization> MakeInstance();
+	static TSharedRef<IPropertyTypeCustomization> MakeInstance();
 private:
 	/** FMathStructCustomization interface */
 	virtual void GetSortedChildren( TSharedRef<IPropertyHandle> StructPropertyHandle, TArray< TSharedRef<IPropertyHandle> >& OutChildren ) OVERRIDE;
@@ -141,7 +141,7 @@ private:
 class FColorStructCustomization : public FMathStructCustomization
 {
 public:
-	static TSharedRef<IStructCustomization> MakeInstance();
+	static TSharedRef<IPropertyTypeCustomization> MakeInstance();
 	
 private:
 	FColorStructCustomization()
@@ -151,7 +151,7 @@ private:
 	{}
 
 	/** FMathStructCustomization interface */
-	virtual void CustomizeStructHeader( TSharedRef<class IPropertyHandle> StructPropertyHandle, class FDetailWidgetRow& HeaderRow, IStructCustomizationUtils& StructCustomizationUtils ) OVERRIDE;
+	virtual void CustomizeHeader( TSharedRef<class IPropertyHandle> StructPropertyHandle, class FDetailWidgetRow& HeaderRow, IStructCustomizationUtils& StructCustomizationUtils ) OVERRIDE;
 	virtual void MakeHeaderRow( TSharedRef<class IPropertyHandle>& InStructPropertyHandle, FDetailWidgetRow& Row ) OVERRIDE;
 	virtual void GetSortedChildren( TSharedRef<IPropertyHandle> StructPropertyHandle, TArray< TSharedRef<IPropertyHandle> >& OutChildren ) OVERRIDE;
 
diff --git a/Engine/Source/Editor/DetailCustomizations/Private/MathStructProxyCustomizations.cpp b/Engine/Source/Editor/DetailCustomizations/Private/MathStructProxyCustomizations.cpp
--- a/Engine/Source/Editor/DetailCustomizations/Private/MathStructProxyCustomizations.cpp
+++ b/Engine/Source/Editor/DetailCustomizations/Private/MathStructProxyCustomizations.cpp
@@ -3,7 +3,7 @@
 #include "DetailCustomizationsPrivatePCH.h"
 #include "MathStructProxyCustomizations.h"
 
-void FMathStructProxyCustomization::CustomizeStructChildren( TSharedRef<class IPropertyHandle> StructPropertyHandle, class IDetailChildrenBuilder& StructBuilder, IStructCustomizationUtils& StructCustomizationUtils )
+void FMathStructProxyCustomization::CustomizeChildren( TSharedRef<class IPropertyHandle> StructPropertyHandle, class IDetailChildrenBuilder& StructBuilder, IStructCustomizationUtils& StructCustomizationUtils )
 {
 
 }
@@ -85,7 +85,7 @@ void FMathStructProxyCustomization::OnEndSliderMovement( NumericType NewValue )
 }
 
 
-TSharedRef<IStructCustomization> FMatrixStructCustomization::MakeInstance() 
+TSharedRef<IPropertyTypeCustomization> FMatrixStructCustomization::MakeInstance() 
 {
 	return MakeShareable( new FMatrixStructCustomization );
 }
@@ -248,7 +248,7 @@ bool FMatrixStructCustomization::FlushValues( TWeakPtr<IPropertyHandle> WeakHand
 	return true;
 }
 
-TSharedRef<IStructCustomization> FTransformStructCustomization::MakeInstance() 
+TSharedRef<IPropertyTypeCustomization> FTransformStructCustomization::MakeInstance() 
 {
 	return MakeShareable( new FTransformStructCustomization );
 }
diff --git a/Engine/Source/Editor/DetailCustomizations/Private/MathStructProxyCustomizations.h b/Engine/Source/Editor/DetailCustomizations/Private/MathStructProxyCustomizations.h
--- a/Engine/Source/Editor/DetailCustomizations/Private/MathStructProxyCustomizations.h
+++ b/Engine/Source/Editor/DetailCustomizations/Private/MathStructProxyCustomizations.h
@@ -137,8 +137,8 @@ private:
 class FMathStructProxyCustomization : public FMathStructCustomization
 {
 public:
-	/** IStructCustomization interface */
-	virtual void CustomizeStructChildren( TSharedRef<class IPropertyHandle> StructPropertyHandle, class IDetailChildrenBuilder& StructBuilder, IStructCustomizationUtils& StructCustomizationUtils ) OVERRIDE;
+	/** IPropertyTypeCustomization interface */
+	virtual void CustomizeChildren( TSharedRef<class IPropertyHandle> StructPropertyHandle, class IDetailChildrenBuilder& StructBuilder, IStructCustomizationUtils& StructCustomizationUtils ) OVERRIDE;
 
 	/** FMathStructCustomization interface */
 	virtual void MakeHeaderRow( TSharedRef<class IPropertyHandle>& InStructPropertyHandle, FDetailWidgetRow& Row ) OVERRIDE;
@@ -213,7 +213,7 @@ private:
 class FMatrixStructCustomization : public FMathStructProxyCustomization
 {
 public:
-	static TSharedRef<IStructCustomization> MakeInstance();
+	static TSharedRef<IPropertyTypeCustomization> MakeInstance();
 
 public:
 	FMatrixStructCustomization()
@@ -266,7 +266,7 @@ protected:
 class FTransformStructCustomization : public FMatrixStructCustomization
 {
 public:
-	static TSharedRef<IStructCustomization> MakeInstance();
+	static TSharedRef<IPropertyTypeCustomization> MakeInstance();
 
 protected:
 	/** FMathStructProxyCustomization interface */
diff --git a/Engine/Source/Editor/DetailCustomizations/Private/NavLinkStructCustomization.cpp b/Engine/Source/Editor/DetailCustomizations/Private/NavLinkStructCustomization.cpp
--- a/Engine/Source/Editor/DetailCustomizations/Private/NavLinkStructCustomization.cpp
+++ b/Engine/Source/Editor/DetailCustomizations/Private/NavLinkStructCustomization.cpp
@@ -5,12 +5,12 @@
 
 #define LOCTEXT_NAMESPACE "FNavLinkStructCustomization"
 
-TSharedRef<IStructCustomization> FNavLinkStructCustomization::MakeInstance( )
+TSharedRef<IPropertyTypeCustomization> FNavLinkStructCustomization::MakeInstance( )
 {
 	return MakeShareable(new FNavLinkStructCustomization);
 }
 
-void FNavLinkStructCustomization::CustomizeStructHeader( TSharedRef<IPropertyHandle> StructPropertyHandle, class FDetailWidgetRow& HeaderRow, IStructCustomizationUtils& StructCustomizationUtils )
+void FNavLinkStructCustomization::CustomizeHeader( TSharedRef<IPropertyHandle> StructPropertyHandle, class FDetailWidgetRow& HeaderRow, IStructCustomizationUtils& StructCustomizationUtils )
 {
 	TSharedPtr<IPropertyHandle> CommentHandle = StructPropertyHandle->GetChildHandle(GET_MEMBER_NAME_CHECKED(FNavigationLinkBase,Description));
 	FString Desc;
@@ -33,7 +33,7 @@ void FNavLinkStructCustomization::CustomizeStructHeader( TSharedRef<IPropertyHan
 	];
 }
 
-void FNavLinkStructCustomization::CustomizeStructChildren( TSharedRef<class IPropertyHandle> StructPropertyHandle, class IDetailChildrenBuilder& StructBuilder, IStructCustomizationUtils& StructCustomizationUtils )
+void FNavLinkStructCustomization::CustomizeChildren( TSharedRef<class IPropertyHandle> StructPropertyHandle, class IDetailChildrenBuilder& StructBuilder, IStructCustomizationUtils& StructCustomizationUtils )
 {
 	uint32 NumChildren = 0;
 	StructPropertyHandle->GetNumChildren(NumChildren);
diff --git a/Engine/Source/Editor/DetailCustomizations/Private/NavLinkStructCustomization.h b/Engine/Source/Editor/DetailCustomizations/Private/NavLinkStructCustomization.h
--- a/Engine/Source/Editor/DetailCustomizations/Private/NavLinkStructCustomization.h
+++ b/Engine/Source/Editor/DetailCustomizations/Private/NavLinkStructCustomization.h
@@ -2,13 +2,13 @@
 
 #pragma once
 
-class FNavLinkStructCustomization : public IStructCustomization
+class FNavLinkStructCustomization : public IPropertyTypeCustomization
 {
 public:
-	// Begin IStructCustomization interface
-	virtual void CustomizeStructHeader( TSharedRef<IPropertyHandle> StructPropertyHandle, class FDetailWidgetRow& HeaderRow, IStructCustomizationUtils& StructCustomizationUtils ) OVERRIDE;
-	virtual void CustomizeStructChildren( TSharedRef<class IPropertyHandle> StructPropertyHandle, class IDetailChildrenBuilder& StructBuilder, IStructCustomizationUtils& StructCustomizationUtils ) OVERRIDE;
-	// End IStructCustomization interface
+	// Begin IPropertyTypeCustomization interface
+	virtual void CustomizeHeader( TSharedRef<IPropertyHandle> StructPropertyHandle, class FDetailWidgetRow& HeaderRow, IStructCustomizationUtils& StructCustomizationUtils ) OVERRIDE;
+	virtual void CustomizeChildren( TSharedRef<class IPropertyHandle> StructPropertyHandle, class IDetailChildrenBuilder& StructBuilder, IStructCustomizationUtils& StructCustomizationUtils ) OVERRIDE;
+	// End IPropertyTypeCustomization interface
 
-	static TSharedRef<IStructCustomization> MakeInstance();
+	static TSharedRef<IPropertyTypeCustomization> MakeInstance();
 };
diff --git a/Engine/Source/Editor/DetailCustomizations/Private/RangeStructCustomization.cpp b/Engine/Source/Editor/DetailCustomizations/Private/RangeStructCustomization.cpp
--- a/Engine/Source/Editor/DetailCustomizations/Private/RangeStructCustomization.cpp
+++ b/Engine/Source/Editor/DetailCustomizations/Private/RangeStructCustomization.cpp
@@ -43,17 +43,17 @@ namespace
  *****************************************************************************/
 
 template <typename NumericType>
-TSharedRef<IStructCustomization> FRangeStructCustomization<NumericType>::MakeInstance()
+TSharedRef<IPropertyTypeCustomization> FRangeStructCustomization<NumericType>::MakeInstance()
 {
 	return MakeShareable(new FRangeStructCustomization<NumericType>);
 }
 
 
-/* IStructCustomization interface
+/* IPropertyTypeCustomization interface
  *****************************************************************************/
 
 template <typename NumericType>
-void FRangeStructCustomization<NumericType>::CustomizeStructHeader(TSharedRef<IPropertyHandle> StructPropertyHandle, FDetailWidgetRow& HeaderRow, IStructCustomizationUtils& StructCustomizationUtils)
+void FRangeStructCustomization<NumericType>::CustomizeHeader(TSharedRef<IPropertyHandle> StructPropertyHandle, FDetailWidgetRow& HeaderRow, IStructCustomizationUtils& StructCustomizationUtils)
 {
 	// Get handles to the properties we're interested in
 	LowerBoundStructHandle = StructPropertyHandle->GetChildHandle(TEXT("LowerBound"));
@@ -221,7 +221,7 @@ void FRangeStructCustomization<NumericType>::CustomizeStructHeader(TSharedRef<IP
 
 
 template <typename NumericType>
-void FRangeStructCustomization<NumericType>::CustomizeStructChildren(TSharedRef<IPropertyHandle> StructPropertyHandle, IDetailChildrenBuilder& StructBuilder, IStructCustomizationUtils& StructCustomizationUtils)
+void FRangeStructCustomization<NumericType>::CustomizeChildren(TSharedRef<IPropertyHandle> StructPropertyHandle, IDetailChildrenBuilder& StructBuilder, IStructCustomizationUtils& StructCustomizationUtils)
 {
 	// Don't display children, as editing them directly can break the constraints
 }
diff --git a/Engine/Source/Editor/DetailCustomizations/Private/RangeStructCustomization.h b/Engine/Source/Editor/DetailCustomizations/Private/RangeStructCustomization.h
--- a/Engine/Source/Editor/DetailCustomizations/Private/RangeStructCustomization.h
+++ b/Engine/Source/Editor/DetailCustomizations/Private/RangeStructCustomization.h
@@ -10,20 +10,20 @@
 * Implements a details panel customization for FFloatRange structures.
 */
 template <typename NumericType>
-class FRangeStructCustomization : public IStructCustomization
+class FRangeStructCustomization : public IPropertyTypeCustomization
 {
 public:
 
 	FRangeStructCustomization()
 		: bIsUsingSlider(false) {}
 
-	// Begin IStructCustomization interface
+	// Begin IPropertyTypeCustomization interface
 
-	virtual void CustomizeStructHeader(TSharedRef<IPropertyHandle> StructPropertyHandle, FDetailWidgetRow& HeaderRow, IStructCustomizationUtils& StructCustomizationUtils) OVERRIDE;
+	virtual void CustomizeHeader(TSharedRef<IPropertyHandle> StructPropertyHandle, FDetailWidgetRow& HeaderRow, IStructCustomizationUtils& StructCustomizationUtils) OVERRIDE;
 
-	virtual void CustomizeStructChildren(TSharedRef<IPropertyHandle> StructPropertyHandle, IDetailChildrenBuilder& StructBuilder, IStructCustomizationUtils& StructCustomizationUtils) OVERRIDE;
+	virtual void CustomizeChildren(TSharedRef<IPropertyHandle> StructPropertyHandle, IDetailChildrenBuilder& StructBuilder, IStructCustomizationUtils& StructCustomizationUtils) OVERRIDE;
 
-	// End IStructCustomization interface
+	// End IPropertyTypeCustomization interface
 
 public:
 
@@ -32,7 +32,7 @@ public:
 	*
 	* @return A new struct customization for Guids.
 	*/
-	static TSharedRef<IStructCustomization> MakeInstance();
+	static TSharedRef<IPropertyTypeCustomization> MakeInstance();
 
 protected:
 
diff --git a/Engine/Source/Editor/DetailCustomizations/Private/SlateBrushCustomization.cpp b/Engine/Source/Editor/DetailCustomizations/Private/SlateBrushCustomization.cpp
--- a/Engine/Source/Editor/DetailCustomizations/Private/SlateBrushCustomization.cpp
+++ b/Engine/Source/Editor/DetailCustomizations/Private/SlateBrushCustomization.cpp
@@ -933,12 +933,12 @@ const float SSlateBrushPreview::ImagePadding = 5.0f;
 const float SSlateBrushPreview::BorderHitSize = 8.0f;
 
 
-TSharedRef<IStructCustomization> FSlateBrushStructCustomization::MakeInstance() 
+TSharedRef<IPropertyTypeCustomization> FSlateBrushStructCustomization::MakeInstance() 
 {
 	return MakeShareable( new FSlateBrushStructCustomization() );
 }
 
-void FSlateBrushStructCustomization::CustomizeStructHeader( TSharedRef<IPropertyHandle> StructPropertyHandle, FDetailWidgetRow& HeaderRow, IStructCustomizationUtils& StructCustomizationUtils )
+void FSlateBrushStructCustomization::CustomizeHeader( TSharedRef<IPropertyHandle> StructPropertyHandle, FDetailWidgetRow& HeaderRow, IStructCustomizationUtils& StructCustomizationUtils )
 {
 	bool ShowOnlyInnerProperties = StructPropertyHandle->GetProperty()->HasMetaData(TEXT("ShowOnlyInnerProperties"));
 
@@ -952,7 +952,7 @@ void FSlateBrushStructCustomization::CustomizeStructHeader( TSharedRef<IProperty
 	}
 }
 
-void FSlateBrushStructCustomization::CustomizeStructChildren( TSharedRef<IPropertyHandle> StructPropertyHandle, IDetailChildrenBuilder& StructBuilder, IStructCustomizationUtils& StructCustomizationUtils )
+void FSlateBrushStructCustomization::CustomizeChildren( TSharedRef<IPropertyHandle> StructPropertyHandle, IDetailChildrenBuilder& StructBuilder, IStructCustomizationUtils& StructCustomizationUtils )
 {
 	// Add the child properties
 	TSharedPtr<IPropertyHandle> ImageSizeProperty = StructPropertyHandle->GetChildHandle( TEXT("ImageSize") );
diff --git a/Engine/Source/Editor/DetailCustomizations/Private/SlateBrushCustomization.h b/Engine/Source/Editor/DetailCustomizations/Private/SlateBrushCustomization.h
--- a/Engine/Source/Editor/DetailCustomizations/Private/SlateBrushCustomization.h
+++ b/Engine/Source/Editor/DetailCustomizations/Private/SlateBrushCustomization.h
@@ -3,15 +3,15 @@
 #pragma once
 
 
-class FSlateBrushStructCustomization : public IStructCustomization
+class FSlateBrushStructCustomization : public IPropertyTypeCustomization
 {
 public:
-	static TSharedRef<IStructCustomization> MakeInstance();
+	static TSharedRef<IPropertyTypeCustomization> MakeInstance();
 
-	/** IStructCustomization interface */
-	virtual void CustomizeStructHeader( TSharedRef<class IPropertyHandle> StructPropertyHandle, class FDetailWidgetRow& HeaderRow, IStructCustomizationUtils& StructCustomizationUtils ) OVERRIDE;
+	/** IPropertyTypeCustomization interface */
+	virtual void CustomizeHeader( TSharedRef<class IPropertyHandle> StructPropertyHandle, class FDetailWidgetRow& HeaderRow, IStructCustomizationUtils& StructCustomizationUtils ) OVERRIDE;
 
-	virtual void CustomizeStructChildren( TSharedRef<class IPropertyHandle> StructPropertyHandle, class IDetailChildrenBuilder& StructBuilder, IStructCustomizationUtils& StructCustomizationUtils ) OVERRIDE;              
+	virtual void CustomizeChildren( TSharedRef<class IPropertyHandle> StructPropertyHandle, class IDetailChildrenBuilder& StructBuilder, IStructCustomizationUtils& StructCustomizationUtils ) OVERRIDE;              
 
 private:
 	/**
diff --git a/Engine/Source/Editor/DetailCustomizations/Private/SlateColorCustomization.cpp b/Engine/Source/Editor/DetailCustomizations/Private/SlateColorCustomization.cpp
--- a/Engine/Source/Editor/DetailCustomizations/Private/SlateColorCustomization.cpp
+++ b/Engine/Source/Editor/DetailCustomizations/Private/SlateColorCustomization.cpp
@@ -3,12 +3,12 @@
 #include "DetailCustomizationsPrivatePCH.h"
 #include "SlateColorCustomization.h"
 
-TSharedRef<IStructCustomization> FSlateColorCustomization::MakeInstance() 
+TSharedRef<IPropertyTypeCustomization> FSlateColorCustomization::MakeInstance() 
 {
 	return MakeShareable( new FSlateColorCustomization() );
 }
 
-void FSlateColorCustomization::CustomizeStructHeader( TSharedRef<class IPropertyHandle> StructPropertyHandle, class FDetailWidgetRow& HeaderRow, IStructCustomizationUtils& StructCustomizationUtils )
+void FSlateColorCustomization::CustomizeHeader( TSharedRef<class IPropertyHandle> StructPropertyHandle, class FDetailWidgetRow& HeaderRow, IStructCustomizationUtils& StructCustomizationUtils )
 {
 	uint32 NumChildren;
 	StructPropertyHandle->GetNumChildren( NumChildren );
@@ -31,7 +31,7 @@ void FSlateColorCustomization::CustomizeStructHeader( TSharedRef<class IProperty
 	check( SpecifiedColorHandle.IsValid() );
 }
 
-void FSlateColorCustomization::CustomizeStructChildren( TSharedRef<class IPropertyHandle> StructPropertyHandle, class IDetailChildrenBuilder& StructBuilder, IStructCustomizationUtils& StructCustomizationUtils )
+void FSlateColorCustomization::CustomizeChildren( TSharedRef<class IPropertyHandle> StructPropertyHandle, class IDetailChildrenBuilder& StructBuilder, IStructCustomizationUtils& StructCustomizationUtils )
 {
 	SpecifiedColorHandle->SetOnPropertyValueChanged( FSimpleDelegate::CreateSP( SharedThis( this ), &FSlateColorCustomization::OnValueChanged ) );
 	StructBuilder.AddChildProperty( SpecifiedColorHandle.ToSharedRef() )
diff --git a/Engine/Source/Editor/DetailCustomizations/Private/SlateColorCustomization.h b/Engine/Source/Editor/DetailCustomizations/Private/SlateColorCustomization.h
--- a/Engine/Source/Editor/DetailCustomizations/Private/SlateColorCustomization.h
+++ b/Engine/Source/Editor/DetailCustomizations/Private/SlateColorCustomization.h
@@ -2,15 +2,15 @@
 
 #pragma once
 
-class FSlateColorCustomization : public IStructCustomization
+class FSlateColorCustomization : public IPropertyTypeCustomization
 {
 public:
-	static TSharedRef<IStructCustomization> MakeInstance();
+	static TSharedRef<IPropertyTypeCustomization> MakeInstance();
 
-	/** IStructCustomization interface */
-	virtual void CustomizeStructHeader( TSharedRef<class IPropertyHandle> StructPropertyHandle, class FDetailWidgetRow& HeaderRow, IStructCustomizationUtils& StructCustomizationUtils ) OVERRIDE;
+	/** IPropertyTypeCustomization interface */
+	virtual void CustomizeHeader( TSharedRef<class IPropertyHandle> StructPropertyHandle, class FDetailWidgetRow& HeaderRow, IStructCustomizationUtils& StructCustomizationUtils ) OVERRIDE;
 
-	virtual void CustomizeStructChildren( TSharedRef<class IPropertyHandle> StructPropertyHandle, class IDetailChildrenBuilder& StructBuilder, IStructCustomizationUtils& StructCustomizationUtils ) OVERRIDE;              
+	virtual void CustomizeChildren( TSharedRef<class IPropertyHandle> StructPropertyHandle, class IDetailChildrenBuilder& StructBuilder, IStructCustomizationUtils& StructCustomizationUtils ) OVERRIDE;              
 
 private:
 
diff --git a/Engine/Source/Editor/DetailCustomizations/Private/SlateSoundCustomization.cpp b/Engine/Source/Editor/DetailCustomizations/Private/SlateSoundCustomization.cpp
--- a/Engine/Source/Editor/DetailCustomizations/Private/SlateSoundCustomization.cpp
+++ b/Engine/Source/Editor/DetailCustomizations/Private/SlateSoundCustomization.cpp
@@ -3,12 +3,12 @@
 #include "DetailCustomizationsPrivatePCH.h"
 #include "SlateSoundCustomization.h"
 
-TSharedRef<IStructCustomization> FSlateSoundStructCustomization::MakeInstance() 
+TSharedRef<IPropertyTypeCustomization> FSlateSoundStructCustomization::MakeInstance() 
 {
 	return MakeShareable(new FSlateSoundStructCustomization());
 }
 
-void FSlateSoundStructCustomization::CustomizeStructHeader(TSharedRef<IPropertyHandle> StructPropertyHandle, FDetailWidgetRow& HeaderRow, IStructCustomizationUtils& StructCustomizationUtils)
+void FSlateSoundStructCustomization::CustomizeHeader(TSharedRef<IPropertyHandle> StructPropertyHandle, FDetailWidgetRow& HeaderRow, IStructCustomizationUtils& StructCustomizationUtils)
 {
 	TSharedPtr<IPropertyHandle> ResourceObjectProperty = StructPropertyHandle->GetChildHandle(TEXT("ResourceObject"));
 	check(ResourceObjectProperty.IsValid());
@@ -36,7 +36,7 @@ void FSlateSoundStructCustomization::CustomizeStructHeader(TSharedRef<IPropertyH
 	];
 }
 
-void FSlateSoundStructCustomization::CustomizeStructChildren(TSharedRef<IPropertyHandle> StructPropertyHandle, IDetailChildrenBuilder& StructBuilder, IStructCustomizationUtils& StructCustomizationUtils)
+void FSlateSoundStructCustomization::CustomizeChildren(TSharedRef<IPropertyHandle> StructPropertyHandle, IDetailChildrenBuilder& StructBuilder, IStructCustomizationUtils& StructCustomizationUtils)
 {
 }
 
diff --git a/Engine/Source/Editor/DetailCustomizations/Private/SlateSoundCustomization.h b/Engine/Source/Editor/DetailCustomizations/Private/SlateSoundCustomization.h
--- a/Engine/Source/Editor/DetailCustomizations/Private/SlateSoundCustomization.h
+++ b/Engine/Source/Editor/DetailCustomizations/Private/SlateSoundCustomization.h
@@ -3,14 +3,14 @@
 #pragma once
 
 /** Customize the appearance of an FSlateSound */
-class FSlateSoundStructCustomization : public IStructCustomization
+class FSlateSoundStructCustomization : public IPropertyTypeCustomization
 {
 public:
-	static TSharedRef<IStructCustomization> MakeInstance();
+	static TSharedRef<IPropertyTypeCustomization> MakeInstance();
 
-	/** IStructCustomization interface */
-	virtual void CustomizeStructHeader(TSharedRef<class IPropertyHandle> StructPropertyHandle, class FDetailWidgetRow& HeaderRow, IStructCustomizationUtils& StructCustomizationUtils) OVERRIDE;
-	virtual void CustomizeStructChildren(TSharedRef<class IPropertyHandle> StructPropertyHandle, class IDetailChildrenBuilder& StructBuilder, IStructCustomizationUtils& StructCustomizationUtils) OVERRIDE;
+	/** IPropertyTypeCustomization interface */
+	virtual void CustomizeHeader(TSharedRef<class IPropertyHandle> StructPropertyHandle, class FDetailWidgetRow& HeaderRow, IStructCustomizationUtils& StructCustomizationUtils) OVERRIDE;
+	virtual void CustomizeChildren(TSharedRef<class IPropertyHandle> StructPropertyHandle, class IDetailChildrenBuilder& StructBuilder, IStructCustomizationUtils& StructCustomizationUtils) OVERRIDE;
 
 protected:
 	/** Called when the resource object used by this FSlateSound has been changed */
diff --git a/Engine/Source/Editor/DetailCustomizations/Private/StringAssetReferenceCustomization.cpp b/Engine/Source/Editor/DetailCustomizations/Private/StringAssetReferenceCustomization.cpp
--- a/Engine/Source/Editor/DetailCustomizations/Private/StringAssetReferenceCustomization.cpp
+++ b/Engine/Source/Editor/DetailCustomizations/Private/StringAssetReferenceCustomization.cpp
@@ -4,7 +4,7 @@
 #include "StringAssetReferenceCustomization.h"
 #include "AssetData.h"
 
-void FStringAssetReferenceCustomization::CustomizeStructHeader( TSharedRef<IPropertyHandle> InStructPropertyHandle, FDetailWidgetRow& HeaderRow, IStructCustomizationUtils& StructCustomizationUtils )
+void FStringAssetReferenceCustomization::CustomizeHeader( TSharedRef<IPropertyHandle> InStructPropertyHandle, FDetailWidgetRow& HeaderRow, IStructCustomizationUtils& StructCustomizationUtils )
 {
 	StructPropertyHandle = InStructPropertyHandle;
 
@@ -70,7 +70,7 @@ void FStringAssetReferenceCustomization::CustomizeStructHeader( TSharedRef<IProp
 	];
 }
 
-void FStringAssetReferenceCustomization::CustomizeStructChildren(TSharedRef<IPropertyHandle> InStructPropertyHandle, IDetailChildrenBuilder& StructBuilder, IStructCustomizationUtils& StructCustomizationUtils)
+void FStringAssetReferenceCustomization::CustomizeChildren(TSharedRef<IPropertyHandle> InStructPropertyHandle, IDetailChildrenBuilder& StructBuilder, IStructCustomizationUtils& StructCustomizationUtils)
 {
 }
 
diff --git a/Engine/Source/Editor/DetailCustomizations/Private/StringAssetReferenceCustomization.h b/Engine/Source/Editor/DetailCustomizations/Private/StringAssetReferenceCustomization.h
--- a/Engine/Source/Editor/DetailCustomizations/Private/StringAssetReferenceCustomization.h
+++ b/Engine/Source/Editor/DetailCustomizations/Private/StringAssetReferenceCustomization.h
@@ -5,17 +5,17 @@
 /**
  * Customizes a string asset reference to look like a UObject property
  */
-class FStringAssetReferenceCustomization : public IStructCustomization
+class FStringAssetReferenceCustomization : public IPropertyTypeCustomization
 {
 public:
-	static TSharedRef<IStructCustomization> MakeInstance() 
+	static TSharedRef<IPropertyTypeCustomization> MakeInstance() 
 	{
 		return MakeShareable( new FStringAssetReferenceCustomization );
 	}
 
-	/** IStructCustomization interface */
-	virtual void CustomizeStructHeader( TSharedRef<class IPropertyHandle> InStructPropertyHandle, class FDetailWidgetRow& HeaderRow, IStructCustomizationUtils& StructCustomizationUtils ) OVERRIDE;
-	virtual void CustomizeStructChildren( TSharedRef<class IPropertyHandle> InStructPropertyHandle, class IDetailChildrenBuilder& StructBuilder, IStructCustomizationUtils& StructCustomizationUtils ) OVERRIDE;
+	/** IPropertyTypeCustomization interface */
+	virtual void CustomizeHeader( TSharedRef<class IPropertyHandle> InStructPropertyHandle, class FDetailWidgetRow& HeaderRow, IStructCustomizationUtils& StructCustomizationUtils ) OVERRIDE;
+	virtual void CustomizeChildren( TSharedRef<class IPropertyHandle> InStructPropertyHandle, class IDetailChildrenBuilder& StructBuilder, IStructCustomizationUtils& StructCustomizationUtils ) OVERRIDE;
 
 private:
 	bool OnShouldFilterAsset( const class FAssetData& InAssetData ) const;
diff --git a/Engine/Source/Editor/DetailCustomizations/Private/StringClassReferenceCustomization.cpp b/Engine/Source/Editor/DetailCustomizations/Private/StringClassReferenceCustomization.cpp
--- a/Engine/Source/Editor/DetailCustomizations/Private/StringClassReferenceCustomization.cpp
+++ b/Engine/Source/Editor/DetailCustomizations/Private/StringClassReferenceCustomization.cpp
@@ -3,7 +3,7 @@
 #include "DetailCustomizationsPrivatePCH.h"
 #include "StringClassReferenceCustomization.h"
 
-void FStringClassReferenceCustomization::CustomizeStructHeader(TSharedRef<IPropertyHandle> InStructPropertyHandle, FDetailWidgetRow& HeaderRow, IStructCustomizationUtils& StructCustomizationUtils)
+void FStringClassReferenceCustomization::CustomizeHeader(TSharedRef<IPropertyHandle> InStructPropertyHandle, FDetailWidgetRow& HeaderRow, IStructCustomizationUtils& StructCustomizationUtils)
 {
 	StructPropertyHandle = InStructPropertyHandle;
 
@@ -41,7 +41,7 @@ void FStringClassReferenceCustomization::CustomizeStructHeader(TSharedRef<IPrope
 	];
 }
 
-void FStringClassReferenceCustomization::CustomizeStructChildren(TSharedRef<IPropertyHandle> InStructPropertyHandle, IDetailChildrenBuilder& StructBuilder, IStructCustomizationUtils& StructCustomizationUtils)
+void FStringClassReferenceCustomization::CustomizeChildren(TSharedRef<IPropertyHandle> InStructPropertyHandle, IDetailChildrenBuilder& StructBuilder, IStructCustomizationUtils& StructCustomizationUtils)
 {
 }
 
diff --git a/Engine/Source/Editor/DetailCustomizations/Private/StringClassReferenceCustomization.h b/Engine/Source/Editor/DetailCustomizations/Private/StringClassReferenceCustomization.h
--- a/Engine/Source/Editor/DetailCustomizations/Private/StringClassReferenceCustomization.h
+++ b/Engine/Source/Editor/DetailCustomizations/Private/StringClassReferenceCustomization.h
@@ -5,17 +5,17 @@
 /**
  * Customizes a string class reference to look like a UClass property
  */
-class FStringClassReferenceCustomization : public IStructCustomization
+class FStringClassReferenceCustomization : public IPropertyTypeCustomization
 {
 public:
-	static TSharedRef<IStructCustomization> MakeInstance() 
+	static TSharedRef<IPropertyTypeCustomization> MakeInstance() 
 	{
 		return MakeShareable(new FStringClassReferenceCustomization);
 	}
 
-	/** IStructCustomization interface */
-	virtual void CustomizeStructHeader(TSharedRef<class IPropertyHandle> InStructPropertyHandle, class FDetailWidgetRow& HeaderRow, IStructCustomizationUtils& StructCustomizationUtils) OVERRIDE;
-	virtual void CustomizeStructChildren(TSharedRef<class IPropertyHandle> InStructPropertyHandle, class IDetailChildrenBuilder& StructBuilder, IStructCustomizationUtils& StructCustomizationUtils) OVERRIDE;
+	/** IPropertyTypeCustomization interface */
+	virtual void CustomizeHeader(TSharedRef<class IPropertyHandle> InStructPropertyHandle, class FDetailWidgetRow& HeaderRow, IStructCustomizationUtils& StructCustomizationUtils) OVERRIDE;
+	virtual void CustomizeChildren(TSharedRef<class IPropertyHandle> InStructPropertyHandle, class IDetailChildrenBuilder& StructBuilder, IStructCustomizationUtils& StructCustomizationUtils) OVERRIDE;
 
 private:
 	/** @return The class currently set on this reference */
diff --git a/Engine/Source/Editor/DetailCustomizations/Private/VehicleTransmissionDataCustomization.cpp b/Engine/Source/Editor/DetailCustomizations/Private/VehicleTransmissionDataCustomization.cpp
--- a/Engine/Source/Editor/DetailCustomizations/Private/VehicleTransmissionDataCustomization.cpp
+++ b/Engine/Source/Editor/DetailCustomizations/Private/VehicleTransmissionDataCustomization.cpp
@@ -166,7 +166,7 @@ void FVehicleTransmissionDataCustomization::CreateGearUIDelegate(TSharedRef<IPro
 	CreateGearUIHelper(ChildrenBuilder.AddChildContent(Label.ToString()), Label, GearProperty, ForwardGear);
 }
 
-void FVehicleTransmissionDataCustomization::CustomizeStructChildren(TSharedRef<class IPropertyHandle> StructPropertyHandle, class IDetailChildrenBuilder& StructBuilder, IStructCustomizationUtils& StructCustomizationUtils)
+void FVehicleTransmissionDataCustomization::CustomizeChildren(TSharedRef<class IPropertyHandle> StructPropertyHandle, class IDetailChildrenBuilder& StructBuilder, IStructCustomizationUtils& StructCustomizationUtils)
 {
 	// copy all transmision instances I'm accessing right now
 	TArray<void*> StructPtrs;
@@ -268,7 +268,7 @@ void FVehicleTransmissionDataCustomization::RemoveGear(TSharedRef<IPropertyHandl
 }
 
 
-void FVehicleTransmissionDataCustomization::CustomizeStructHeader(TSharedRef<IPropertyHandle> StructPropertyHandle, FDetailWidgetRow& HeaderRow, IStructCustomizationUtils& StructCustomizationUtils)
+void FVehicleTransmissionDataCustomization::CustomizeHeader(TSharedRef<IPropertyHandle> StructPropertyHandle, FDetailWidgetRow& HeaderRow, IStructCustomizationUtils& StructCustomizationUtils)
 {
 	HeaderRow.
 	NameContent()
diff --git a/Engine/Source/Editor/DetailCustomizations/Private/VehicleTransmissionDataCustomization.h b/Engine/Source/Editor/DetailCustomizations/Private/VehicleTransmissionDataCustomization.h
--- a/Engine/Source/Editor/DetailCustomizations/Private/VehicleTransmissionDataCustomization.h
+++ b/Engine/Source/Editor/DetailCustomizations/Private/VehicleTransmissionDataCustomization.h
@@ -2,18 +2,18 @@
 
 #pragma once
 
-class FVehicleTransmissionDataCustomization : public IStructCustomization
+class FVehicleTransmissionDataCustomization : public IPropertyTypeCustomization
 {
 public:
 	
-	static TSharedRef<IStructCustomization> MakeInstance() 
+	static TSharedRef<IPropertyTypeCustomization> MakeInstance() 
 	{
 		return MakeShareable(new FVehicleTransmissionDataCustomization);
 	}
 
-	/** IStructCustomization interface */
-	virtual void CustomizeStructHeader(TSharedRef<IPropertyHandle> StructPropertyHandle, class FDetailWidgetRow& HeaderRow, IStructCustomizationUtils& StructCustomizationUtils) OVERRIDE;
-	virtual void CustomizeStructChildren( TSharedRef<class IPropertyHandle> StructPropertyHandle, class IDetailChildrenBuilder& StructBuilder, IStructCustomizationUtils& StructCustomizationUtils ) OVERRIDE;
+	/** IPropertyTypeCustomization interface */
+	virtual void CustomizeHeader(TSharedRef<IPropertyHandle> StructPropertyHandle, class FDetailWidgetRow& HeaderRow, IStructCustomizationUtils& StructCustomizationUtils) OVERRIDE;
+	virtual void CustomizeChildren( TSharedRef<class IPropertyHandle> StructPropertyHandle, class IDetailChildrenBuilder& StructBuilder, IStructCustomizationUtils& StructCustomizationUtils ) OVERRIDE;
 
 private:
 	FVehicleTransmissionData * SelectedTransmission;
diff --git a/Engine/Source/Editor/DetailCustomizations/Public/DetailCustomizations.h b/Engine/Source/Editor/DetailCustomizations/Public/DetailCustomizations.h
--- a/Engine/Source/Editor/DetailCustomizations/Public/DetailCustomizations.h
+++ b/Engine/Source/Editor/DetailCustomizations/Public/DetailCustomizations.h
@@ -46,8 +46,6 @@ private:
 	* @param StructLayoutDelegate	The delegate to call to get the custom detail layout instance
 	*/
 	void RegisterCustomPropertyTypeLayout(FName PropertyTypeName, FOnGetPropertyTypeCustomizationInstance PropertyTypeLayoutDelegate );
-
-	void RegisterStructPropertyLayout(FName StructPropertyName, FOnGetStructCustomizationInstance StructLayoutDelegate );
 private:
 	/** List of registered class that we must unregister when the module shuts down */
 	TSet< FName > RegisteredClassNames;
diff --git a/Engine/Source/Editor/EnvironmentQueryEditor/Private/DetailCustomizations/EnvDirectionCustomization.cpp b/Engine/Source/Editor/EnvironmentQueryEditor/Private/DetailCustomizations/EnvDirectionCustomization.cpp
--- a/Engine/Source/Editor/EnvironmentQueryEditor/Private/DetailCustomizations/EnvDirectionCustomization.cpp
+++ b/Engine/Source/Editor/EnvironmentQueryEditor/Private/DetailCustomizations/EnvDirectionCustomization.cpp
@@ -6,12 +6,12 @@
 
 #define LOCTEXT_NAMESPACE "FEnvQueryCustomization"
 
-TSharedRef<IStructCustomization> FEnvDirectionCustomization::MakeInstance( )
+TSharedRef<IPropertyTypeCustomization> FEnvDirectionCustomization::MakeInstance( )
 {
 	return MakeShareable(new FEnvDirectionCustomization);
 }
 
-void FEnvDirectionCustomization::CustomizeStructHeader( TSharedRef<class IPropertyHandle> StructPropertyHandle, class FDetailWidgetRow& HeaderRow, IStructCustomizationUtils& StructCustomizationUtils )
+void FEnvDirectionCustomization::CustomizeHeader( TSharedRef<class IPropertyHandle> StructPropertyHandle, class FDetailWidgetRow& HeaderRow, IStructCustomizationUtils& StructCustomizationUtils )
 {
 	// create struct header
 	HeaderRow.NameContent()
@@ -36,7 +36,7 @@ void FEnvDirectionCustomization::CustomizeStructHeader( TSharedRef<class IProper
 	OnModeChanged();
 }
 
-void FEnvDirectionCustomization::CustomizeStructChildren( TSharedRef<class IPropertyHandle> StructPropertyHandle, class IDetailChildrenBuilder& StructBuilder, IStructCustomizationUtils& StructCustomizationUtils )
+void FEnvDirectionCustomization::CustomizeChildren( TSharedRef<class IPropertyHandle> StructPropertyHandle, class IDetailChildrenBuilder& StructBuilder, IStructCustomizationUtils& StructCustomizationUtils )
 {
 	StructBuilder.AddChildProperty(ModeProp.ToSharedRef());
 
diff --git a/Engine/Source/Editor/EnvironmentQueryEditor/Private/DetailCustomizations/EnvDirectionCustomization.h b/Engine/Source/Editor/EnvironmentQueryEditor/Private/DetailCustomizations/EnvDirectionCustomization.h
--- a/Engine/Source/Editor/EnvironmentQueryEditor/Private/DetailCustomizations/EnvDirectionCustomization.h
+++ b/Engine/Source/Editor/EnvironmentQueryEditor/Private/DetailCustomizations/EnvDirectionCustomization.h
@@ -5,16 +5,16 @@
 #include "PropertyEditing.h"
 #include "PropertyCustomizationHelpers.h"
 
-class FEnvDirectionCustomization : public IStructCustomization
+class FEnvDirectionCustomization : public IPropertyTypeCustomization
 {
 public:
 
-	// Begin IStructCustomization interface
-	virtual void CustomizeStructHeader( TSharedRef<class IPropertyHandle> StructPropertyHandle, class FDetailWidgetRow& HeaderRow, IStructCustomizationUtils& StructCustomizationUtils ) OVERRIDE;
-	virtual void CustomizeStructChildren( TSharedRef<class IPropertyHandle> StructPropertyHandle, class IDetailChildrenBuilder& StructBuilder, IStructCustomizationUtils& StructCustomizationUtils ) OVERRIDE;
-	// End IStructCustomization interface
+	// Begin IPropertyTypeCustomization interface
+	virtual void CustomizeHeader( TSharedRef<class IPropertyHandle> StructPropertyHandle, class FDetailWidgetRow& HeaderRow, IStructCustomizationUtils& StructCustomizationUtils ) OVERRIDE;
+	virtual void CustomizeChildren( TSharedRef<class IPropertyHandle> StructPropertyHandle, class IDetailChildrenBuilder& StructBuilder, IStructCustomizationUtils& StructCustomizationUtils ) OVERRIDE;
+	// End IPropertyTypeCustomization interface
 
-	static TSharedRef<IStructCustomization> MakeInstance( );
+	static TSharedRef<IPropertyTypeCustomization> MakeInstance( );
 
 protected:
 
diff --git a/Engine/Source/Editor/EnvironmentQueryEditor/Private/DetailCustomizations/EnvQueryParamSetupCustomization.cpp b/Engine/Source/Editor/EnvironmentQueryEditor/Private/DetailCustomizations/EnvQueryParamSetupCustomization.cpp
--- a/Engine/Source/Editor/EnvironmentQueryEditor/Private/DetailCustomizations/EnvQueryParamSetupCustomization.cpp
+++ b/Engine/Source/Editor/EnvironmentQueryEditor/Private/DetailCustomizations/EnvQueryParamSetupCustomization.cpp
@@ -5,12 +5,12 @@
 
 #define LOCTEXT_NAMESPACE "FEnvQueryCustomization"
 
-TSharedRef<IStructCustomization> FEnvQueryParamSetupCustomization::MakeInstance( )
+TSharedRef<IPropertyTypeCustomization> FEnvQueryParamSetupCustomization::MakeInstance( )
 {
 	return MakeShareable(new FEnvQueryParamSetupCustomization);
 }
 
-void FEnvQueryParamSetupCustomization::CustomizeStructHeader( TSharedRef<class IPropertyHandle> StructPropertyHandle, class FDetailWidgetRow& HeaderRow, IStructCustomizationUtils& StructCustomizationUtils )
+void FEnvQueryParamSetupCustomization::CustomizeHeader( TSharedRef<class IPropertyHandle> StructPropertyHandle, class FDetailWidgetRow& HeaderRow, IStructCustomizationUtils& StructCustomizationUtils )
 {
 	PropertyHandle = StructPropertyHandle;
 	CacheMyValues();
@@ -69,7 +69,7 @@ void FEnvQueryParamSetupCustomization::CustomizeStructHeader( TSharedRef<class I
 	OnModeChanged();
 }
 
-void FEnvQueryParamSetupCustomization::CustomizeStructChildren( TSharedRef<class IPropertyHandle> StructPropertyHandle, class IDetailChildrenBuilder& StructBuilder, IStructCustomizationUtils& StructCustomizationUtils )
+void FEnvQueryParamSetupCustomization::CustomizeChildren( TSharedRef<class IPropertyHandle> StructPropertyHandle, class IDetailChildrenBuilder& StructBuilder, IStructCustomizationUtils& StructCustomizationUtils )
 {
 	// do nothing
 }
diff --git a/Engine/Source/Editor/EnvironmentQueryEditor/Private/DetailCustomizations/EnvQueryParamSetupCustomization.h b/Engine/Source/Editor/EnvironmentQueryEditor/Private/DetailCustomizations/EnvQueryParamSetupCustomization.h
--- a/Engine/Source/Editor/EnvironmentQueryEditor/Private/DetailCustomizations/EnvQueryParamSetupCustomization.h
+++ b/Engine/Source/Editor/EnvironmentQueryEditor/Private/DetailCustomizations/EnvQueryParamSetupCustomization.h
@@ -5,16 +5,16 @@
 #include "PropertyEditing.h"
 #include "PropertyCustomizationHelpers.h"
 
-class FEnvQueryParamSetupCustomization : public IStructCustomization
+class FEnvQueryParamSetupCustomization : public IPropertyTypeCustomization
 {
 public:
 
-	// Begin IStructCustomization interface
-	virtual void CustomizeStructHeader( TSharedRef<class IPropertyHandle> StructPropertyHandle, class FDetailWidgetRow& HeaderRow, IStructCustomizationUtils& StructCustomizationUtils ) OVERRIDE;
-	virtual void CustomizeStructChildren( TSharedRef<class IPropertyHandle> StructPropertyHandle, class IDetailChildrenBuilder& StructBuilder, IStructCustomizationUtils& StructCustomizationUtils ) OVERRIDE;
-	// End IStructCustomization interface
+	// Begin IPropertyTypeCustomization interface
+	virtual void CustomizeHeader( TSharedRef<class IPropertyHandle> StructPropertyHandle, class FDetailWidgetRow& HeaderRow, IStructCustomizationUtils& StructCustomizationUtils ) OVERRIDE;
+	virtual void CustomizeChildren( TSharedRef<class IPropertyHandle> StructPropertyHandle, class IDetailChildrenBuilder& StructBuilder, IStructCustomizationUtils& StructCustomizationUtils ) OVERRIDE;
+	// End IPropertyTypeCustomization interface
 
-	static TSharedRef<IStructCustomization> MakeInstance( );
+	static TSharedRef<IPropertyTypeCustomization> MakeInstance( );
 
 protected:
 
diff --git a/Engine/Source/Editor/EnvironmentQueryEditor/Private/DetailCustomizations/EnvTraceDataCustomization.cpp b/Engine/Source/Editor/EnvironmentQueryEditor/Private/DetailCustomizations/EnvTraceDataCustomization.cpp
--- a/Engine/Source/Editor/EnvironmentQueryEditor/Private/DetailCustomizations/EnvTraceDataCustomization.cpp
+++ b/Engine/Source/Editor/EnvironmentQueryEditor/Private/DetailCustomizations/EnvTraceDataCustomization.cpp
@@ -6,13 +6,13 @@
 
 #define LOCTEXT_NAMESPACE "EnvTraceDataCustomization"
 
-TSharedRef<IStructCustomization> FEnvTraceDataCustomization::MakeInstance()
+TSharedRef<IPropertyTypeCustomization> FEnvTraceDataCustomization::MakeInstance()
 {
 	return MakeShareable( new FEnvTraceDataCustomization );
 }
 
 BEGIN_SLATE_FUNCTION_BUILD_OPTIMIZATION
-void FEnvTraceDataCustomization::CustomizeStructHeader( TSharedRef<class IPropertyHandle> StructPropertyHandle, class FDetailWidgetRow& HeaderRow, IStructCustomizationUtils& StructCustomizationUtils )
+void FEnvTraceDataCustomization::CustomizeHeader( TSharedRef<class IPropertyHandle> StructPropertyHandle, class FDetailWidgetRow& HeaderRow, IStructCustomizationUtils& StructCustomizationUtils )
 {
 	HeaderRow.NameContent()
 	[
@@ -31,7 +31,7 @@ void FEnvTraceDataCustomization::CustomizeStructHeader( TSharedRef<class IProper
 	CacheTraceModes(StructPropertyHandle);
 }
 
-void FEnvTraceDataCustomization::CustomizeStructChildren( TSharedRef<class IPropertyHandle> StructPropertyHandle, class IDetailChildrenBuilder& StructBuilder, IStructCustomizationUtils& StructCustomizationUtils )
+void FEnvTraceDataCustomization::CustomizeChildren( TSharedRef<class IPropertyHandle> StructPropertyHandle, class IDetailChildrenBuilder& StructBuilder, IStructCustomizationUtils& StructCustomizationUtils )
 {
 	if (TraceModes.Num() > 1)
 	{
diff --git a/Engine/Source/Editor/EnvironmentQueryEditor/Private/DetailCustomizations/EnvTraceDataCustomization.h b/Engine/Source/Editor/EnvironmentQueryEditor/Private/DetailCustomizations/EnvTraceDataCustomization.h
--- a/Engine/Source/Editor/EnvironmentQueryEditor/Private/DetailCustomizations/EnvTraceDataCustomization.h
+++ b/Engine/Source/Editor/EnvironmentQueryEditor/Private/DetailCustomizations/EnvTraceDataCustomization.h
@@ -5,16 +5,16 @@
 #include "PropertyEditing.h"
 #include "PropertyCustomizationHelpers.h"
 
-class FEnvTraceDataCustomization : public IStructCustomization
+class FEnvTraceDataCustomization : public IPropertyTypeCustomization
 {
 public:
 	/** Makes a new instance of this detail layout class for a specific detail view requesting it */
-	static TSharedRef<IStructCustomization> MakeInstance();
+	static TSharedRef<IPropertyTypeCustomization> MakeInstance();
 
-	// Begin IStructCustomization interface
-	virtual void CustomizeStructHeader( TSharedRef<class IPropertyHandle> StructPropertyHandle, class FDetailWidgetRow& HeaderRow, IStructCustomizationUtils& StructCustomizationUtils ) OVERRIDE;
-	virtual void CustomizeStructChildren( TSharedRef<class IPropertyHandle> StructPropertyHandle, class IDetailChildrenBuilder& StructBuilder, IStructCustomizationUtils& StructCustomizationUtils ) OVERRIDE;
-	// End IStructCustomization interface
+	// Begin IPropertyTypeCustomization interface
+	virtual void CustomizeHeader( TSharedRef<class IPropertyHandle> StructPropertyHandle, class FDetailWidgetRow& HeaderRow, IStructCustomizationUtils& StructCustomizationUtils ) OVERRIDE;
+	virtual void CustomizeChildren( TSharedRef<class IPropertyHandle> StructPropertyHandle, class IDetailChildrenBuilder& StructBuilder, IStructCustomizationUtils& StructCustomizationUtils ) OVERRIDE;
+	// End IPropertyTypeCustomization interface
 
 protected:
 
diff --git a/Engine/Source/Editor/EnvironmentQueryEditor/Private/EnvironmentQueryEditorModule.cpp b/Engine/Source/Editor/EnvironmentQueryEditor/Private/EnvironmentQueryEditorModule.cpp
--- a/Engine/Source/Editor/EnvironmentQueryEditor/Private/EnvironmentQueryEditorModule.cpp
+++ b/Engine/Source/Editor/EnvironmentQueryEditor/Private/EnvironmentQueryEditorModule.cpp
@@ -46,11 +46,11 @@ void FEnvironmentQueryEditorModule::StartupModule()
 
 	// Register the details customizer
 	FPropertyEditorModule& PropertyModule = FModuleManager::LoadModuleChecked<FPropertyEditorModule>("PropertyEditor");
-	PropertyModule.RegisterStructPropertyLayout( "EnvFloatParam", FOnGetStructCustomizationInstance::CreateStatic( &FEnvQueryParamSetupCustomization::MakeInstance ) );
-	PropertyModule.RegisterStructPropertyLayout( "EnvIntParam", FOnGetStructCustomizationInstance::CreateStatic( &FEnvQueryParamSetupCustomization::MakeInstance ) );
-	PropertyModule.RegisterStructPropertyLayout( "EnvBoolParam", FOnGetStructCustomizationInstance::CreateStatic( &FEnvQueryParamSetupCustomization::MakeInstance ) );
-	PropertyModule.RegisterStructPropertyLayout( "EnvDirection", FOnGetStructCustomizationInstance::CreateStatic( &FEnvDirectionCustomization::MakeInstance ) );
-	PropertyModule.RegisterStructPropertyLayout( "EnvTraceData", FOnGetStructCustomizationInstance::CreateStatic( &FEnvTraceDataCustomization::MakeInstance ) );
+	PropertyModule.RegisterCustomPropertyTypeLayout( "EnvFloatParam", FOnGetPropertyTypeCustomizationInstance::CreateStatic( &FEnvQueryParamSetupCustomization::MakeInstance ) );
+	PropertyModule.RegisterCustomPropertyTypeLayout( "EnvIntParam", FOnGetPropertyTypeCustomizationInstance::CreateStatic( &FEnvQueryParamSetupCustomization::MakeInstance ) );
+	PropertyModule.RegisterCustomPropertyTypeLayout( "EnvBoolParam", FOnGetPropertyTypeCustomizationInstance::CreateStatic( &FEnvQueryParamSetupCustomization::MakeInstance ) );
+	PropertyModule.RegisterCustomPropertyTypeLayout( "EnvDirection", FOnGetPropertyTypeCustomizationInstance::CreateStatic( &FEnvDirectionCustomization::MakeInstance ) );
+	PropertyModule.RegisterCustomPropertyTypeLayout( "EnvTraceData", FOnGetPropertyTypeCustomizationInstance::CreateStatic( &FEnvTraceDataCustomization::MakeInstance ) );
 	PropertyModule.RegisterCustomPropertyLayout( "EnvQueryTest", FOnGetDetailCustomizationInstance::CreateStatic( &FEnvQueryTestDetails::MakeInstance ) );
 	PropertyModule.NotifyCustomizationModuleChanged();
 }
diff --git a/Engine/Source/Editor/GameplayTagsEditor/Private/GameplayTagContainerCustomization.cpp b/Engine/Source/Editor/GameplayTagsEditor/Private/GameplayTagContainerCustomization.cpp
--- a/Engine/Source/Editor/GameplayTagsEditor/Private/GameplayTagContainerCustomization.cpp
+++ b/Engine/Source/Editor/GameplayTagsEditor/Private/GameplayTagContainerCustomization.cpp
@@ -10,7 +10,7 @@
 
 #define LOCTEXT_NAMESPACE "GameplayTagContainerCustomization"
 
-void FGameplayTagContainerCustomization::CustomizeStructHeader(TSharedRef<class IPropertyHandle> InStructPropertyHandle, class FDetailWidgetRow& HeaderRow, IStructCustomizationUtils& StructCustomizationUtils)
+void FGameplayTagContainerCustomization::CustomizeHeader(TSharedRef<class IPropertyHandle> InStructPropertyHandle, class FDetailWidgetRow& HeaderRow, IStructCustomizationUtils& StructCustomizationUtils)
 {
 	StructPropertyHandle = InStructPropertyHandle;
 
diff --git a/Engine/Source/Editor/GameplayTagsEditor/Private/GameplayTagContainerCustomization.h b/Engine/Source/Editor/GameplayTagsEditor/Private/GameplayTagContainerCustomization.h
--- a/Engine/Source/Editor/GameplayTagsEditor/Private/GameplayTagContainerCustomization.h
+++ b/Engine/Source/Editor/GameplayTagsEditor/Private/GameplayTagContainerCustomization.h
@@ -6,10 +6,10 @@
 #include "SGameplayTagWidget.h"
 
 /** Customization for the gameplay tag container struct */
-class FGameplayTagContainerCustomization : public IStructCustomization, public FEditorUndoClient
+class FGameplayTagContainerCustomization : public IPropertyTypeCustomization, public FEditorUndoClient
 {
 public:
-	static TSharedRef<IStructCustomization> MakeInstance()
+	static TSharedRef<IPropertyTypeCustomization> MakeInstance()
 	{
 		return MakeShareable(new FGameplayTagContainerCustomization);
 	}
@@ -17,10 +17,10 @@ public:
 	~FGameplayTagContainerCustomization();
 
 	/** Overridden to show an edit button to launch the gameplay tag editor */
-	virtual void CustomizeStructHeader(TSharedRef<class IPropertyHandle> StructPropertyHandle, class FDetailWidgetRow& HeaderRow, IStructCustomizationUtils& StructCustomizationUtils) OVERRIDE;
+	virtual void CustomizeHeader(TSharedRef<class IPropertyHandle> StructPropertyHandle, class FDetailWidgetRow& HeaderRow, IStructCustomizationUtils& StructCustomizationUtils) OVERRIDE;
 	
 	/** Overridden to do nothing */
-	virtual void CustomizeStructChildren(TSharedRef<IPropertyHandle> InStructPropertyHandle, class IDetailChildrenBuilder& ChildBuilder, IStructCustomizationUtils& StructCustomizationUtils) OVERRIDE {}
+	virtual void CustomizeChildren(TSharedRef<IPropertyHandle> InStructPropertyHandle, class IDetailChildrenBuilder& ChildBuilder, IStructCustomizationUtils& StructCustomizationUtils) OVERRIDE {}
 
 	// Begin FEditorUndoClient Interface
 	virtual void PostUndo( bool bSuccess ) OVERRIDE;
diff --git a/Engine/Source/Editor/GameplayTagsEditor/Private/GameplayTagsEditorModule.cpp b/Engine/Source/Editor/GameplayTagsEditor/Private/GameplayTagsEditorModule.cpp
--- a/Engine/Source/Editor/GameplayTagsEditor/Private/GameplayTagsEditorModule.cpp
+++ b/Engine/Source/Editor/GameplayTagsEditor/Private/GameplayTagsEditorModule.cpp
@@ -20,7 +20,7 @@ void FGameplayTagsEditorModule::StartupModule()
 {
 	// Register the details customizer
 	FPropertyEditorModule& PropertyModule = FModuleManager::LoadModuleChecked<FPropertyEditorModule>("PropertyEditor");
-	PropertyModule.RegisterStructPropertyLayout("GameplayTagContainer", FOnGetStructCustomizationInstance::CreateStatic(&FGameplayTagContainerCustomization::MakeInstance));
+	PropertyModule.RegisterCustomPropertyTypeLayout("GameplayTagContainer", FOnGetPropertyTypeCustomizationInstance::CreateStatic(&FGameplayTagContainerCustomization::MakeInstance));
 
 	TSharedPtr<FGameplayTagsGraphPanelPinFactory> GameplayTagsGraphPanelPinFactory = MakeShareable( new FGameplayTagsGraphPanelPinFactory() );
 	FEdGraphUtilities::RegisterVisualPinFactory(GameplayTagsGraphPanelPinFactory);
diff --git a/Engine/Source/Editor/LandscapeEditor/Private/LandscapeEditorDetailCustomization_Base.h b/Engine/Source/Editor/LandscapeEditor/Private/LandscapeEditorDetailCustomization_Base.h
--- a/Engine/Source/Editor/LandscapeEditor/Private/LandscapeEditorDetailCustomization_Base.h
+++ b/Engine/Source/Editor/LandscapeEditor/Private/LandscapeEditorDetailCustomization_Base.h
@@ -94,7 +94,7 @@ void FLandscapeEditorDetailCustomization_Base::SetPropertyValue(type NewValue, E
 	ensure(PropertyHandle->SetValue(NewValue) == FPropertyAccess::Success);
 }
 
-class FLandscapeEditorStructCustomization_Base : public IStructCustomization
+class FLandscapeEditorStructCustomization_Base : public IPropertyTypeCustomization
 {
 protected:
 	static class FEdModeLandscape* GetEditorMode();
diff --git a/Engine/Source/Editor/LandscapeEditor/Private/LandscapeEditorDetailCustomization_CopyPaste.cpp b/Engine/Source/Editor/LandscapeEditor/Private/LandscapeEditorDetailCustomization_CopyPaste.cpp
--- a/Engine/Source/Editor/LandscapeEditor/Private/LandscapeEditorDetailCustomization_CopyPaste.cpp
+++ b/Engine/Source/Editor/LandscapeEditor/Private/LandscapeEditorDetailCustomization_CopyPaste.cpp
@@ -457,16 +457,16 @@ FReply FLandscapeEditorDetailCustomization_CopyPaste::OnGizmoExportButtonClicked
 
 //////////////////////////////////////////////////////////////////////////
 
-TSharedRef<IStructCustomization> FLandscapeEditorStructCustomization_FGizmoImportLayer::MakeInstance()
+TSharedRef<IPropertyTypeCustomization> FLandscapeEditorStructCustomization_FGizmoImportLayer::MakeInstance()
 {
 	return MakeShareable(new FLandscapeEditorStructCustomization_FGizmoImportLayer);
 }
 
-void FLandscapeEditorStructCustomization_FGizmoImportLayer::CustomizeStructHeader(TSharedRef<IPropertyHandle> StructPropertyHandle, class FDetailWidgetRow& HeaderRow, IStructCustomizationUtils& StructCustomizationUtils)
+void FLandscapeEditorStructCustomization_FGizmoImportLayer::CustomizeHeader(TSharedRef<IPropertyHandle> StructPropertyHandle, class FDetailWidgetRow& HeaderRow, IStructCustomizationUtils& StructCustomizationUtils)
 {
 }
 
-void FLandscapeEditorStructCustomization_FGizmoImportLayer::CustomizeStructChildren(TSharedRef<IPropertyHandle> StructPropertyHandle, class IDetailChildrenBuilder& ChildBuilder, IStructCustomizationUtils& StructCustomizationUtils)
+void FLandscapeEditorStructCustomization_FGizmoImportLayer::CustomizeChildren(TSharedRef<IPropertyHandle> StructPropertyHandle, class IDetailChildrenBuilder& ChildBuilder, IStructCustomizationUtils& StructCustomizationUtils)
 {
 	TSharedRef<IPropertyHandle> PropertyHandle_LayerFilename = StructPropertyHandle->GetChildHandle(GET_MEMBER_NAME_CHECKED(FGizmoImportLayer, LayerFilename)).ToSharedRef();
 	ChildBuilder.AddChildProperty(PropertyHandle_LayerFilename)
diff --git a/Engine/Source/Editor/LandscapeEditor/Private/LandscapeEditorDetailCustomization_CopyPaste.h b/Engine/Source/Editor/LandscapeEditor/Private/LandscapeEditorDetailCustomization_CopyPaste.h
--- a/Engine/Source/Editor/LandscapeEditor/Private/LandscapeEditorDetailCustomization_CopyPaste.h
+++ b/Engine/Source/Editor/LandscapeEditor/Private/LandscapeEditorDetailCustomization_CopyPaste.h
@@ -33,10 +33,10 @@ public:
 class FLandscapeEditorStructCustomization_FGizmoImportLayer : public FLandscapeEditorStructCustomization_Base
 {
 public:
-	static TSharedRef<IStructCustomization> MakeInstance();
+	static TSharedRef<IPropertyTypeCustomization> MakeInstance();
 
-	virtual void CustomizeStructHeader(TSharedRef<IPropertyHandle> StructPropertyHandle, class FDetailWidgetRow& HeaderRow, IStructCustomizationUtils& StructCustomizationUtils) OVERRIDE;
-	virtual void CustomizeStructChildren(TSharedRef<IPropertyHandle> StructPropertyHandle, class IDetailChildrenBuilder& ChildBuilder, IStructCustomizationUtils& StructCustomizationUtils) OVERRIDE;
+	virtual void CustomizeHeader(TSharedRef<IPropertyHandle> StructPropertyHandle, class FDetailWidgetRow& HeaderRow, IStructCustomizationUtils& StructCustomizationUtils) OVERRIDE;
+	virtual void CustomizeChildren(TSharedRef<IPropertyHandle> StructPropertyHandle, class IDetailChildrenBuilder& ChildBuilder, IStructCustomizationUtils& StructCustomizationUtils) OVERRIDE;
 
 public:
 	static FReply OnGizmoImportLayerFilenameButtonClicked(TSharedRef<IPropertyHandle> PropertyHandle_LayerFilename);
diff --git a/Engine/Source/Editor/LandscapeEditor/Private/LandscapeEditorDetailCustomization_NewLandscape.cpp b/Engine/Source/Editor/LandscapeEditor/Private/LandscapeEditorDetailCustomization_NewLandscape.cpp
--- a/Engine/Source/Editor/LandscapeEditor/Private/LandscapeEditorDetailCustomization_NewLandscape.cpp
+++ b/Engine/Source/Editor/LandscapeEditor/Private/LandscapeEditorDetailCustomization_NewLandscape.cpp
@@ -1307,16 +1307,16 @@ EVisibility FLandscapeEditorDetailCustomization_NewLandscape::GetMaterialTipVisi
 
 //////////////////////////////////////////////////////////////////////////
 
-TSharedRef<IStructCustomization> FLandscapeEditorStructCustomization_FLandscapeImportLayer::MakeInstance()
+TSharedRef<IPropertyTypeCustomization> FLandscapeEditorStructCustomization_FLandscapeImportLayer::MakeInstance()
 {
 	return MakeShareable(new FLandscapeEditorStructCustomization_FLandscapeImportLayer);
 }
 
-void FLandscapeEditorStructCustomization_FLandscapeImportLayer::CustomizeStructHeader(TSharedRef<IPropertyHandle> StructPropertyHandle, class FDetailWidgetRow& HeaderRow, IStructCustomizationUtils& StructCustomizationUtils)
+void FLandscapeEditorStructCustomization_FLandscapeImportLayer::CustomizeHeader(TSharedRef<IPropertyHandle> StructPropertyHandle, class FDetailWidgetRow& HeaderRow, IStructCustomizationUtils& StructCustomizationUtils)
 {
 }
 
-void FLandscapeEditorStructCustomization_FLandscapeImportLayer::CustomizeStructChildren(TSharedRef<IPropertyHandle> StructPropertyHandle, class IDetailChildrenBuilder& ChildBuilder, IStructCustomizationUtils& StructCustomizationUtils)
+void FLandscapeEditorStructCustomization_FLandscapeImportLayer::CustomizeChildren(TSharedRef<IPropertyHandle> StructPropertyHandle, class IDetailChildrenBuilder& ChildBuilder, IStructCustomizationUtils& StructCustomizationUtils)
 {
 	TSharedRef<IPropertyHandle> PropertyHandle_LayerName = StructPropertyHandle->GetChildHandle(GET_MEMBER_NAME_CHECKED(FLandscapeImportLayer, LayerName)).ToSharedRef();
 	TSharedRef<IPropertyHandle> PropertyHandle_LayerInfo = StructPropertyHandle->GetChildHandle(GET_MEMBER_NAME_CHECKED(FLandscapeImportLayer, LayerInfo)).ToSharedRef();
diff --git a/Engine/Source/Editor/LandscapeEditor/Private/LandscapeEditorDetailCustomization_NewLandscape.h b/Engine/Source/Editor/LandscapeEditor/Private/LandscapeEditorDetailCustomization_NewLandscape.h
--- a/Engine/Source/Editor/LandscapeEditor/Private/LandscapeEditorDetailCustomization_NewLandscape.h
+++ b/Engine/Source/Editor/LandscapeEditor/Private/LandscapeEditorDetailCustomization_NewLandscape.h
@@ -84,10 +84,10 @@ protected:
 class FLandscapeEditorStructCustomization_FLandscapeImportLayer : public FLandscapeEditorStructCustomization_Base
 {
 public:
-	static TSharedRef<IStructCustomization> MakeInstance();
+	static TSharedRef<IPropertyTypeCustomization> MakeInstance();
 
-	virtual void CustomizeStructHeader(TSharedRef<IPropertyHandle> StructPropertyHandle, class FDetailWidgetRow& HeaderRow, IStructCustomizationUtils& StructCustomizationUtils) OVERRIDE;
-	virtual void CustomizeStructChildren(TSharedRef<IPropertyHandle> StructPropertyHandle, class IDetailChildrenBuilder& ChildBuilder, IStructCustomizationUtils& StructCustomizationUtils) OVERRIDE;
+	virtual void CustomizeHeader(TSharedRef<IPropertyHandle> StructPropertyHandle, class FDetailWidgetRow& HeaderRow, IStructCustomizationUtils& StructCustomizationUtils) OVERRIDE;
+	virtual void CustomizeChildren(TSharedRef<IPropertyHandle> StructPropertyHandle, class IDetailChildrenBuilder& ChildBuilder, IStructCustomizationUtils& StructCustomizationUtils) OVERRIDE;
 
 public:
 	static FReply OnLayerFilenameButtonClicked(TSharedRef<IPropertyHandle> PropertyHandle_LayerFilename);
diff --git a/Engine/Source/Editor/LandscapeEditor/Private/LandscapeEditorModule.cpp b/Engine/Source/Editor/LandscapeEditor/Private/LandscapeEditorModule.cpp
--- a/Engine/Source/Editor/LandscapeEditor/Private/LandscapeEditorModule.cpp
+++ b/Engine/Source/Editor/LandscapeEditor/Private/LandscapeEditorModule.cpp
@@ -38,8 +38,8 @@ public:
 
 		FPropertyEditorModule& PropertyModule = FModuleManager::GetModuleChecked<FPropertyEditorModule>("PropertyEditor");
 		PropertyModule.RegisterCustomPropertyLayout("LandscapeEditorObject", FOnGetDetailCustomizationInstance::CreateStatic(&FLandscapeEditorDetails::MakeInstance));
-		PropertyModule.RegisterStructPropertyLayout("GizmoImportLayer", FOnGetStructCustomizationInstance::CreateStatic(&FLandscapeEditorStructCustomization_FGizmoImportLayer::MakeInstance));
-		PropertyModule.RegisterStructPropertyLayout("LandscapeImportLayer", FOnGetStructCustomizationInstance::CreateStatic(&FLandscapeEditorStructCustomization_FLandscapeImportLayer::MakeInstance));
+		PropertyModule.RegisterCustomPropertyTypeLayout("GizmoImportLayer", FOnGetPropertyTypeCustomizationInstance::CreateStatic(&FLandscapeEditorStructCustomization_FGizmoImportLayer::MakeInstance));
+		PropertyModule.RegisterCustomPropertyTypeLayout("LandscapeImportLayer", FOnGetPropertyTypeCustomizationInstance::CreateStatic(&FLandscapeEditorStructCustomization_FLandscapeImportLayer::MakeInstance));
 
 		TSharedRef<FUICommandList> CommandList = MakeShareable(new FUICommandList);
 		const FLandscapeEditorCommands& LandscapeActions = FLandscapeEditorCommands::Get();
diff --git a/Engine/Source/Editor/Persona/Private/AnimGraphNodeDetails.cpp b/Engine/Source/Editor/Persona/Private/AnimGraphNodeDetails.cpp
--- a/Engine/Source/Editor/Persona/Private/AnimGraphNodeDetails.cpp
+++ b/Engine/Source/Editor/Persona/Private/AnimGraphNodeDetails.cpp
@@ -258,13 +258,13 @@ void FAnimGraphNodeDetails::AbortDisplayOfAllNodes(TArray< TWeakObjectPtr<UObjec
 }
 
 
-TSharedRef<IStructCustomization> FInputScaleBiasCustomization::MakeInstance() 
+TSharedRef<IPropertyTypeCustomization> FInputScaleBiasCustomization::MakeInstance() 
 {
 	return MakeShareable(new FInputScaleBiasCustomization());
 }
 
 BEGIN_SLATE_FUNCTION_BUILD_OPTIMIZATION
-void FInputScaleBiasCustomization::CustomizeStructHeader(TSharedRef<class IPropertyHandle> StructPropertyHandle, class FDetailWidgetRow& HeaderRow, IStructCustomizationUtils& StructCustomizationUtils)
+void FInputScaleBiasCustomization::CustomizeHeader(TSharedRef<class IPropertyHandle> StructPropertyHandle, class FDetailWidgetRow& HeaderRow, IStructCustomizationUtils& StructCustomizationUtils)
 {
 	HeaderRow.NameContent()
 	[
@@ -335,7 +335,7 @@ void FInputScaleBiasCustomization::CustomizeStructHeader(TSharedRef<class IPrope
 END_SLATE_FUNCTION_BUILD_OPTIMIZATION
 
 /////////////////////////////////////////////////////
-void FInputScaleBiasCustomization::CustomizeStructChildren( TSharedRef<class IPropertyHandle> StructPropertyHandle, class IDetailChildrenBuilder& StructBuilder, IStructCustomizationUtils& StructCustomizationUtils )
+void FInputScaleBiasCustomization::CustomizeChildren( TSharedRef<class IPropertyHandle> StructPropertyHandle, class IDetailChildrenBuilder& StructBuilder, IStructCustomizationUtils& StructCustomizationUtils )
 {
 	// nothing here
 }
@@ -401,13 +401,13 @@ void FInputScaleBiasCustomization::OnMaxValueCommitted(const FText& NewText, ETe
 
 /////////////////////////////////////////////////////
 
-TSharedRef<IStructCustomization> FBoneReferenceCustomization::MakeInstance()
+TSharedRef<IPropertyTypeCustomization> FBoneReferenceCustomization::MakeInstance()
 {
 	return MakeShareable(new FBoneReferenceCustomization());
 }
 
 BEGIN_SLATE_FUNCTION_BUILD_OPTIMIZATION
-void FBoneReferenceCustomization::CustomizeStructHeader( TSharedRef<IPropertyHandle> StructPropertyHandle, class FDetailWidgetRow& HeaderRow, IStructCustomizationUtils& StructCustomizationUtils )
+void FBoneReferenceCustomization::CustomizeHeader( TSharedRef<IPropertyHandle> StructPropertyHandle, class FDetailWidgetRow& HeaderRow, IStructCustomizationUtils& StructCustomizationUtils )
 {
 	uint32 NumChildren = 0;
 	StructPropertyHandle->GetNumChildren(NumChildren);
@@ -464,7 +464,7 @@ void FBoneReferenceCustomization::CustomizeStructHeader( TSharedRef<IPropertyHan
 }
 END_SLATE_FUNCTION_BUILD_OPTIMIZATION
 
-void FBoneReferenceCustomization::CustomizeStructChildren( TSharedRef<IPropertyHandle> StructPropertyHandle, class IDetailChildrenBuilder& ChildBuilder, IStructCustomizationUtils& StructCustomizationUtils )
+void FBoneReferenceCustomization::CustomizeChildren( TSharedRef<IPropertyHandle> StructPropertyHandle, class IDetailChildrenBuilder& ChildBuilder, IStructCustomizationUtils& StructCustomizationUtils )
 {
 	// No child customisations as the properties are shown in the header
 }
diff --git a/Engine/Source/Editor/Persona/Private/AnimGraphNodeDetails.h b/Engine/Source/Editor/Persona/Private/AnimGraphNodeDetails.h
--- a/Engine/Source/Editor/Persona/Private/AnimGraphNodeDetails.h
+++ b/Engine/Source/Editor/Persona/Private/AnimGraphNodeDetails.h
@@ -36,15 +36,15 @@ protected:
 /////////////////////////////////////////////////////
 // FInputScaleBiasCustomization
 
-class FInputScaleBiasCustomization : public IStructCustomization
+class FInputScaleBiasCustomization : public IPropertyTypeCustomization
 {
 public:
-	static TSharedRef<IStructCustomization> MakeInstance();
+	static TSharedRef<IPropertyTypeCustomization> MakeInstance();
 
-	// IStructCustomization interface
-	virtual void CustomizeStructHeader(TSharedRef<IPropertyHandle> StructPropertyHandle, class FDetailWidgetRow& HeaderRow, IStructCustomizationUtils& StructCustomizationUtils) OVERRIDE;
-	virtual void CustomizeStructChildren(TSharedRef<IPropertyHandle> StructPropertyHandle, class IDetailChildrenBuilder& ChildBuilder, IStructCustomizationUtils& StructCustomizationUtils) OVERRIDE;
-	// End of IStructCustomization interface
+	// IPropertyTypeCustomization interface
+	virtual void CustomizeHeader(TSharedRef<IPropertyHandle> StructPropertyHandle, class FDetailWidgetRow& HeaderRow, IStructCustomizationUtils& StructCustomizationUtils) OVERRIDE;
+	virtual void CustomizeChildren(TSharedRef<IPropertyHandle> StructPropertyHandle, class IDetailChildrenBuilder& ChildBuilder, IStructCustomizationUtils& StructCustomizationUtils) OVERRIDE;
+	// End of IPropertyTypeCustomization interface
 
 	FText GetMinValue(TSharedRef<class IPropertyHandle> StructPropertyHandle) const;
 	FText GetMaxValue(TSharedRef<class IPropertyHandle> StructPropertyHandle) const;
@@ -55,14 +55,14 @@ public:
 //////////////////////////////////////////////////////////////////////////
 // FBoneReferenceCustomization
 
-class FBoneReferenceCustomization : public IStructCustomization
+class FBoneReferenceCustomization : public IPropertyTypeCustomization
 {
 public:
-	static TSharedRef<IStructCustomization> MakeInstance();
+	static TSharedRef<IPropertyTypeCustomization> MakeInstance();
 
-	// IStructCustomization interface
-	virtual void CustomizeStructHeader(TSharedRef<IPropertyHandle> StructPropertyHandle, class FDetailWidgetRow& HeaderRow, IStructCustomizationUtils& StructCustomizationUtils) OVERRIDE;
-	virtual void CustomizeStructChildren(TSharedRef<IPropertyHandle> StructPropertyHandle, class IDetailChildrenBuilder& ChildBuilder, IStructCustomizationUtils& StructCustomizationUtils) OVERRIDE;
+	// IPropertyTypeCustomization interface
+	virtual void CustomizeHeader(TSharedRef<IPropertyHandle> StructPropertyHandle, class FDetailWidgetRow& HeaderRow, IStructCustomizationUtils& StructCustomizationUtils) OVERRIDE;
+	virtual void CustomizeChildren(TSharedRef<IPropertyHandle> StructPropertyHandle, class IDetailChildrenBuilder& ChildBuilder, IStructCustomizationUtils& StructCustomizationUtils) OVERRIDE;
 
 private:
 	// Storage object for bone hierarchy
diff --git a/Engine/Source/Editor/Persona/Private/PersonaModule.cpp b/Engine/Source/Editor/Persona/Private/PersonaModule.cpp
--- a/Engine/Source/Editor/Persona/Private/PersonaModule.cpp
+++ b/Engine/Source/Editor/Persona/Private/PersonaModule.cpp
@@ -60,8 +60,8 @@ void FPersonaModule::StartupModule()
 		PropertyModule.RegisterCustomPropertyLayout( "AnimNotifyState", FOnGetDetailCustomizationInstance::CreateStatic(&FAnimNotifyDetails::MakeInstance));
 		PropertyModule.RegisterCustomPropertyLayout( "AnimGraphNode_Base", FOnGetDetailCustomizationInstance::CreateStatic( &FAnimGraphNodeDetails::MakeInstance ) );
 
-		PropertyModule.RegisterStructPropertyLayout( "InputScaleBias", FOnGetStructCustomizationInstance::CreateStatic( &FInputScaleBiasCustomization::MakeInstance ) );
-		PropertyModule.RegisterStructPropertyLayout( "BoneReference", FOnGetStructCustomizationInstance::CreateStatic( &FBoneReferenceCustomization::MakeInstance ) );
+		PropertyModule.RegisterCustomPropertyTypeLayout( "InputScaleBias", FOnGetPropertyTypeCustomizationInstance::CreateStatic( &FInputScaleBiasCustomization::MakeInstance ) );
+		PropertyModule.RegisterCustomPropertyTypeLayout( "BoneReference", FOnGetPropertyTypeCustomizationInstance::CreateStatic( &FBoneReferenceCustomization::MakeInstance ) );
 	}
 }
 
diff --git a/Engine/Source/Editor/PropertyEditor/Private/CustomChildBuilder.cpp b/Engine/Source/Editor/PropertyEditor/Private/CustomChildBuilder.cpp
--- a/Engine/Source/Editor/PropertyEditor/Private/CustomChildBuilder.cpp
+++ b/Engine/Source/Editor/PropertyEditor/Private/CustomChildBuilder.cpp
@@ -91,7 +91,7 @@ TSharedRef<SWidget> FCustomChildrenBuilder::GenerateStructValueWidget( TSharedRe
 	UStructProperty* StructProperty = CastChecked<UStructProperty>( StructPropertyHandle->GetProperty() );
 
 	FPropertyEditorModule& PropertyEditorModule = FModuleManager::GetModuleChecked<FPropertyEditorModule>("PropertyEditor");
-	FPropertyTypeLayoutCallback LayoutCallback = PropertyEditorModule.GetPropertyTypeCustomization(StructProperty);
+	FPropertyTypeLayoutCallback LayoutCallback = PropertyEditorModule.GetPropertyTypeCustomization(StructProperty, *StructPropertyHandle);
 	if (LayoutCallback.IsValid())
 	{
 		TSharedRef<IPropertyTypeCustomization> CustomStructInterface = LayoutCallback.GetCustomizationInstance();
diff --git a/Engine/Source/Editor/PropertyEditor/Private/DetailPropertyRow.cpp b/Engine/Source/Editor/PropertyEditor/Private/DetailPropertyRow.cpp
--- a/Engine/Source/Editor/PropertyEditor/Private/DetailPropertyRow.cpp
+++ b/Engine/Source/Editor/PropertyEditor/Private/DetailPropertyRow.cpp
@@ -35,7 +35,7 @@ FDetailPropertyRow::FDetailPropertyRow(TSharedPtr<FPropertyNode> InPropertyNode,
 		{
 			static FName PropertyEditor("PropertyEditor");
 			FPropertyEditorModule& PropertyEditorModule = FModuleManager::GetModuleChecked<FPropertyEditorModule>(PropertyEditor);
-			FPropertyTypeLayoutCallback LayoutCallback = PropertyEditorModule.GetPropertyTypeCustomization(Property);
+			FPropertyTypeLayoutCallback LayoutCallback = PropertyEditorModule.GetPropertyTypeCustomization(Property,*PropertyHandle);
 			if (LayoutCallback.IsValid())
 			{
 				if (PropertyHandle->IsValidHandle())
diff --git a/Engine/Source/Editor/PropertyEditor/Private/PropertyEditorModule.cpp b/Engine/Source/Editor/PropertyEditor/Private/PropertyEditorModule.cpp
--- a/Engine/Source/Editor/PropertyEditor/Private/PropertyEditorModule.cpp
+++ b/Engine/Source/Editor/PropertyEditor/Private/PropertyEditorModule.cpp
@@ -49,16 +49,16 @@ void FPropertyTypeLayoutCallbackList::Remove( const TSharedPtr<IPropertyTypeIden
 	}
 }
 
-const FPropertyTypeLayoutCallback& FPropertyTypeLayoutCallbackList::Find( const UProperty& Property )
+const FPropertyTypeLayoutCallback& FPropertyTypeLayoutCallbackList::Find( const IPropertyHandle& PropertyHandle )
 {
 	if( IdentifierList.Num() > 0 )
 	{
 		FPropertyTypeLayoutCallback* Callback =
 			IdentifierList.FindByPredicate
 			(
-				[&Property]( const FPropertyTypeLayoutCallback& Callback )
+				[&]( const FPropertyTypeLayoutCallback& Callback )
 				{
-					return Callback.PropertyTypeIdentifier->IsPropertyTypeCustomized( Property );
+					return Callback.PropertyTypeIdentifier->IsPropertyTypeCustomized( PropertyHandle );
 				}
 			);
 
@@ -571,7 +571,7 @@ bool FPropertyEditorModule::IsCustomizedStruct(const UStruct* Struct) const
 	return false;
 }
 
-FPropertyTypeLayoutCallback FPropertyEditorModule::GetPropertyTypeCustomization(const UProperty* Property)
+FPropertyTypeLayoutCallback FPropertyEditorModule::GetPropertyTypeCustomization(const UProperty* Property, const IPropertyHandle& PropertyHandle)
 {
 	if( Property )
 	{
@@ -610,7 +610,7 @@ FPropertyTypeLayoutCallback FPropertyEditorModule::GetPropertyTypeCustomization(
 			FPropertyTypeLayoutCallbackList* LayoutCallbacks = PropertyTypeToLayoutMap.Find(PropertyTypeName);
 			if (LayoutCallbacks)
 			{
-				const FPropertyTypeLayoutCallback& Callback = LayoutCallbacks->Find(*Property);
+				const FPropertyTypeLayoutCallback& Callback = LayoutCallbacks->Find(PropertyHandle);
 				return Callback;
 			}
 		}
diff --git a/Engine/Source/Editor/PropertyEditor/Public/PropertyEditorModule.h b/Engine/Source/Editor/PropertyEditor/Public/PropertyEditorModule.h
--- a/Engine/Source/Editor/PropertyEditor/Public/PropertyEditorModule.h
+++ b/Engine/Source/Editor/PropertyEditor/Public/PropertyEditorModule.h
@@ -212,11 +212,14 @@ class IPropertyTypeIdentifier
 public:
 	virtual ~IPropertyTypeIdentifier() {}
 
-	virtual bool IsPropertyTypeCustomized( const UProperty& InProperty ) const = 0;
+	/**
+	 * Called to identify if a property type is customized
+	 * 
+	 * @param IPropertyHandle	Handle to the property being tested
+	 */
+	virtual bool IsPropertyTypeCustomized( const IPropertyHandle& PropertyHandle ) const = 0;
 };
 
-
-
 /**
  * Callback executed to query the custom layout of details
  */
@@ -240,7 +243,7 @@ struct FPropertyTypeLayoutCallback
 
 	TSharedRef<IPropertyTypeCustomization> GetCustomizationInstance() const
 	{
-		return PropertyTypeLayoutDelegate.IsBound() ? PropertyTypeLayoutDelegate.Execute() : static_cast<TSharedRef<IPropertyTypeCustomization>>(DeprecatedLayoutDelegate.Execute());
+		return PropertyTypeLayoutDelegate.IsBound() ? PropertyTypeLayoutDelegate.Execute() : StaticCastSharedRef<IPropertyTypeCustomization>(DeprecatedLayoutDelegate.Execute());
 	}
 };
 
@@ -257,7 +260,7 @@ struct FPropertyTypeLayoutCallbackList
 	
 	void Remove( const TSharedPtr<IPropertyTypeIdentifier>& InIdentifier );
 
-	const FPropertyTypeLayoutCallback& Find( const UProperty& Property );
+	const FPropertyTypeLayoutCallback& Find( const IPropertyHandle& PropertyHandle );
 };
 
 typedef TMap< TWeakObjectPtr<UStruct>, FDetailLayoutCallback > FCustomDetailLayoutMap;
@@ -347,6 +350,7 @@ public:
 	 * @param Class	The class the custom detail layout is for
 	 * @param StructLayoutDelegate	The delegate to call when querying for custom detail layouts for the struct properties
 	 */
+	DEPRECATED(4.3, "RegisterCustomPropertyTypeLayout is deprecated.  Use RegisterCustomPropertyTypeLayout instead.")
 	virtual void RegisterStructPropertyLayout( FName StructTypeName, FOnGetStructCustomizationInstance StructLayoutDelegate );
 
 	/**
@@ -354,6 +358,7 @@ public:
 	 *
 	 * @param StructTypeName	The structure with the custom detail layout delegate to remove
 	 */
+	DEPRECATED(4.3, "UnregisterStructPropertyLayout is deprecated.  Use UnregisterCustomPropertyTypeLayout instead.")
 	virtual void UnregisterStructPropertyLayout( FName StructTypeName );
 
 	/**
@@ -422,7 +427,7 @@ public:
 	virtual TSharedRef< FAssetEditorToolkit > CreatePropertyEditorToolkit( const EToolkitMode::Type Mode, const TSharedPtr< IToolkitHost >& InitToolkitHost, const TArray< UObject* >& ObjectsToEdit );
 	virtual TSharedRef< FAssetEditorToolkit > CreatePropertyEditorToolkit( const EToolkitMode::Type Mode, const TSharedPtr< IToolkitHost >& InitToolkitHost, const TArray< TWeakObjectPtr< UObject > >& ObjectsToEdit );
 
-	FPropertyTypeLayoutCallback GetPropertyTypeCustomization(const UProperty* InProperty);
+	FPropertyTypeLayoutCallback GetPropertyTypeCustomization(const UProperty* InProperty,const IPropertyHandle& PropertyHandle);
 	bool IsCustomizedStruct(const UStruct* Struct) const;
 
 private:
diff --git a/Engine/Source/Editor/SkillSystemEditor/Private/AttributeDetails.cpp b/Engine/Source/Editor/SkillSystemEditor/Private/AttributeDetails.cpp
--- a/Engine/Source/Editor/SkillSystemEditor/Private/AttributeDetails.cpp
+++ b/Engine/Source/Editor/SkillSystemEditor/Private/AttributeDetails.cpp
@@ -17,13 +17,13 @@
 
 DEFINE_LOG_CATEGORY(LogAttributeDetails);
 
-TSharedRef<IStructCustomization> FAttributePropertyDetails::MakeInstance()
+TSharedRef<IPropertyTypeCustomization> FAttributePropertyDetails::MakeInstance()
 {
 	return MakeShareable(new FAttributePropertyDetails());
 }
 
 BEGIN_SLATE_FUNCTION_BUILD_OPTIMIZATION
-void FAttributePropertyDetails::CustomizeStructHeader( TSharedRef<IPropertyHandle> StructPropertyHandle, class FDetailWidgetRow& HeaderRow, IStructCustomizationUtils& StructCustomizationUtils )
+void FAttributePropertyDetails::CustomizeHeader( TSharedRef<IPropertyHandle> StructPropertyHandle, class FDetailWidgetRow& HeaderRow, IStructCustomizationUtils& StructCustomizationUtils )
 {
 	MyProperty = StructPropertyHandle->GetChildHandle(GET_MEMBER_NAME_CHECKED(FGameplayAttribute,Attribute));
 
@@ -78,7 +78,7 @@ void FAttributePropertyDetails::CustomizeStructHeader( TSharedRef<IPropertyHandl
 }
 END_SLATE_FUNCTION_BUILD_OPTIMIZATION
 
-void FAttributePropertyDetails::CustomizeStructChildren( TSharedRef<IPropertyHandle> StructPropertyHandle, class IDetailChildrenBuilder& StructBuilder, IStructCustomizationUtils& StructCustomizationUtils )
+void FAttributePropertyDetails::CustomizeChildren( TSharedRef<IPropertyHandle> StructPropertyHandle, class IDetailChildrenBuilder& StructBuilder, IStructCustomizationUtils& StructCustomizationUtils )
 {
 
 
@@ -232,13 +232,13 @@ void FAttributeDetails::OnChangeProperty(TSharedPtr<FString> ItemSelected, ESele
 // ------------------------------------------------------------------------------------
 // ------------------------------------------------------------------------------------
 
-TSharedRef<IStructCustomization> FScalableFloatDetails::MakeInstance()
+TSharedRef<IPropertyTypeCustomization> FScalableFloatDetails::MakeInstance()
 {
 	return MakeShareable(new FScalableFloatDetails());
 }
 
 BEGIN_SLATE_FUNCTION_BUILD_OPTIMIZATION
-void FScalableFloatDetails::CustomizeStructHeader( TSharedRef<class IPropertyHandle> StructPropertyHandle, class FDetailWidgetRow& HeaderRow, IStructCustomizationUtils& StructCustomizationUtils )
+void FScalableFloatDetails::CustomizeHeader( TSharedRef<class IPropertyHandle> StructPropertyHandle, class FDetailWidgetRow& HeaderRow, IStructCustomizationUtils& StructCustomizationUtils )
 {
 	uint32 NumChildren = 0;
 	StructPropertyHandle->GetNumChildren(NumChildren);
@@ -471,7 +471,7 @@ bool FScalableFloatDetails::IsEditable( ) const
 	return true;
 }
 
-void FScalableFloatDetails::CustomizeStructChildren( TSharedRef<class IPropertyHandle> StructPropertyHandle, class IDetailChildrenBuilder& StructBuilder, IStructCustomizationUtils& StructCustomizationUtils )
+void FScalableFloatDetails::CustomizeChildren( TSharedRef<class IPropertyHandle> StructPropertyHandle, class IDetailChildrenBuilder& StructBuilder, IStructCustomizationUtils& StructCustomizationUtils )
 {
 	
 }
diff --git a/Engine/Source/Editor/SkillSystemEditor/Private/AttributeDetails.h b/Engine/Source/Editor/SkillSystemEditor/Private/AttributeDetails.h
--- a/Engine/Source/Editor/SkillSystemEditor/Private/AttributeDetails.h
+++ b/Engine/Source/Editor/SkillSystemEditor/Private/AttributeDetails.h
@@ -30,14 +30,14 @@ private:
 };
 
 
-class FAttributePropertyDetails : public IStructCustomization
+class FAttributePropertyDetails : public IPropertyTypeCustomization
 {
 public:
-	static TSharedRef<IStructCustomization> MakeInstance();
+	static TSharedRef<IPropertyTypeCustomization> MakeInstance();
 
-	/** IStructCustomization interface */
-	virtual void CustomizeStructHeader( TSharedRef<class IPropertyHandle> StructPropertyHandle, class FDetailWidgetRow& HeaderRow, IStructCustomizationUtils& StructCustomizationUtils ) OVERRIDE;
-	virtual void CustomizeStructChildren( TSharedRef<class IPropertyHandle> StructPropertyHandle, class IDetailChildrenBuilder& StructBuilder, IStructCustomizationUtils& StructCustomizationUtils ) OVERRIDE;              
+	/** IPropertyTypeCustomization interface */
+	virtual void CustomizeHeader( TSharedRef<class IPropertyHandle> StructPropertyHandle, class FDetailWidgetRow& HeaderRow, IStructCustomizationUtils& StructCustomizationUtils ) OVERRIDE;
+	virtual void CustomizeChildren( TSharedRef<class IPropertyHandle> StructPropertyHandle, class IDetailChildrenBuilder& StructBuilder, IStructCustomizationUtils& StructCustomizationUtils ) OVERRIDE;              
 
 private:
 
@@ -49,16 +49,16 @@ private:
 	void OnChangeProperty(TSharedPtr<FString> ItemSelected, ESelectInfo::Type SelectInfo);
 };
 
-class FScalableFloatDetails : public IStructCustomization
+class FScalableFloatDetails : public IPropertyTypeCustomization
 {
 public:
-	static TSharedRef<IStructCustomization> MakeInstance();
+	static TSharedRef<IPropertyTypeCustomization> MakeInstance();
 
 
 private:
 
-	virtual void CustomizeStructHeader( TSharedRef<class IPropertyHandle> StructPropertyHandle, class FDetailWidgetRow& HeaderRow, IStructCustomizationUtils& StructCustomizationUtils ) OVERRIDE;
-	virtual void CustomizeStructChildren( TSharedRef<class IPropertyHandle> StructPropertyHandle, class IDetailChildrenBuilder& StructBuilder, IStructCustomizationUtils& StructCustomizationUtils ) OVERRIDE;              
+	virtual void CustomizeHeader( TSharedRef<class IPropertyHandle> StructPropertyHandle, class FDetailWidgetRow& HeaderRow, IStructCustomizationUtils& StructCustomizationUtils ) OVERRIDE;
+	virtual void CustomizeChildren( TSharedRef<class IPropertyHandle> StructPropertyHandle, class IDetailChildrenBuilder& StructBuilder, IStructCustomizationUtils& StructCustomizationUtils ) OVERRIDE;              
 
 	bool IsEditable( ) const;
 
diff --git a/Engine/Source/Editor/SkillSystemEditor/Private/SkillSystemEditorModule.cpp b/Engine/Source/Editor/SkillSystemEditor/Private/SkillSystemEditorModule.cpp
--- a/Engine/Source/Editor/SkillSystemEditor/Private/SkillSystemEditorModule.cpp
+++ b/Engine/Source/Editor/SkillSystemEditor/Private/SkillSystemEditorModule.cpp
@@ -19,8 +19,8 @@ void FSkillSystemEditorModule::StartupModule()
 {
 	// Register the details customizer
 	FPropertyEditorModule& PropertyModule = FModuleManager::LoadModuleChecked<FPropertyEditorModule>("PropertyEditor");
-	PropertyModule.RegisterStructPropertyLayout( "GameplayAttribute", FOnGetStructCustomizationInstance::CreateStatic( &FAttributePropertyDetails::MakeInstance ) );
-	PropertyModule.RegisterStructPropertyLayout( "ScalableFloat", FOnGetStructCustomizationInstance::CreateStatic( &FScalableFloatDetails::MakeInstance ) );
+	PropertyModule.RegisterCustomPropertyTypeLayout( "GameplayAttribute", FOnGetPropertyTypeCustomizationInstance::CreateStatic( &FAttributePropertyDetails::MakeInstance ) );
+	PropertyModule.RegisterCustomPropertyTypeLayout( "ScalableFloat", FOnGetPropertyTypeCustomizationInstance::CreateStatic( &FScalableFloatDetails::MakeInstance ) );
 
 	PropertyModule.RegisterCustomPropertyLayout( "AttributeSet", FOnGetDetailCustomizationInstance::CreateStatic( &FAttributeDetails::MakeInstance ) );
 }
diff --git a/Engine/Source/Editor/WorldBrowser/Private/Tiles/WorldTileCollectionModel.cpp b/Engine/Source/Editor/WorldBrowser/Private/Tiles/WorldTileCollectionModel.cpp
--- a/Engine/Source/Editor/WorldBrowser/Private/Tiles/WorldTileCollectionModel.cpp
+++ b/Engine/Source/Editor/WorldBrowser/Private/Tiles/WorldTileCollectionModel.cpp
@@ -559,12 +559,12 @@ void FWorldTileCollectionModel::RegisterDetailsCustomization(FPropertyEditorModu
 	TSharedRef<FWorldTileCollectionModel> WorldModel = StaticCastSharedRef<FWorldTileCollectionModel>(this->AsShared());
 	
 	// Register our struct customizations
-	InPropertyModule.RegisterStructPropertyLayout("TileStreamingLevelDetails", 
-		FOnGetStructCustomizationInstance::CreateStatic(&FStreamingLevelDetailsCustomization::MakeInstance, WorldModel)
+	InPropertyModule.RegisterCustomPropertyTypeLayout("TileStreamingLevelDetails", 
+		FOnGetPropertyTypeCustomizationInstance::CreateStatic(&FStreamingLevelDetailsCustomization::MakeInstance, WorldModel)
 		);
 
-	InPropertyModule.RegisterStructPropertyLayout("TileLODEntryDetails", 
-		FOnGetStructCustomizationInstance::CreateStatic(&FTileLODEntryDetailsCustomization::MakeInstance,  WorldModel)
+	InPropertyModule.RegisterCustomPropertyTypeLayout("TileLODEntryDetails", 
+		FOnGetPropertyTypeCustomizationInstance::CreateStatic(&FTileLODEntryDetailsCustomization::MakeInstance,  WorldModel)
 		);
 
 	InDetailsView->RegisterInstancedCustomPropertyLayout(UWorldTileDetails::StaticClass(), 
diff --git a/Engine/Source/Editor/WorldBrowser/Private/Tiles/WorldTileDetailsCustomization.cpp b/Engine/Source/Editor/WorldBrowser/Private/Tiles/WorldTileDetailsCustomization.cpp
--- a/Engine/Source/Editor/WorldBrowser/Private/Tiles/WorldTileDetailsCustomization.cpp
+++ b/Engine/Source/Editor/WorldBrowser/Private/Tiles/WorldTileDetailsCustomization.cpp
@@ -171,14 +171,14 @@ bool FWorldTileDetailsCustomization::OnShouldFilterParentPackage(const FString&
 
 //////////////////////////////////////////////////////////////////////////////////////////////
 // FStreamingLevelDetailsCustomization
-TSharedRef<IStructCustomization> FStreamingLevelDetailsCustomization::MakeInstance(TSharedRef<FWorldTileCollectionModel> InWorldData) 
+TSharedRef<IPropertyTypeCustomization> FStreamingLevelDetailsCustomization::MakeInstance(TSharedRef<FWorldTileCollectionModel> InWorldData) 
 {
 	TSharedRef<FStreamingLevelDetailsCustomization> Instance = MakeShareable(new FStreamingLevelDetailsCustomization());
 	Instance->WorldModel = InWorldData;
 	return Instance;
 }
 
-void FStreamingLevelDetailsCustomization::CustomizeStructHeader(TSharedRef<class IPropertyHandle> StructPropertyHandle, 
+void FStreamingLevelDetailsCustomization::CustomizeHeader(TSharedRef<class IPropertyHandle> StructPropertyHandle, 
 																class FDetailWidgetRow& HeaderRow, 
 																IStructCustomizationUtils& StructCustomizationUtils )
 {
@@ -189,7 +189,7 @@ void FStreamingLevelDetailsCustomization::CustomizeStructHeader(TSharedRef<class
 		];
 }
 
-void FStreamingLevelDetailsCustomization::CustomizeStructChildren(TSharedRef<IPropertyHandle> StructPropertyHandle, 
+void FStreamingLevelDetailsCustomization::CustomizeChildren(TSharedRef<IPropertyHandle> StructPropertyHandle, 
 																	class IDetailChildrenBuilder& ChildBuilder, 
 																	IStructCustomizationUtils& StructCustomizationUtils )
 {
@@ -229,14 +229,14 @@ bool FStreamingLevelDetailsCustomization::OnShouldFilterStreamingPackage(const F
 
 //////////////////////////////////////////////////////////////////////////////////////////////
 // FTileLODEntryDetailsCustomization
-TSharedRef<IStructCustomization> FTileLODEntryDetailsCustomization::MakeInstance(TSharedRef<FWorldTileCollectionModel> InWorldData) 
+TSharedRef<IPropertyTypeCustomization> FTileLODEntryDetailsCustomization::MakeInstance(TSharedRef<FWorldTileCollectionModel> InWorldData) 
 {
 	TSharedRef<FTileLODEntryDetailsCustomization> Instance = MakeShareable(new FTileLODEntryDetailsCustomization());
 	Instance->WorldModel = InWorldData;
 	return Instance;
 }
 
-void FTileLODEntryDetailsCustomization::CustomizeStructHeader(TSharedRef<class IPropertyHandle> StructPropertyHandle, 
+void FTileLODEntryDetailsCustomization::CustomizeHeader(TSharedRef<class IPropertyHandle> StructPropertyHandle, 
 																class FDetailWidgetRow& HeaderRow, 
 																IStructCustomizationUtils& StructCustomizationUtils )
 {
@@ -258,7 +258,7 @@ void FTileLODEntryDetailsCustomization::CustomizeStructHeader(TSharedRef<class I
 		];
 }
 
-void FTileLODEntryDetailsCustomization::CustomizeStructChildren(TSharedRef<IPropertyHandle> StructPropertyHandle, 
+void FTileLODEntryDetailsCustomization::CustomizeChildren(TSharedRef<IPropertyHandle> StructPropertyHandle, 
 																	class IDetailChildrenBuilder& ChildBuilder, 
 																	IStructCustomizationUtils& StructCustomizationUtils )
 {
diff --git a/Engine/Source/Editor/WorldBrowser/Private/Tiles/WorldTileDetailsCustomization.h b/Engine/Source/Editor/WorldBrowser/Private/Tiles/WorldTileDetailsCustomization.h
--- a/Engine/Source/Editor/WorldBrowser/Private/Tiles/WorldTileDetailsCustomization.h
+++ b/Engine/Source/Editor/WorldBrowser/Private/Tiles/WorldTileDetailsCustomization.h
@@ -32,14 +32,14 @@ private:
 
 /////////////////////////////////////////////////////
 // FStreamingLevelDetailsCustomization
-class FStreamingLevelDetailsCustomization : public IStructCustomization
+class FStreamingLevelDetailsCustomization : public IPropertyTypeCustomization
 {
 public:
-	static TSharedRef<IStructCustomization> MakeInstance(TSharedRef<FWorldTileCollectionModel> InWorldData);
+	static TSharedRef<IPropertyTypeCustomization> MakeInstance(TSharedRef<FWorldTileCollectionModel> InWorldData);
 
-	/** IStructCustomization interface */
-	virtual void CustomizeStructHeader(TSharedRef<class IPropertyHandle> StructPropertyHandle, class FDetailWidgetRow& HeaderRow, IStructCustomizationUtils& StructCustomizationUtils ) OVERRIDE;
-	virtual void CustomizeStructChildren(TSharedRef<class IPropertyHandle> StructPropertyHandle, class IDetailChildrenBuilder& ChildBuilder, IStructCustomizationUtils& StructCustomizationUtils ) OVERRIDE;
+	/** IPropertyTypeCustomization interface */
+	virtual void CustomizeHeader(TSharedRef<class IPropertyHandle> StructPropertyHandle, class FDetailWidgetRow& HeaderRow, IStructCustomizationUtils& StructCustomizationUtils ) OVERRIDE;
+	virtual void CustomizeChildren(TSharedRef<class IPropertyHandle> StructPropertyHandle, class IDetailChildrenBuilder& ChildBuilder, IStructCustomizationUtils& StructCustomizationUtils ) OVERRIDE;
 
 private:
 	bool OnShouldFilterStreamingPackage(const FString& InPackageName) const;
@@ -50,14 +50,14 @@ private:
 
 /////////////////////////////////////////////////////
 // FTileLODEntryDetailsCustomization
-class FTileLODEntryDetailsCustomization : public IStructCustomization
+class FTileLODEntryDetailsCustomization : public IPropertyTypeCustomization
 {
 public:
-	static TSharedRef<IStructCustomization> MakeInstance(TSharedRef<FWorldTileCollectionModel> InWorldData);
+	static TSharedRef<IPropertyTypeCustomization> MakeInstance(TSharedRef<FWorldTileCollectionModel> InWorldData);
 
-	/** IStructCustomization interface */
-	virtual void CustomizeStructHeader( TSharedRef<class IPropertyHandle> StructPropertyHandle, class FDetailWidgetRow& HeaderRow, IStructCustomizationUtils& StructCustomizationUtils ) OVERRIDE;
-	virtual void CustomizeStructChildren( TSharedRef<IPropertyHandle> StructPropertyHandle, class IDetailChildrenBuilder& ChildBuilder, IStructCustomizationUtils& StructCustomizationUtils ) OVERRIDE;
+	/** IPropertyTypeCustomization interface */
+	virtual void CustomizeHeader( TSharedRef<class IPropertyHandle> StructPropertyHandle, class FDetailWidgetRow& HeaderRow, IStructCustomizationUtils& StructCustomizationUtils ) OVERRIDE;
+	virtual void CustomizeChildren( TSharedRef<IPropertyHandle> StructPropertyHandle, class IDetailChildrenBuilder& ChildBuilder, IStructCustomizationUtils& StructCustomizationUtils ) OVERRIDE;
 
 private:
 	FReply OnGenerateTile();
