Fixup usage of deprecated IStructCustomizationUtils interface

diff --git a/Engine/Plugins/Experimental/Paper2D/Source/Paper2DEditor/Private/SpriteEditor/SpritePolygonCollectionCustomization.cpp b/Engine/Plugins/Experimental/Paper2D/Source/Paper2DEditor/Private/SpriteEditor/SpritePolygonCollectionCustomization.cpp
--- a/Engine/Plugins/Experimental/Paper2D/Source/Paper2DEditor/Private/SpriteEditor/SpritePolygonCollectionCustomization.cpp
+++ b/Engine/Plugins/Experimental/Paper2D/Source/Paper2DEditor/Private/SpriteEditor/SpritePolygonCollectionCustomization.cpp
@@ -11,11 +11,11 @@ TSharedRef<IPropertyTypeCustomization> FSpritePolygonCollectionCustomization::Ma
 	return MakeShareable(new FSpritePolygonCollectionCustomization);
 }
 
-void FSpritePolygonCollectionCustomization::CustomizeHeader(TSharedRef<IPropertyHandle> StructPropertyHandle, class FDetailWidgetRow& HeaderRow, IStructCustomizationUtils& StructCustomizationUtils)
+void FSpritePolygonCollectionCustomization::CustomizeHeader(TSharedRef<IPropertyHandle> StructPropertyHandle, class FDetailWidgetRow& HeaderRow, IPropertyTypeCustomizationUtils& StructCustomizationUtils)
 {
 }
 
-void FSpritePolygonCollectionCustomization::CustomizeChildren(TSharedRef<class IPropertyHandle> StructPropertyHandle, class IDetailChildrenBuilder& StructBuilder, IStructCustomizationUtils& StructCustomizationUtils)
+void FSpritePolygonCollectionCustomization::CustomizeChildren(TSharedRef<class IPropertyHandle> StructPropertyHandle, class IDetailChildrenBuilder& StructBuilder, IPropertyTypeCustomizationUtils& StructCustomizationUtils)
 {
 }
 #if 0
diff --git a/Engine/Plugins/Experimental/Paper2D/Source/Paper2DEditor/Private/SpriteEditor/SpritePolygonCollectionCustomization.h b/Engine/Plugins/Experimental/Paper2D/Source/Paper2DEditor/Private/SpriteEditor/SpritePolygonCollectionCustomization.h
--- a/Engine/Plugins/Experimental/Paper2D/Source/Paper2DEditor/Private/SpriteEditor/SpritePolygonCollectionCustomization.h
+++ b/Engine/Plugins/Experimental/Paper2D/Source/Paper2DEditor/Private/SpriteEditor/SpritePolygonCollectionCustomization.h
@@ -14,8 +14,8 @@ public:
 	static TSharedRef<IPropertyTypeCustomization> MakeInstance();
 
 	// IPropertyTypeCustomization interface
-	virtual void CustomizeHeader(TSharedRef<IPropertyHandle> StructPropertyHandle, class FDetailWidgetRow& HeaderRow, IStructCustomizationUtils& StructCustomizationUtils) OVERRIDE;
-	virtual void CustomizeChildren(TSharedRef<class IPropertyHandle> StructPropertyHandle, class IDetailChildrenBuilder& StructBuilder, IStructCustomizationUtils& StructCustomizationUtils) OVERRIDE;
+	virtual void CustomizeHeader(TSharedRef<IPropertyHandle> StructPropertyHandle, class FDetailWidgetRow& HeaderRow, IPropertyTypeCustomizationUtils& StructCustomizationUtils) OVERRIDE;
+	virtual void CustomizeChildren(TSharedRef<class IPropertyHandle> StructPropertyHandle, class IDetailChildrenBuilder& StructBuilder, IPropertyTypeCustomizationUtils& StructCustomizationUtils) OVERRIDE;
 	// End of IPropertyTypeCustomization interface
 
 protected:
diff --git a/Engine/Source/Editor/BehaviorTreeEditor/Private/DetailCustomizations/BlackboardSelectorDetails.cpp b/Engine/Source/Editor/BehaviorTreeEditor/Private/DetailCustomizations/BlackboardSelectorDetails.cpp
--- a/Engine/Source/Editor/BehaviorTreeEditor/Private/DetailCustomizations/BlackboardSelectorDetails.cpp
+++ b/Engine/Source/Editor/BehaviorTreeEditor/Private/DetailCustomizations/BlackboardSelectorDetails.cpp
@@ -13,7 +13,7 @@ TSharedRef<IPropertyTypeCustomization> FBlackboardSelectorDetails::MakeInstance(
 	return MakeShareable( new FBlackboardSelectorDetails );
 }
 
-void FBlackboardSelectorDetails::CustomizeHeader( TSharedRef<class IPropertyHandle> StructPropertyHandle, class FDetailWidgetRow& HeaderRow, IStructCustomizationUtils& StructCustomizationUtils )
+void FBlackboardSelectorDetails::CustomizeHeader( TSharedRef<class IPropertyHandle> StructPropertyHandle, class FDetailWidgetRow& HeaderRow, IPropertyTypeCustomizationUtils& StructCustomizationUtils )
 {
 	MyStructProperty = StructPropertyHandle;
 	PropUtils = StructCustomizationUtils.GetPropertyUtilities().Get();
@@ -42,7 +42,7 @@ void FBlackboardSelectorDetails::CustomizeHeader( TSharedRef<class IPropertyHand
 	InitKeyFromProperty();
 }
 
-void FBlackboardSelectorDetails::CustomizeChildren( TSharedRef<class IPropertyHandle> StructPropertyHandle, class IDetailChildrenBuilder& StructBuilder, IStructCustomizationUtils& StructCustomizationUtils )
+void FBlackboardSelectorDetails::CustomizeChildren( TSharedRef<class IPropertyHandle> StructPropertyHandle, class IDetailChildrenBuilder& StructBuilder, IPropertyTypeCustomizationUtils& StructCustomizationUtils )
 {
 }
 
diff --git a/Engine/Source/Editor/BehaviorTreeEditor/Private/DetailCustomizations/BlackboardSelectorDetails.h b/Engine/Source/Editor/BehaviorTreeEditor/Private/DetailCustomizations/BlackboardSelectorDetails.h
--- a/Engine/Source/Editor/BehaviorTreeEditor/Private/DetailCustomizations/BlackboardSelectorDetails.h
+++ b/Engine/Source/Editor/BehaviorTreeEditor/Private/DetailCustomizations/BlackboardSelectorDetails.h
@@ -12,8 +12,8 @@ public:
 	static TSharedRef<IPropertyTypeCustomization> MakeInstance();
 
 	/** IPropertyTypeCustomization interface */
-	virtual void CustomizeHeader( TSharedRef<class IPropertyHandle> StructPropertyHandle, class FDetailWidgetRow& HeaderRow, IStructCustomizationUtils& StructCustomizationUtils ) OVERRIDE;
-	virtual void CustomizeChildren( TSharedRef<class IPropertyHandle> StructPropertyHandle, class IDetailChildrenBuilder& StructBuilder, IStructCustomizationUtils& StructCustomizationUtils ) OVERRIDE;
+	virtual void CustomizeHeader( TSharedRef<class IPropertyHandle> StructPropertyHandle, class FDetailWidgetRow& HeaderRow, IPropertyTypeCustomizationUtils& StructCustomizationUtils ) OVERRIDE;
+	virtual void CustomizeChildren( TSharedRef<class IPropertyHandle> StructPropertyHandle, class IDetailChildrenBuilder& StructBuilder, IPropertyTypeCustomizationUtils& StructCustomizationUtils ) OVERRIDE;
 
 private:
 
diff --git a/Engine/Source/Editor/DetailCustomizations/Private/AttenuationSettingsCustomizations.cpp b/Engine/Source/Editor/DetailCustomizations/Private/AttenuationSettingsCustomizations.cpp
--- a/Engine/Source/Editor/DetailCustomizations/Private/AttenuationSettingsCustomizations.cpp
+++ b/Engine/Source/Editor/DetailCustomizations/Private/AttenuationSettingsCustomizations.cpp
@@ -8,7 +8,7 @@ TSharedRef<IPropertyTypeCustomization> FAttenuationSettingsCustomization::MakeIn
 	return MakeShareable( new FAttenuationSettingsCustomization );
 }
 
-void FAttenuationSettingsCustomization::CustomizeHeader( TSharedRef<class IPropertyHandle> StructPropertyHandle, class FDetailWidgetRow& HeaderRow, IStructCustomizationUtils& StructCustomizationUtils )
+void FAttenuationSettingsCustomization::CustomizeHeader( TSharedRef<class IPropertyHandle> StructPropertyHandle, class FDetailWidgetRow& HeaderRow, IPropertyTypeCustomizationUtils& StructCustomizationUtils )
 {
 	// We'll set up reset to default ourselves
 	const bool bDisplayResetToDefault = false;
@@ -21,7 +21,7 @@ void FAttenuationSettingsCustomization::CustomizeHeader( TSharedRef<class IPrope
 		];
 }
 
-void FAttenuationSettingsCustomization::CustomizeChildren( TSharedRef<IPropertyHandle> StructPropertyHandle, class IDetailChildrenBuilder& ChildBuilder, IStructCustomizationUtils& StructCustomizationUtils )
+void FAttenuationSettingsCustomization::CustomizeChildren( TSharedRef<IPropertyHandle> StructPropertyHandle, class IDetailChildrenBuilder& ChildBuilder, IPropertyTypeCustomizationUtils& StructCustomizationUtils )
 {
 	uint32 NumChildren;
 	StructPropertyHandle->GetNumChildren( NumChildren );
diff --git a/Engine/Source/Editor/DetailCustomizations/Private/AttenuationSettingsCustomizations.h b/Engine/Source/Editor/DetailCustomizations/Private/AttenuationSettingsCustomizations.h
--- a/Engine/Source/Editor/DetailCustomizations/Private/AttenuationSettingsCustomizations.h
+++ b/Engine/Source/Editor/DetailCustomizations/Private/AttenuationSettingsCustomizations.h
@@ -8,8 +8,8 @@ public:
 	static TSharedRef<IPropertyTypeCustomization> MakeInstance();
 
 	/** IPropertyTypeCustomization instance */
-	virtual void CustomizeHeader( TSharedRef<class IPropertyHandle> StructPropertyHandle, class FDetailWidgetRow& HeaderRow, IStructCustomizationUtils& StructCustomizationUtils ) OVERRIDE;
-	virtual void CustomizeChildren( TSharedRef<IPropertyHandle> StructPropertyHandle, class IDetailChildrenBuilder& ChildBuilder, IStructCustomizationUtils& StructCustomizationUtils ) OVERRIDE;
+	virtual void CustomizeHeader( TSharedRef<class IPropertyHandle> StructPropertyHandle, class FDetailWidgetRow& HeaderRow, IPropertyTypeCustomizationUtils& StructCustomizationUtils ) OVERRIDE;
+	virtual void CustomizeChildren( TSharedRef<IPropertyHandle> StructPropertyHandle, class IDetailChildrenBuilder& ChildBuilder, IPropertyTypeCustomizationUtils& StructCustomizationUtils ) OVERRIDE;
 
 protected:
 
diff --git a/Engine/Source/Editor/DetailCustomizations/Private/BlackboardEntryDetails.cpp b/Engine/Source/Editor/DetailCustomizations/Private/BlackboardEntryDetails.cpp
--- a/Engine/Source/Editor/DetailCustomizations/Private/BlackboardEntryDetails.cpp
+++ b/Engine/Source/Editor/DetailCustomizations/Private/BlackboardEntryDetails.cpp
@@ -12,7 +12,7 @@ TSharedRef<IPropertyTypeCustomization> FBlackboardEntryDetails::MakeInstance()
 }
 
 BEGIN_SLATE_FUNCTION_BUILD_OPTIMIZATION
-void FBlackboardEntryDetails::CustomizeHeader( TSharedRef<class IPropertyHandle> StructPropertyHandle, class FDetailWidgetRow& HeaderRow, IStructCustomizationUtils& StructCustomizationUtils )
+void FBlackboardEntryDetails::CustomizeHeader( TSharedRef<class IPropertyHandle> StructPropertyHandle, class FDetailWidgetRow& HeaderRow, IPropertyTypeCustomizationUtils& StructCustomizationUtils )
 {
 	MyStructProperty = StructPropertyHandle;
 	MyNameProperty = MyStructProperty->GetChildHandle(TEXT("EntryName"));
@@ -32,7 +32,7 @@ void FBlackboardEntryDetails::CustomizeHeader( TSharedRef<class IPropertyHandle>
 }
 END_SLATE_FUNCTION_BUILD_OPTIMIZATION
 
-void FBlackboardEntryDetails::CustomizeChildren( TSharedRef<class IPropertyHandle> StructPropertyHandle, class IDetailChildrenBuilder& StructBuilder, IStructCustomizationUtils& StructCustomizationUtils )
+void FBlackboardEntryDetails::CustomizeChildren( TSharedRef<class IPropertyHandle> StructPropertyHandle, class IDetailChildrenBuilder& StructBuilder, IPropertyTypeCustomizationUtils& StructCustomizationUtils )
 {
 	if (StructPropertyHandle->IsValidHandle())
 	{
diff --git a/Engine/Source/Editor/DetailCustomizations/Private/BlackboardEntryDetails.h b/Engine/Source/Editor/DetailCustomizations/Private/BlackboardEntryDetails.h
--- a/Engine/Source/Editor/DetailCustomizations/Private/BlackboardEntryDetails.h
+++ b/Engine/Source/Editor/DetailCustomizations/Private/BlackboardEntryDetails.h
@@ -9,8 +9,8 @@ public:
 	static TSharedRef<IPropertyTypeCustomization> MakeInstance();
 
 	/** IPropertyTypeCustomization interface */
-	virtual void CustomizeHeader( TSharedRef<class IPropertyHandle> StructPropertyHandle, class FDetailWidgetRow& HeaderRow, IStructCustomizationUtils& StructCustomizationUtils ) OVERRIDE;
-	virtual void CustomizeChildren( TSharedRef<class IPropertyHandle> StructPropertyHandle, class IDetailChildrenBuilder& StructBuilder, IStructCustomizationUtils& StructCustomizationUtils ) OVERRIDE;
+	virtual void CustomizeHeader( TSharedRef<class IPropertyHandle> StructPropertyHandle, class FDetailWidgetRow& HeaderRow, IPropertyTypeCustomizationUtils& StructCustomizationUtils ) OVERRIDE;
+	virtual void CustomizeChildren( TSharedRef<class IPropertyHandle> StructPropertyHandle, class IDetailChildrenBuilder& StructBuilder, IPropertyTypeCustomizationUtils& StructCustomizationUtils ) OVERRIDE;
 
 private:
 
diff --git a/Engine/Source/Editor/DetailCustomizations/Private/BodyInstanceCustomization.cpp b/Engine/Source/Editor/DetailCustomizations/Private/BodyInstanceCustomization.cpp
--- a/Engine/Source/Editor/DetailCustomizations/Private/BodyInstanceCustomization.cpp
+++ b/Engine/Source/Editor/DetailCustomizations/Private/BodyInstanceCustomization.cpp
@@ -34,7 +34,7 @@ void FBodyInstanceCustomization::RefreshCollisionProfiles()
 }
 
 BEGIN_SLATE_FUNCTION_BUILD_OPTIMIZATION
-void FBodyInstanceCustomization::CustomizeChildren( TSharedRef<class IPropertyHandle> StructPropertyHandle, class IDetailChildrenBuilder& StructBuilder, IStructCustomizationUtils& StructCustomizationUtils )
+void FBodyInstanceCustomization::CustomizeChildren( TSharedRef<class IPropertyHandle> StructPropertyHandle, class IDetailChildrenBuilder& StructBuilder, IPropertyTypeCustomizationUtils& StructCustomizationUtils )
 {
 	CollisionProfileNameHandle = StructPropertyHandle->GetChildHandle(TEXT("CollisionProfileName"));
 	CollisionEnabledHandle = StructPropertyHandle->GetChildHandle(TEXT("CollisionEnabled"));
diff --git a/Engine/Source/Editor/DetailCustomizations/Private/BodyInstanceCustomization.h b/Engine/Source/Editor/DetailCustomizations/Private/BodyInstanceCustomization.h
--- a/Engine/Source/Editor/DetailCustomizations/Private/BodyInstanceCustomization.h
+++ b/Engine/Source/Editor/DetailCustomizations/Private/BodyInstanceCustomization.h
@@ -23,8 +23,8 @@ public:
 	}
 
 	/** IPropertyTypeCustomization interface */
-	virtual void CustomizeHeader( TSharedRef<IPropertyHandle> StructPropertyHandle, class FDetailWidgetRow& HeaderRow, IStructCustomizationUtils& StructCustomizationUtils ) OVERRIDE {};
-	virtual void CustomizeChildren( TSharedRef<class IPropertyHandle> StructPropertyHandle, class IDetailChildrenBuilder& StructBuilder, IStructCustomizationUtils& StructCustomizationUtils ) OVERRIDE;
+	virtual void CustomizeHeader( TSharedRef<IPropertyHandle> StructPropertyHandle, class FDetailWidgetRow& HeaderRow, IPropertyTypeCustomizationUtils& StructCustomizationUtils ) OVERRIDE {};
+	virtual void CustomizeChildren( TSharedRef<class IPropertyHandle> StructPropertyHandle, class IDetailChildrenBuilder& StructBuilder, IPropertyTypeCustomizationUtils& StructCustomizationUtils ) OVERRIDE;
 
 private:
 	// Profile combo related
diff --git a/Engine/Source/Editor/DetailCustomizations/Private/CurveStructCustomization.cpp b/Engine/Source/Editor/DetailCustomizations/Private/CurveStructCustomization.cpp
--- a/Engine/Source/Editor/DetailCustomizations/Private/CurveStructCustomization.cpp
+++ b/Engine/Source/Editor/DetailCustomizations/Private/CurveStructCustomization.cpp
@@ -29,7 +29,7 @@ FCurveStructCustomization::FCurveStructCustomization()
 {
 }
 
-void FCurveStructCustomization::CustomizeHeader( TSharedRef<class IPropertyHandle> InStructPropertyHandle, class FDetailWidgetRow& HeaderRow, IStructCustomizationUtils& StructCustomizationUtils )
+void FCurveStructCustomization::CustomizeHeader( TSharedRef<class IPropertyHandle> InStructPropertyHandle, class FDetailWidgetRow& HeaderRow, IPropertyTypeCustomizationUtils& StructCustomizationUtils )
 {
 	this->StructPropertyHandle = InStructPropertyHandle;
 
@@ -106,7 +106,7 @@ void FCurveStructCustomization::CustomizeHeader( TSharedRef<class IPropertyHandl
 	}
 }
 
-void FCurveStructCustomization::CustomizeChildren( TSharedRef<class IPropertyHandle> InStructPropertyHandle, class IDetailChildrenBuilder& StructBuilder, IStructCustomizationUtils& StructCustomizationUtils )
+void FCurveStructCustomization::CustomizeChildren( TSharedRef<class IPropertyHandle> InStructPropertyHandle, class IDetailChildrenBuilder& StructBuilder, IPropertyTypeCustomizationUtils& StructCustomizationUtils )
 {
 	uint32 NumChildren = 0;
 	StructPropertyHandle->GetNumChildren(NumChildren);
diff --git a/Engine/Source/Editor/DetailCustomizations/Private/CurveStructCustomization.h b/Engine/Source/Editor/DetailCustomizations/Private/CurveStructCustomization.h
--- a/Engine/Source/Editor/DetailCustomizations/Private/CurveStructCustomization.h
+++ b/Engine/Source/Editor/DetailCustomizations/Private/CurveStructCustomization.h
@@ -18,8 +18,8 @@ public:
 	virtual ~FCurveStructCustomization();
 
 	/** IPropertyTypeCustomization interface */
-	virtual void CustomizeHeader( TSharedRef<class IPropertyHandle> InStructPropertyHandle, class FDetailWidgetRow& HeaderRow, IStructCustomizationUtils& StructCustomizationUtils ) OVERRIDE;
-	virtual void CustomizeChildren( TSharedRef<class IPropertyHandle> InStructPropertyHandle, class IDetailChildrenBuilder& StructBuilder, IStructCustomizationUtils& StructCustomizationUtils ) OVERRIDE;
+	virtual void CustomizeHeader( TSharedRef<class IPropertyHandle> InStructPropertyHandle, class FDetailWidgetRow& HeaderRow, IPropertyTypeCustomizationUtils& StructCustomizationUtils ) OVERRIDE;
+	virtual void CustomizeChildren( TSharedRef<class IPropertyHandle> InStructPropertyHandle, class IDetailChildrenBuilder& StructBuilder, IPropertyTypeCustomizationUtils& StructCustomizationUtils ) OVERRIDE;
 
 	/** FCurveOwnerInterface interface */
 	virtual TArray<FRichCurveEditInfoConst> GetCurves() const OVERRIDE;
diff --git a/Engine/Source/Editor/DetailCustomizations/Private/CurveTableCustomization.h b/Engine/Source/Editor/DetailCustomizations/Private/CurveTableCustomization.h
--- a/Engine/Source/Editor/DetailCustomizations/Private/CurveTableCustomization.h
+++ b/Engine/Source/Editor/DetailCustomizations/Private/CurveTableCustomization.h
@@ -14,7 +14,7 @@ public:
 	}
 
 	/** IPropertyTypeCustomization interface */
-	virtual void CustomizeHeader( TSharedRef<class IPropertyHandle> InStructPropertyHandle, class FDetailWidgetRow& HeaderRow, IStructCustomizationUtils& StructCustomizationUtils ) OVERRIDE
+	virtual void CustomizeHeader( TSharedRef<class IPropertyHandle> InStructPropertyHandle, class FDetailWidgetRow& HeaderRow, IPropertyTypeCustomizationUtils& StructCustomizationUtils ) OVERRIDE
 	{
 		this->StructPropertyHandle = InStructPropertyHandle;
 
@@ -25,7 +25,7 @@ public:
 			];
 	}
 
-	virtual void CustomizeChildren( TSharedRef<class IPropertyHandle> InStructPropertyHandle, class IDetailChildrenBuilder& StructBuilder, IStructCustomizationUtils& StructCustomizationUtils ) OVERRIDE
+	virtual void CustomizeChildren( TSharedRef<class IPropertyHandle> InStructPropertyHandle, class IDetailChildrenBuilder& StructBuilder, IPropertyTypeCustomizationUtils& StructCustomizationUtils ) OVERRIDE
 	{
 		/** Get all the existing property handles */
 		CurveTablePropertyHandle = InStructPropertyHandle->GetChildHandle( "CurveTable" );
diff --git a/Engine/Source/Editor/DetailCustomizations/Private/DataTableCategoryCustomization.h b/Engine/Source/Editor/DetailCustomizations/Private/DataTableCategoryCustomization.h
--- a/Engine/Source/Editor/DetailCustomizations/Private/DataTableCategoryCustomization.h
+++ b/Engine/Source/Editor/DetailCustomizations/Private/DataTableCategoryCustomization.h
@@ -16,7 +16,7 @@ public:
 	}
 
 	/** IPropertyTypeCustomization interface */
-	virtual void CustomizeHeader(TSharedRef<class IPropertyHandle> InStructPropertyHandle, class FDetailWidgetRow& HeaderRow, IStructCustomizationUtils& StructCustomizationUtils) OVERRIDE
+	virtual void CustomizeHeader(TSharedRef<class IPropertyHandle> InStructPropertyHandle, class FDetailWidgetRow& HeaderRow, IPropertyTypeCustomizationUtils& StructCustomizationUtils) OVERRIDE
 	{
 		this->StructPropertyHandle = InStructPropertyHandle;
 
@@ -27,7 +27,7 @@ public:
 			];
 	}
 
-	virtual void CustomizeChildren(TSharedRef<class IPropertyHandle> InStructPropertyHandle, class IDetailChildrenBuilder& StructBuilder, IStructCustomizationUtils& StructCustomizationUtils) OVERRIDE
+	virtual void CustomizeChildren(TSharedRef<class IPropertyHandle> InStructPropertyHandle, class IDetailChildrenBuilder& StructBuilder, IPropertyTypeCustomizationUtils& StructCustomizationUtils) OVERRIDE
 	{
 		/** Get all the existing property handles */
 		DataTablePropertyHandle		= InStructPropertyHandle->GetChildHandle("DataTable");
diff --git a/Engine/Source/Editor/DetailCustomizations/Private/DataTableCustomization.h b/Engine/Source/Editor/DetailCustomizations/Private/DataTableCustomization.h
--- a/Engine/Source/Editor/DetailCustomizations/Private/DataTableCustomization.h
+++ b/Engine/Source/Editor/DetailCustomizations/Private/DataTableCustomization.h
@@ -16,7 +16,7 @@ public:
 	}
 
 	/** IPropertyTypeCustomization interface */
-	virtual void CustomizeHeader( TSharedRef<class IPropertyHandle> InStructPropertyHandle, class FDetailWidgetRow& HeaderRow, IStructCustomizationUtils& StructCustomizationUtils ) OVERRIDE
+	virtual void CustomizeHeader( TSharedRef<class IPropertyHandle> InStructPropertyHandle, class FDetailWidgetRow& HeaderRow, IPropertyTypeCustomizationUtils& StructCustomizationUtils ) OVERRIDE
 	{
 		this->StructPropertyHandle = InStructPropertyHandle;
 
@@ -27,7 +27,7 @@ public:
 			];
 	}
 
-	virtual void CustomizeChildren( TSharedRef<class IPropertyHandle> InStructPropertyHandle, class IDetailChildrenBuilder& StructBuilder, IStructCustomizationUtils& StructCustomizationUtils ) OVERRIDE
+	virtual void CustomizeChildren( TSharedRef<class IPropertyHandle> InStructPropertyHandle, class IDetailChildrenBuilder& StructBuilder, IPropertyTypeCustomizationUtils& StructCustomizationUtils ) OVERRIDE
 	{
 		/** Get all the existing property handles */
 		DataTablePropertyHandle = InStructPropertyHandle->GetChildHandle( "DataTable" );
diff --git a/Engine/Source/Editor/DetailCustomizations/Private/DialogueStructsCustomizations.cpp b/Engine/Source/Editor/DetailCustomizations/Private/DialogueStructsCustomizations.cpp
--- a/Engine/Source/Editor/DetailCustomizations/Private/DialogueStructsCustomizations.cpp
+++ b/Engine/Source/Editor/DetailCustomizations/Private/DialogueStructsCustomizations.cpp
@@ -13,7 +13,7 @@ TSharedRef<IPropertyTypeCustomization> FDialogueContextStructCustomization::Make
 	return MakeShareable( new FDialogueContextStructCustomization );
 }
 
-void FDialogueContextStructCustomization::CustomizeHeader( TSharedRef<IPropertyHandle> StructPropertyHandle, class FDetailWidgetRow& HeaderRow, IStructCustomizationUtils& StructCustomizationUtils )
+void FDialogueContextStructCustomization::CustomizeHeader( TSharedRef<IPropertyHandle> StructPropertyHandle, class FDetailWidgetRow& HeaderRow, IPropertyTypeCustomizationUtils& StructCustomizationUtils )
 {
 	if( StructPropertyHandle->IsValidHandle() )
 	{
@@ -39,7 +39,7 @@ void FDialogueContextStructCustomization::CustomizeHeader( TSharedRef<IPropertyH
 	}
 }
 
-void FDialogueContextStructCustomization::CustomizeChildren( TSharedRef<IPropertyHandle> StructPropertyHandle, class IDetailChildrenBuilder& ChildBuilder, IStructCustomizationUtils& StructCustomizationUtils )
+void FDialogueContextStructCustomization::CustomizeChildren( TSharedRef<IPropertyHandle> StructPropertyHandle, class IDetailChildrenBuilder& ChildBuilder, IPropertyTypeCustomizationUtils& StructCustomizationUtils )
 {
 	if( StructPropertyHandle->IsValidHandle() )
 	{
@@ -1009,11 +1009,11 @@ TSharedRef<IPropertyTypeCustomization> FDialogueWaveParameterStructCustomization
 	return MakeShareable( new FDialogueWaveParameterStructCustomization );
 }
 
-void FDialogueWaveParameterStructCustomization::CustomizeHeader( TSharedRef<IPropertyHandle> StructPropertyHandle, class FDetailWidgetRow& HeaderRow, IStructCustomizationUtils& StructCustomizationUtils )
+void FDialogueWaveParameterStructCustomization::CustomizeHeader( TSharedRef<IPropertyHandle> StructPropertyHandle, class FDetailWidgetRow& HeaderRow, IPropertyTypeCustomizationUtils& StructCustomizationUtils )
 {
 }
 
-void FDialogueWaveParameterStructCustomization::CustomizeChildren( TSharedRef<IPropertyHandle> StructPropertyHandle, class IDetailChildrenBuilder& ChildBuilder, IStructCustomizationUtils& StructCustomizationUtils )
+void FDialogueWaveParameterStructCustomization::CustomizeChildren( TSharedRef<IPropertyHandle> StructPropertyHandle, class IDetailChildrenBuilder& ChildBuilder, IPropertyTypeCustomizationUtils& StructCustomizationUtils )
 {
 	if( StructPropertyHandle->IsValidHandle() )
 	{
diff --git a/Engine/Source/Editor/DetailCustomizations/Private/DialogueStructsCustomizations.h b/Engine/Source/Editor/DetailCustomizations/Private/DialogueStructsCustomizations.h
--- a/Engine/Source/Editor/DetailCustomizations/Private/DialogueStructsCustomizations.h
+++ b/Engine/Source/Editor/DetailCustomizations/Private/DialogueStructsCustomizations.h
@@ -9,8 +9,8 @@ public:
 	static TSharedRef<IPropertyTypeCustomization> MakeInstance();
 
 	/** IPropertyTypeCustomization interface begin */
-	virtual void CustomizeHeader( TSharedRef<IPropertyHandle> StructPropertyHandle, class FDetailWidgetRow& HeaderRow, IStructCustomizationUtils& StructCustomizationUtils ) OVERRIDE;
-	virtual void CustomizeChildren( TSharedRef<IPropertyHandle> StructPropertyHandle, class IDetailChildrenBuilder& ChildBuilder, IStructCustomizationUtils& StructCustomizationUtils ) OVERRIDE;
+	virtual void CustomizeHeader( TSharedRef<IPropertyHandle> StructPropertyHandle, class FDetailWidgetRow& HeaderRow, IPropertyTypeCustomizationUtils& StructCustomizationUtils ) OVERRIDE;
+	virtual void CustomizeChildren( TSharedRef<IPropertyHandle> StructPropertyHandle, class IDetailChildrenBuilder& ChildBuilder, IPropertyTypeCustomizationUtils& StructCustomizationUtils ) OVERRIDE;
 	/** IPropertyTypeCustomization interface end */
 };
 
@@ -21,7 +21,7 @@ public:
 	static TSharedRef<IPropertyTypeCustomization> MakeInstance();
 
 	/** IPropertyTypeCustomization interface begin */
-	virtual void CustomizeHeader( TSharedRef<IPropertyHandle> StructPropertyHandle, class FDetailWidgetRow& HeaderRow, IStructCustomizationUtils& StructCustomizationUtils ) OVERRIDE;
-	virtual void CustomizeChildren( TSharedRef<IPropertyHandle> StructPropertyHandle, class IDetailChildrenBuilder& ChildBuilder, IStructCustomizationUtils& StructCustomizationUtils ) OVERRIDE;
+	virtual void CustomizeHeader( TSharedRef<IPropertyHandle> StructPropertyHandle, class FDetailWidgetRow& HeaderRow, IPropertyTypeCustomizationUtils& StructCustomizationUtils ) OVERRIDE;
+	virtual void CustomizeChildren( TSharedRef<IPropertyHandle> StructPropertyHandle, class IDetailChildrenBuilder& ChildBuilder, IPropertyTypeCustomizationUtils& StructCustomizationUtils ) OVERRIDE;
 	/** IPropertyTypeCustomization interface end */
 };
\ No newline at end of file
diff --git a/Engine/Source/Editor/DetailCustomizations/Private/DirectoryPathStructCustomization.cpp b/Engine/Source/Editor/DetailCustomizations/Private/DirectoryPathStructCustomization.cpp
--- a/Engine/Source/Editor/DetailCustomizations/Private/DirectoryPathStructCustomization.cpp
+++ b/Engine/Source/Editor/DetailCustomizations/Private/DirectoryPathStructCustomization.cpp
@@ -12,7 +12,7 @@ TSharedRef<IPropertyTypeCustomization> FDirectoryPathStructCustomization::MakeIn
 	return MakeShareable(new FDirectoryPathStructCustomization());
 }
 
-void FDirectoryPathStructCustomization::CustomizeHeader( TSharedRef<IPropertyHandle> StructPropertyHandle, class FDetailWidgetRow& HeaderRow, IStructCustomizationUtils& StructCustomizationUtils )
+void FDirectoryPathStructCustomization::CustomizeHeader( TSharedRef<IPropertyHandle> StructPropertyHandle, class FDetailWidgetRow& HeaderRow, IPropertyTypeCustomizationUtils& StructCustomizationUtils )
 {
 	TSharedPtr<IPropertyHandle> PathProperty = StructPropertyHandle->GetChildHandle("Path");
 
@@ -60,7 +60,7 @@ void FDirectoryPathStructCustomization::CustomizeHeader( TSharedRef<IPropertyHan
 	}
 }
 
-void FDirectoryPathStructCustomization::CustomizeChildren( TSharedRef<IPropertyHandle> StructPropertyHandle, class IDetailChildrenBuilder& StructBuilder, IStructCustomizationUtils& StructCustomizationUtils )
+void FDirectoryPathStructCustomization::CustomizeChildren( TSharedRef<IPropertyHandle> StructPropertyHandle, class IDetailChildrenBuilder& StructBuilder, IPropertyTypeCustomizationUtils& StructCustomizationUtils )
 {
 }
 
diff --git a/Engine/Source/Editor/DetailCustomizations/Private/DirectoryPathStructCustomization.h b/Engine/Source/Editor/DetailCustomizations/Private/DirectoryPathStructCustomization.h
--- a/Engine/Source/Editor/DetailCustomizations/Private/DirectoryPathStructCustomization.h
+++ b/Engine/Source/Editor/DetailCustomizations/Private/DirectoryPathStructCustomization.h
@@ -8,9 +8,9 @@ public:
 	static TSharedRef<IPropertyTypeCustomization> MakeInstance();
 
 	/** IPropertyTypeCustomization interface */
-	virtual void CustomizeHeader( TSharedRef<class IPropertyHandle> StructPropertyHandle, class FDetailWidgetRow& HeaderRow, IStructCustomizationUtils& StructCustomizationUtils ) OVERRIDE;
+	virtual void CustomizeHeader( TSharedRef<class IPropertyHandle> StructPropertyHandle, class FDetailWidgetRow& HeaderRow, IPropertyTypeCustomizationUtils& StructCustomizationUtils ) OVERRIDE;
 
-	virtual void CustomizeChildren( TSharedRef<class IPropertyHandle> StructPropertyHandle, class IDetailChildrenBuilder& StructBuilder, IStructCustomizationUtils& StructCustomizationUtils ) OVERRIDE;              
+	virtual void CustomizeChildren( TSharedRef<class IPropertyHandle> StructPropertyHandle, class IDetailChildrenBuilder& StructBuilder, IPropertyTypeCustomizationUtils& StructCustomizationUtils ) OVERRIDE;              
 
 private:
 
diff --git a/Engine/Source/Editor/DetailCustomizations/Private/EnvQueryParamInstanceCustomization.cpp b/Engine/Source/Editor/DetailCustomizations/Private/EnvQueryParamInstanceCustomization.cpp
--- a/Engine/Source/Editor/DetailCustomizations/Private/EnvQueryParamInstanceCustomization.cpp
+++ b/Engine/Source/Editor/DetailCustomizations/Private/EnvQueryParamInstanceCustomization.cpp
@@ -10,7 +10,7 @@ TSharedRef<IPropertyTypeCustomization> FEnvQueryParamInstanceCustomization::Make
 	return MakeShareable(new FEnvQueryParamInstanceCustomization);
 }
 
-void FEnvQueryParamInstanceCustomization::CustomizeHeader( TSharedRef<class IPropertyHandle> StructPropertyHandle, class FDetailWidgetRow& HeaderRow, IStructCustomizationUtils& StructCustomizationUtils )
+void FEnvQueryParamInstanceCustomization::CustomizeHeader( TSharedRef<class IPropertyHandle> StructPropertyHandle, class FDetailWidgetRow& HeaderRow, IPropertyTypeCustomizationUtils& StructCustomizationUtils )
 {
 	NameProp = StructPropertyHandle->GetChildHandle(GET_MEMBER_NAME_CHECKED(FEnvNamedValue,ParamName));
 	TypeProp = StructPropertyHandle->GetChildHandle(GET_MEMBER_NAME_CHECKED(FEnvNamedValue,ParamType));
@@ -34,7 +34,7 @@ void FEnvQueryParamInstanceCustomization::CustomizeHeader( TSharedRef<class IPro
 	];
 }
 
-void FEnvQueryParamInstanceCustomization::CustomizeChildren( TSharedRef<class IPropertyHandle> StructPropertyHandle, class IDetailChildrenBuilder& StructBuilder, IStructCustomizationUtils& StructCustomizationUtils )
+void FEnvQueryParamInstanceCustomization::CustomizeChildren( TSharedRef<class IPropertyHandle> StructPropertyHandle, class IDetailChildrenBuilder& StructBuilder, IPropertyTypeCustomizationUtils& StructCustomizationUtils )
 {
 	StructBuilder.AddChildProperty(NameProp.ToSharedRef());
 	StructBuilder.AddChildProperty(TypeProp.ToSharedRef());
diff --git a/Engine/Source/Editor/DetailCustomizations/Private/EnvQueryParamInstanceCustomization.h b/Engine/Source/Editor/DetailCustomizations/Private/EnvQueryParamInstanceCustomization.h
--- a/Engine/Source/Editor/DetailCustomizations/Private/EnvQueryParamInstanceCustomization.h
+++ b/Engine/Source/Editor/DetailCustomizations/Private/EnvQueryParamInstanceCustomization.h
@@ -9,8 +9,8 @@ class FEnvQueryParamInstanceCustomization : public IPropertyTypeCustomization
 public:
 
 	// Begin IPropertyTypeCustomization interface
-	virtual void CustomizeHeader( TSharedRef<class IPropertyHandle> StructPropertyHandle, class FDetailWidgetRow& HeaderRow, IStructCustomizationUtils& StructCustomizationUtils ) OVERRIDE;
-	virtual void CustomizeChildren( TSharedRef<class IPropertyHandle> StructPropertyHandle, class IDetailChildrenBuilder& StructBuilder, IStructCustomizationUtils& StructCustomizationUtils ) OVERRIDE;
+	virtual void CustomizeHeader( TSharedRef<class IPropertyHandle> StructPropertyHandle, class FDetailWidgetRow& HeaderRow, IPropertyTypeCustomizationUtils& StructCustomizationUtils ) OVERRIDE;
+	virtual void CustomizeChildren( TSharedRef<class IPropertyHandle> StructPropertyHandle, class IDetailChildrenBuilder& StructBuilder, IPropertyTypeCustomizationUtils& StructCustomizationUtils ) OVERRIDE;
 	// End IPropertyTypeCustomization interface
 
 	static TSharedRef<IPropertyTypeCustomization> MakeInstance( );
diff --git a/Engine/Source/Editor/DetailCustomizations/Private/FilePathStructCustomization.cpp b/Engine/Source/Editor/DetailCustomizations/Private/FilePathStructCustomization.cpp
--- a/Engine/Source/Editor/DetailCustomizations/Private/FilePathStructCustomization.cpp
+++ b/Engine/Source/Editor/DetailCustomizations/Private/FilePathStructCustomization.cpp
@@ -105,7 +105,7 @@ TSharedRef<IPropertyTypeCustomization> FFilePathStructCustomization::MakeInstanc
 	return MakeShareable(new FFilePathStructCustomization());
 }
 
-void FFilePathStructCustomization::CustomizeHeader( TSharedRef<IPropertyHandle> StructPropertyHandle, class FDetailWidgetRow& HeaderRow, IStructCustomizationUtils& StructCustomizationUtils )
+void FFilePathStructCustomization::CustomizeHeader( TSharedRef<IPropertyHandle> StructPropertyHandle, class FDetailWidgetRow& HeaderRow, IPropertyTypeCustomizationUtils& StructCustomizationUtils )
 {
 	const UProperty* Property = StructPropertyHandle->GetProperty();
 	checkSlow(Property);
@@ -140,7 +140,7 @@ TSharedPtr<SWidget> FFilePathStructCustomization::CreatePickerWidget(TSharedRef<
 	return TSharedPtr<SWidget>();
 }
 
-void FFilePathStructCustomization::CustomizeChildren( TSharedRef<IPropertyHandle> StructPropertyHandle, class IDetailChildrenBuilder& StructBuilder, IStructCustomizationUtils& StructCustomizationUtils )
+void FFilePathStructCustomization::CustomizeChildren( TSharedRef<IPropertyHandle> StructPropertyHandle, class IDetailChildrenBuilder& StructBuilder, IPropertyTypeCustomizationUtils& StructCustomizationUtils )
 {
 }
 
diff --git a/Engine/Source/Editor/DetailCustomizations/Private/FilePathStructCustomization.h b/Engine/Source/Editor/DetailCustomizations/Private/FilePathStructCustomization.h
--- a/Engine/Source/Editor/DetailCustomizations/Private/FilePathStructCustomization.h
+++ b/Engine/Source/Editor/DetailCustomizations/Private/FilePathStructCustomization.h
@@ -11,9 +11,9 @@ public:
 	static TSharedRef<IPropertyTypeCustomization> MakeInstance();
 
 	/** IPropertyTypeCustomization interface */
-	virtual void CustomizeHeader( TSharedRef<class IPropertyHandle> StructPropertyHandle, class FDetailWidgetRow& HeaderRow, IStructCustomizationUtils& StructCustomizationUtils ) OVERRIDE;
+	virtual void CustomizeHeader( TSharedRef<class IPropertyHandle> StructPropertyHandle, class FDetailWidgetRow& HeaderRow, IPropertyTypeCustomizationUtils& StructCustomizationUtils ) OVERRIDE;
 
-	virtual void CustomizeChildren( TSharedRef<class IPropertyHandle> StructPropertyHandle, class IDetailChildrenBuilder& StructBuilder, IStructCustomizationUtils& StructCustomizationUtils ) OVERRIDE;              
+	virtual void CustomizeChildren( TSharedRef<class IPropertyHandle> StructPropertyHandle, class IDetailChildrenBuilder& StructBuilder, IPropertyTypeCustomizationUtils& StructCustomizationUtils ) OVERRIDE;              
 
 	static TSharedPtr<SWidget> CreatePickerWidget(TSharedRef<class IPropertyHandle> StructPropertyHandle,  const FString& FileFilterExtension = FString(TEXT("*")), const FOnPathPicked& OnPathPicked = FOnPathPicked());
 
diff --git a/Engine/Source/Editor/DetailCustomizations/Private/GuidStructCustomization.cpp b/Engine/Source/Editor/DetailCustomizations/Private/GuidStructCustomization.cpp
--- a/Engine/Source/Editor/DetailCustomizations/Private/GuidStructCustomization.cpp
+++ b/Engine/Source/Editor/DetailCustomizations/Private/GuidStructCustomization.cpp
@@ -23,7 +23,7 @@ TSharedRef<IPropertyTypeCustomization> FGuidStructCustomization::MakeInstance( )
 /* IPropertyTypeCustomization interface
  *****************************************************************************/
 
-void FGuidStructCustomization::CustomizeHeader( TSharedRef<class IPropertyHandle> StructPropertyHandle, class FDetailWidgetRow& HeaderRow, IStructCustomizationUtils& StructCustomizationUtils )
+void FGuidStructCustomization::CustomizeHeader( TSharedRef<class IPropertyHandle> StructPropertyHandle, class FDetailWidgetRow& HeaderRow, IPropertyTypeCustomizationUtils& StructCustomizationUtils )
 {
 	PropertyHandle = StructPropertyHandle;
 	InputValid = true;
@@ -90,7 +90,7 @@ void FGuidStructCustomization::CustomizeHeader( TSharedRef<class IPropertyHandle
 }
 
 
-void FGuidStructCustomization::CustomizeChildren( TSharedRef<class IPropertyHandle> StructPropertyHandle, class IDetailChildrenBuilder& StructBuilder, IStructCustomizationUtils& StructCustomizationUtils )
+void FGuidStructCustomization::CustomizeChildren( TSharedRef<class IPropertyHandle> StructPropertyHandle, class IDetailChildrenBuilder& StructBuilder, IPropertyTypeCustomizationUtils& StructCustomizationUtils )
 {
 	// do nothing
 }
diff --git a/Engine/Source/Editor/DetailCustomizations/Private/GuidStructCustomization.h b/Engine/Source/Editor/DetailCustomizations/Private/GuidStructCustomization.h
--- a/Engine/Source/Editor/DetailCustomizations/Private/GuidStructCustomization.h
+++ b/Engine/Source/Editor/DetailCustomizations/Private/GuidStructCustomization.h
@@ -33,9 +33,9 @@ public:
 
 	// Begin IPropertyTypeCustomization interface
 
-	virtual void CustomizeHeader( TSharedRef<class IPropertyHandle> StructPropertyHandle, class FDetailWidgetRow& HeaderRow, IStructCustomizationUtils& StructCustomizationUtils ) OVERRIDE;
+	virtual void CustomizeHeader( TSharedRef<class IPropertyHandle> StructPropertyHandle, class FDetailWidgetRow& HeaderRow, IPropertyTypeCustomizationUtils& StructCustomizationUtils ) OVERRIDE;
 
-	virtual void CustomizeChildren( TSharedRef<class IPropertyHandle> StructPropertyHandle, class IDetailChildrenBuilder& StructBuilder, IStructCustomizationUtils& StructCustomizationUtils ) OVERRIDE;
+	virtual void CustomizeChildren( TSharedRef<class IPropertyHandle> StructPropertyHandle, class IDetailChildrenBuilder& StructBuilder, IPropertyTypeCustomizationUtils& StructCustomizationUtils ) OVERRIDE;
 
 	// End IPropertyTypeCustomization interface
 
diff --git a/Engine/Source/Editor/DetailCustomizations/Private/InputStructCustomization.cpp b/Engine/Source/Editor/DetailCustomizations/Private/InputStructCustomization.cpp
--- a/Engine/Source/Editor/DetailCustomizations/Private/InputStructCustomization.cpp
+++ b/Engine/Source/Editor/DetailCustomizations/Private/InputStructCustomization.cpp
@@ -15,7 +15,7 @@ TSharedRef<IPropertyTypeCustomization> FInputAxisConfigCustomization::MakeInstan
 	return MakeShareable( new FInputAxisConfigCustomization );
 }
 
-void FInputAxisConfigCustomization::CustomizeHeader( TSharedRef<class IPropertyHandle> InStructPropertyHandle, class FDetailWidgetRow& HeaderRow, IStructCustomizationUtils& StructCustomizationUtils )
+void FInputAxisConfigCustomization::CustomizeHeader( TSharedRef<class IPropertyHandle> InStructPropertyHandle, class FDetailWidgetRow& HeaderRow, IPropertyTypeCustomizationUtils& StructCustomizationUtils )
 {
 	FString AxisKeyName;
 	InStructPropertyHandle->GetChildHandle(GET_MEMBER_NAME_CHECKED(FInputAxisConfigEntry, AxisKeyName))->GetValue(AxisKeyName);
@@ -26,7 +26,7 @@ void FInputAxisConfigCustomization::CustomizeHeader( TSharedRef<class IPropertyH
 	];
 }
 
-void FInputAxisConfigCustomization::CustomizeChildren( TSharedRef<class IPropertyHandle> InStructPropertyHandle, class IDetailChildrenBuilder& StructBuilder, IStructCustomizationUtils& StructCustomizationUtils )
+void FInputAxisConfigCustomization::CustomizeChildren( TSharedRef<class IPropertyHandle> InStructPropertyHandle, class IDetailChildrenBuilder& StructBuilder, IPropertyTypeCustomizationUtils& StructCustomizationUtils )
 {
 	TSharedPtr<IPropertyHandle> AxisProperties = InStructPropertyHandle->GetChildHandle(GET_MEMBER_NAME_CHECKED(FInputAxisConfigEntry, AxisProperties));
 
@@ -48,12 +48,12 @@ TSharedRef<IPropertyTypeCustomization> FInputActionMappingCustomization::MakeIns
 	return MakeShareable( new FInputActionMappingCustomization );
 }
 
-void FInputActionMappingCustomization::CustomizeHeader( TSharedRef<class IPropertyHandle> InStructPropertyHandle, class FDetailWidgetRow& HeaderRow, IStructCustomizationUtils& StructCustomizationUtils )
+void FInputActionMappingCustomization::CustomizeHeader( TSharedRef<class IPropertyHandle> InStructPropertyHandle, class FDetailWidgetRow& HeaderRow, IPropertyTypeCustomizationUtils& StructCustomizationUtils )
 {
 	ActionMappingHandle = InStructPropertyHandle;
 }
 
-void FInputActionMappingCustomization::CustomizeChildren( TSharedRef<class IPropertyHandle> InStructPropertyHandle, class IDetailChildrenBuilder& StructBuilder, IStructCustomizationUtils& StructCustomizationUtils )
+void FInputActionMappingCustomization::CustomizeChildren( TSharedRef<class IPropertyHandle> InStructPropertyHandle, class IDetailChildrenBuilder& StructBuilder, IPropertyTypeCustomizationUtils& StructCustomizationUtils )
 {
 	TSharedPtr<IPropertyHandle> KeyHandle = InStructPropertyHandle->GetChildHandle(GET_MEMBER_NAME_CHECKED(FInputActionKeyMapping, Key));
 	TSharedPtr<IPropertyHandle> ShiftHandle = InStructPropertyHandle->GetChildHandle(GET_MEMBER_NAME_CHECKED(FInputActionKeyMapping, bShift));
@@ -155,12 +155,12 @@ TSharedRef<IPropertyTypeCustomization> FInputAxisMappingCustomization::MakeInsta
 	return MakeShareable( new FInputAxisMappingCustomization );
 }
 
-void FInputAxisMappingCustomization::CustomizeHeader( TSharedRef<class IPropertyHandle> InStructPropertyHandle, class FDetailWidgetRow& HeaderRow, IStructCustomizationUtils& StructCustomizationUtils )
+void FInputAxisMappingCustomization::CustomizeHeader( TSharedRef<class IPropertyHandle> InStructPropertyHandle, class FDetailWidgetRow& HeaderRow, IPropertyTypeCustomizationUtils& StructCustomizationUtils )
 {
 	AxisMappingHandle = InStructPropertyHandle;
 }
 
-void FInputAxisMappingCustomization::CustomizeChildren( TSharedRef<class IPropertyHandle> InStructPropertyHandle, class IDetailChildrenBuilder& StructBuilder, IStructCustomizationUtils& StructCustomizationUtils )
+void FInputAxisMappingCustomization::CustomizeChildren( TSharedRef<class IPropertyHandle> InStructPropertyHandle, class IDetailChildrenBuilder& StructBuilder, IPropertyTypeCustomizationUtils& StructCustomizationUtils )
 {
 	TSharedPtr<IPropertyHandle> KeyHandle = InStructPropertyHandle->GetChildHandle(GET_MEMBER_NAME_CHECKED(FInputAxisKeyMapping, Key));
 	TSharedPtr<IPropertyHandle> ScaleHandle = InStructPropertyHandle->GetChildHandle(GET_MEMBER_NAME_CHECKED(FInputAxisKeyMapping, Scale));
diff --git a/Engine/Source/Editor/DetailCustomizations/Private/InputStructCustomization.h b/Engine/Source/Editor/DetailCustomizations/Private/InputStructCustomization.h
--- a/Engine/Source/Editor/DetailCustomizations/Private/InputStructCustomization.h
+++ b/Engine/Source/Editor/DetailCustomizations/Private/InputStructCustomization.h
@@ -11,8 +11,8 @@ public:
 	static TSharedRef<IPropertyTypeCustomization> MakeInstance();
 
 	/** IPropertyTypeCustomization interface */
-	virtual void CustomizeHeader( TSharedRef<class IPropertyHandle> InStructPropertyHandle, class FDetailWidgetRow& HeaderRow, IStructCustomizationUtils& StructCustomizationUtils ) OVERRIDE;
-	virtual void CustomizeChildren( TSharedRef<class IPropertyHandle> InStructPropertyHandle, class IDetailChildrenBuilder& StructBuilder, IStructCustomizationUtils& StructCustomizationUtils ) OVERRIDE;
+	virtual void CustomizeHeader( TSharedRef<class IPropertyHandle> InStructPropertyHandle, class FDetailWidgetRow& HeaderRow, IPropertyTypeCustomizationUtils& StructCustomizationUtils ) OVERRIDE;
+	virtual void CustomizeChildren( TSharedRef<class IPropertyHandle> InStructPropertyHandle, class IDetailChildrenBuilder& StructBuilder, IPropertyTypeCustomizationUtils& StructCustomizationUtils ) OVERRIDE;
 };
 
 /**
@@ -24,8 +24,8 @@ public:
 	static TSharedRef<IPropertyTypeCustomization> MakeInstance();
 
 	/** IPropertyTypeCustomization interface */
-	virtual void CustomizeHeader( TSharedRef<class IPropertyHandle> InStructPropertyHandle, class FDetailWidgetRow& HeaderRow, IStructCustomizationUtils& StructCustomizationUtils ) OVERRIDE;
-	virtual void CustomizeChildren( TSharedRef<class IPropertyHandle> InStructPropertyHandle, class IDetailChildrenBuilder& StructBuilder, IStructCustomizationUtils& StructCustomizationUtils ) OVERRIDE;
+	virtual void CustomizeHeader( TSharedRef<class IPropertyHandle> InStructPropertyHandle, class FDetailWidgetRow& HeaderRow, IPropertyTypeCustomizationUtils& StructCustomizationUtils ) OVERRIDE;
+	virtual void CustomizeChildren( TSharedRef<class IPropertyHandle> InStructPropertyHandle, class IDetailChildrenBuilder& StructBuilder, IPropertyTypeCustomizationUtils& StructCustomizationUtils ) OVERRIDE;
 
 private:
 	void RemoveActionMappingButton_OnClick();
@@ -43,8 +43,8 @@ public:
 	static TSharedRef<IPropertyTypeCustomization> MakeInstance();
 
 	/** IPropertyTypeCustomization interface */
-	virtual void CustomizeHeader( TSharedRef<class IPropertyHandle> InStructPropertyHandle, class FDetailWidgetRow& HeaderRow, IStructCustomizationUtils& StructCustomizationUtils ) OVERRIDE;
-	virtual void CustomizeChildren( TSharedRef<class IPropertyHandle> InStructPropertyHandle, class IDetailChildrenBuilder& StructBuilder, IStructCustomizationUtils& StructCustomizationUtils ) OVERRIDE;
+	virtual void CustomizeHeader( TSharedRef<class IPropertyHandle> InStructPropertyHandle, class FDetailWidgetRow& HeaderRow, IPropertyTypeCustomizationUtils& StructCustomizationUtils ) OVERRIDE;
+	virtual void CustomizeChildren( TSharedRef<class IPropertyHandle> InStructPropertyHandle, class IDetailChildrenBuilder& StructBuilder, IPropertyTypeCustomizationUtils& StructCustomizationUtils ) OVERRIDE;
 
 private:
 	void RemoveAxisMappingButton_OnClick();
diff --git a/Engine/Source/Editor/DetailCustomizations/Private/KeyStructCustomization.cpp b/Engine/Source/Editor/DetailCustomizations/Private/KeyStructCustomization.cpp
--- a/Engine/Source/Editor/DetailCustomizations/Private/KeyStructCustomization.cpp
+++ b/Engine/Source/Editor/DetailCustomizations/Private/KeyStructCustomization.cpp
@@ -23,7 +23,7 @@ TSharedRef<IPropertyTypeCustomization> FKeyStructCustomization::MakeInstance( )
 /* IPropertyTypeCustomization interface
  *****************************************************************************/
 
-void FKeyStructCustomization::CustomizeHeader( TSharedRef<class IPropertyHandle> StructPropertyHandle, class FDetailWidgetRow& HeaderRow, IStructCustomizationUtils& StructCustomizationUtils )
+void FKeyStructCustomization::CustomizeHeader( TSharedRef<class IPropertyHandle> StructPropertyHandle, class FDetailWidgetRow& HeaderRow, IPropertyTypeCustomizationUtils& StructCustomizationUtils )
 {
 	PropertyHandle = StructPropertyHandle;
 
diff --git a/Engine/Source/Editor/DetailCustomizations/Private/KeyStructCustomization.h b/Engine/Source/Editor/DetailCustomizations/Private/KeyStructCustomization.h
--- a/Engine/Source/Editor/DetailCustomizations/Private/KeyStructCustomization.h
+++ b/Engine/Source/Editor/DetailCustomizations/Private/KeyStructCustomization.h
@@ -16,8 +16,8 @@ public:
 
 	// Begin IPropertyTypeCustomization interface
 
-	virtual void CustomizeHeader(TSharedRef<class IPropertyHandle> StructPropertyHandle, class FDetailWidgetRow& HeaderRow, IStructCustomizationUtils& StructCustomizationUtils) OVERRIDE;
-	virtual void CustomizeChildren(TSharedRef<IPropertyHandle> StructPropertyHandle, class IDetailChildrenBuilder& ChildBuilder, IStructCustomizationUtils& StructCustomizationUtils) OVERRIDE {};
+	virtual void CustomizeHeader(TSharedRef<class IPropertyHandle> StructPropertyHandle, class FDetailWidgetRow& HeaderRow, IPropertyTypeCustomizationUtils& StructCustomizationUtils) OVERRIDE;
+	virtual void CustomizeChildren(TSharedRef<IPropertyHandle> StructPropertyHandle, class IDetailChildrenBuilder& ChildBuilder, IPropertyTypeCustomizationUtils& StructCustomizationUtils) OVERRIDE {};
 
 	// End IPropertyTypeCustomization interface
 
diff --git a/Engine/Source/Editor/DetailCustomizations/Private/MarginCustomization.cpp b/Engine/Source/Editor/DetailCustomizations/Private/MarginCustomization.cpp
--- a/Engine/Source/Editor/DetailCustomizations/Private/MarginCustomization.cpp
+++ b/Engine/Source/Editor/DetailCustomizations/Private/MarginCustomization.cpp
@@ -8,7 +8,7 @@ TSharedRef<IPropertyTypeCustomization> FMarginStructCustomization::MakeInstance(
 	return MakeShareable( new FMarginStructCustomization() );
 }
 
-void FMarginStructCustomization::CustomizeHeader( TSharedRef<class IPropertyHandle> StructPropertyHandle, class FDetailWidgetRow& HeaderRow, IStructCustomizationUtils& StructCustomizationUtils )
+void FMarginStructCustomization::CustomizeHeader( TSharedRef<class IPropertyHandle> StructPropertyHandle, class FDetailWidgetRow& HeaderRow, IPropertyTypeCustomizationUtils& StructCustomizationUtils )
 {
 	const FString UVSpaceString( StructPropertyHandle->GetProperty()->GetMetaData( TEXT( "UVSpace" ) ) );
 	bIsMarginUsingUVSpace = UVSpaceString.Len() > 0 && UVSpaceString == TEXT( "true" );
@@ -42,7 +42,7 @@ void FMarginStructCustomization::CustomizeHeader( TSharedRef<class IPropertyHand
 	UpdateMarginTextFromProperties();
 }
 
-void FMarginStructCustomization::CustomizeChildren( TSharedRef<class IPropertyHandle> StructPropertyHandle, class IDetailChildrenBuilder& StructBuilder, IStructCustomizationUtils& StructCustomizationUtils )
+void FMarginStructCustomization::CustomizeChildren( TSharedRef<class IPropertyHandle> StructPropertyHandle, class IDetailChildrenBuilder& StructBuilder, IPropertyTypeCustomizationUtils& StructCustomizationUtils )
 {
 	FSimpleDelegate OnPropertyChangedDelegate = FSimpleDelegate::CreateSP( this, &FMarginStructCustomization::OnPropertyChanged );
 	StructPropertyHandle->SetOnPropertyValueChanged( OnPropertyChangedDelegate );
diff --git a/Engine/Source/Editor/DetailCustomizations/Private/MarginCustomization.h b/Engine/Source/Editor/DetailCustomizations/Private/MarginCustomization.h
--- a/Engine/Source/Editor/DetailCustomizations/Private/MarginCustomization.h
+++ b/Engine/Source/Editor/DetailCustomizations/Private/MarginCustomization.h
@@ -8,9 +8,9 @@ public:
 	static TSharedRef<IPropertyTypeCustomization> MakeInstance();
 
 	/** IPropertyTypeCustomization interface */
-	virtual void CustomizeHeader( TSharedRef<class IPropertyHandle> StructPropertyHandle, class FDetailWidgetRow& HeaderRow, IStructCustomizationUtils& StructCustomizationUtils ) OVERRIDE;
+	virtual void CustomizeHeader( TSharedRef<class IPropertyHandle> StructPropertyHandle, class FDetailWidgetRow& HeaderRow, IPropertyTypeCustomizationUtils& StructCustomizationUtils ) OVERRIDE;
 
-	virtual void CustomizeChildren( TSharedRef<class IPropertyHandle> StructPropertyHandle, class IDetailChildrenBuilder& StructBuilder, IStructCustomizationUtils& StructCustomizationUtils ) OVERRIDE;              
+	virtual void CustomizeChildren( TSharedRef<class IPropertyHandle> StructPropertyHandle, class IDetailChildrenBuilder& StructBuilder, IPropertyTypeCustomizationUtils& StructCustomizationUtils ) OVERRIDE;              
 
 private:
 
diff --git a/Engine/Source/Editor/DetailCustomizations/Private/MathStructCustomizations.cpp b/Engine/Source/Editor/DetailCustomizations/Private/MathStructCustomizations.cpp
--- a/Engine/Source/Editor/DetailCustomizations/Private/MathStructCustomizations.cpp
+++ b/Engine/Source/Editor/DetailCustomizations/Private/MathStructCustomizations.cpp
@@ -11,14 +11,14 @@ TSharedRef<IPropertyTypeCustomization> FMathStructCustomization::MakeInstance()
 	return MakeShareable( new FMathStructCustomization );
 }
 
-void FMathStructCustomization::CustomizeHeader( TSharedRef<class IPropertyHandle> StructPropertyHandle, class FDetailWidgetRow& HeaderRow, IStructCustomizationUtils& StructCustomizationUtils )
+void FMathStructCustomization::CustomizeHeader( TSharedRef<class IPropertyHandle> StructPropertyHandle, class FDetailWidgetRow& HeaderRow, IPropertyTypeCustomizationUtils& StructCustomizationUtils )
 {
 	GetSortedChildren( StructPropertyHandle, SortedChildHandles );
 
 	MakeHeaderRow( StructPropertyHandle, HeaderRow );
 }
 
-void FMathStructCustomization::CustomizeChildren( TSharedRef<class IPropertyHandle> StructPropertyHandle, class IDetailChildrenBuilder& StructBuilder, IStructCustomizationUtils& StructCustomizationUtils )
+void FMathStructCustomization::CustomizeChildren( TSharedRef<class IPropertyHandle> StructPropertyHandle, class IDetailChildrenBuilder& StructBuilder, IPropertyTypeCustomizationUtils& StructCustomizationUtils )
 {
 	for( int32 ChildIndex = 0; ChildIndex < SortedChildHandles.Num(); ++ChildIndex )
 	{
@@ -335,7 +335,7 @@ TSharedRef<IPropertyTypeCustomization> FColorStructCustomization::MakeInstance()
 	return MakeShareable( new FColorStructCustomization );
 }
 
-void FColorStructCustomization::CustomizeHeader( TSharedRef<class IPropertyHandle> InStructPropertyHandle, class FDetailWidgetRow& InHeaderRow, IStructCustomizationUtils& StructCustomizationUtils )
+void FColorStructCustomization::CustomizeHeader( TSharedRef<class IPropertyHandle> InStructPropertyHandle, class FDetailWidgetRow& InHeaderRow, IPropertyTypeCustomizationUtils& StructCustomizationUtils )
 {
 	TSharedPtr<FAssetThumbnailPool> Pool = StructCustomizationUtils.GetThumbnailPool();
 
diff --git a/Engine/Source/Editor/DetailCustomizations/Private/MathStructCustomizations.h b/Engine/Source/Editor/DetailCustomizations/Private/MathStructCustomizations.h
--- a/Engine/Source/Editor/DetailCustomizations/Private/MathStructCustomizations.h
+++ b/Engine/Source/Editor/DetailCustomizations/Private/MathStructCustomizations.h
@@ -18,8 +18,8 @@ public:
 	virtual ~FMathStructCustomization() {}
 
 	/** IPropertyTypeCustomization instance */
-	virtual void CustomizeHeader( TSharedRef<class IPropertyHandle> StructPropertyHandle, class FDetailWidgetRow& HeaderRow, IStructCustomizationUtils& StructCustomizationUtils ) OVERRIDE;
-	virtual void CustomizeChildren( TSharedRef<class IPropertyHandle> StructPropertyHandle, class IDetailChildrenBuilder& StructBuilder, IStructCustomizationUtils& StructCustomizationUtils ) OVERRIDE;
+	virtual void CustomizeHeader( TSharedRef<class IPropertyHandle> StructPropertyHandle, class FDetailWidgetRow& HeaderRow, IPropertyTypeCustomizationUtils& StructCustomizationUtils ) OVERRIDE;
+	virtual void CustomizeChildren( TSharedRef<class IPropertyHandle> StructPropertyHandle, class IDetailChildrenBuilder& StructBuilder, IPropertyTypeCustomizationUtils& StructCustomizationUtils ) OVERRIDE;
 	
 protected:
 	/**
@@ -151,7 +151,7 @@ private:
 	{}
 
 	/** FMathStructCustomization interface */
-	virtual void CustomizeHeader( TSharedRef<class IPropertyHandle> StructPropertyHandle, class FDetailWidgetRow& HeaderRow, IStructCustomizationUtils& StructCustomizationUtils ) OVERRIDE;
+	virtual void CustomizeHeader( TSharedRef<class IPropertyHandle> StructPropertyHandle, class FDetailWidgetRow& HeaderRow, IPropertyTypeCustomizationUtils& StructCustomizationUtils ) OVERRIDE;
 	virtual void MakeHeaderRow( TSharedRef<class IPropertyHandle>& InStructPropertyHandle, FDetailWidgetRow& Row ) OVERRIDE;
 	virtual void GetSortedChildren( TSharedRef<IPropertyHandle> StructPropertyHandle, TArray< TSharedRef<IPropertyHandle> >& OutChildren ) OVERRIDE;
 
diff --git a/Engine/Source/Editor/DetailCustomizations/Private/MathStructProxyCustomizations.cpp b/Engine/Source/Editor/DetailCustomizations/Private/MathStructProxyCustomizations.cpp
--- a/Engine/Source/Editor/DetailCustomizations/Private/MathStructProxyCustomizations.cpp
+++ b/Engine/Source/Editor/DetailCustomizations/Private/MathStructProxyCustomizations.cpp
@@ -3,7 +3,7 @@
 #include "DetailCustomizationsPrivatePCH.h"
 #include "MathStructProxyCustomizations.h"
 
-void FMathStructProxyCustomization::CustomizeChildren( TSharedRef<class IPropertyHandle> StructPropertyHandle, class IDetailChildrenBuilder& StructBuilder, IStructCustomizationUtils& StructCustomizationUtils )
+void FMathStructProxyCustomization::CustomizeChildren( TSharedRef<class IPropertyHandle> StructPropertyHandle, class IDetailChildrenBuilder& StructBuilder, IPropertyTypeCustomizationUtils& StructCustomizationUtils )
 {
 
 }
diff --git a/Engine/Source/Editor/DetailCustomizations/Private/MathStructProxyCustomizations.h b/Engine/Source/Editor/DetailCustomizations/Private/MathStructProxyCustomizations.h
--- a/Engine/Source/Editor/DetailCustomizations/Private/MathStructProxyCustomizations.h
+++ b/Engine/Source/Editor/DetailCustomizations/Private/MathStructProxyCustomizations.h
@@ -138,7 +138,7 @@ class FMathStructProxyCustomization : public FMathStructCustomization
 {
 public:
 	/** IPropertyTypeCustomization interface */
-	virtual void CustomizeChildren( TSharedRef<class IPropertyHandle> StructPropertyHandle, class IDetailChildrenBuilder& StructBuilder, IStructCustomizationUtils& StructCustomizationUtils ) OVERRIDE;
+	virtual void CustomizeChildren( TSharedRef<class IPropertyHandle> StructPropertyHandle, class IDetailChildrenBuilder& StructBuilder, IPropertyTypeCustomizationUtils& StructCustomizationUtils ) OVERRIDE;
 
 	/** FMathStructCustomization interface */
 	virtual void MakeHeaderRow( TSharedRef<class IPropertyHandle>& InStructPropertyHandle, FDetailWidgetRow& Row ) OVERRIDE;
diff --git a/Engine/Source/Editor/DetailCustomizations/Private/NavLinkStructCustomization.cpp b/Engine/Source/Editor/DetailCustomizations/Private/NavLinkStructCustomization.cpp
--- a/Engine/Source/Editor/DetailCustomizations/Private/NavLinkStructCustomization.cpp
+++ b/Engine/Source/Editor/DetailCustomizations/Private/NavLinkStructCustomization.cpp
@@ -10,7 +10,7 @@ TSharedRef<IPropertyTypeCustomization> FNavLinkStructCustomization::MakeInstance
 	return MakeShareable(new FNavLinkStructCustomization);
 }
 
-void FNavLinkStructCustomization::CustomizeHeader( TSharedRef<IPropertyHandle> StructPropertyHandle, class FDetailWidgetRow& HeaderRow, IStructCustomizationUtils& StructCustomizationUtils )
+void FNavLinkStructCustomization::CustomizeHeader( TSharedRef<IPropertyHandle> StructPropertyHandle, class FDetailWidgetRow& HeaderRow, IPropertyTypeCustomizationUtils& StructCustomizationUtils )
 {
 	TSharedPtr<IPropertyHandle> CommentHandle = StructPropertyHandle->GetChildHandle(GET_MEMBER_NAME_CHECKED(FNavigationLinkBase,Description));
 	FString Desc;
@@ -33,7 +33,7 @@ void FNavLinkStructCustomization::CustomizeHeader( TSharedRef<IPropertyHandle> S
 	];
 }
 
-void FNavLinkStructCustomization::CustomizeChildren( TSharedRef<class IPropertyHandle> StructPropertyHandle, class IDetailChildrenBuilder& StructBuilder, IStructCustomizationUtils& StructCustomizationUtils )
+void FNavLinkStructCustomization::CustomizeChildren( TSharedRef<class IPropertyHandle> StructPropertyHandle, class IDetailChildrenBuilder& StructBuilder, IPropertyTypeCustomizationUtils& StructCustomizationUtils )
 {
 	uint32 NumChildren = 0;
 	StructPropertyHandle->GetNumChildren(NumChildren);
diff --git a/Engine/Source/Editor/DetailCustomizations/Private/NavLinkStructCustomization.h b/Engine/Source/Editor/DetailCustomizations/Private/NavLinkStructCustomization.h
--- a/Engine/Source/Editor/DetailCustomizations/Private/NavLinkStructCustomization.h
+++ b/Engine/Source/Editor/DetailCustomizations/Private/NavLinkStructCustomization.h
@@ -6,8 +6,8 @@ class FNavLinkStructCustomization : public IPropertyTypeCustomization
 {
 public:
 	// Begin IPropertyTypeCustomization interface
-	virtual void CustomizeHeader( TSharedRef<IPropertyHandle> StructPropertyHandle, class FDetailWidgetRow& HeaderRow, IStructCustomizationUtils& StructCustomizationUtils ) OVERRIDE;
-	virtual void CustomizeChildren( TSharedRef<class IPropertyHandle> StructPropertyHandle, class IDetailChildrenBuilder& StructBuilder, IStructCustomizationUtils& StructCustomizationUtils ) OVERRIDE;
+	virtual void CustomizeHeader( TSharedRef<IPropertyHandle> StructPropertyHandle, class FDetailWidgetRow& HeaderRow, IPropertyTypeCustomizationUtils& StructCustomizationUtils ) OVERRIDE;
+	virtual void CustomizeChildren( TSharedRef<class IPropertyHandle> StructPropertyHandle, class IDetailChildrenBuilder& StructBuilder, IPropertyTypeCustomizationUtils& StructCustomizationUtils ) OVERRIDE;
 	// End IPropertyTypeCustomization interface
 
 	static TSharedRef<IPropertyTypeCustomization> MakeInstance();
diff --git a/Engine/Source/Editor/DetailCustomizations/Private/RangeStructCustomization.cpp b/Engine/Source/Editor/DetailCustomizations/Private/RangeStructCustomization.cpp
--- a/Engine/Source/Editor/DetailCustomizations/Private/RangeStructCustomization.cpp
+++ b/Engine/Source/Editor/DetailCustomizations/Private/RangeStructCustomization.cpp
@@ -53,7 +53,7 @@ TSharedRef<IPropertyTypeCustomization> FRangeStructCustomization<NumericType>::M
  *****************************************************************************/
 
 template <typename NumericType>
-void FRangeStructCustomization<NumericType>::CustomizeHeader(TSharedRef<IPropertyHandle> StructPropertyHandle, FDetailWidgetRow& HeaderRow, IStructCustomizationUtils& StructCustomizationUtils)
+void FRangeStructCustomization<NumericType>::CustomizeHeader(TSharedRef<IPropertyHandle> StructPropertyHandle, FDetailWidgetRow& HeaderRow, IPropertyTypeCustomizationUtils& StructCustomizationUtils)
 {
 	// Get handles to the properties we're interested in
 	LowerBoundStructHandle = StructPropertyHandle->GetChildHandle(TEXT("LowerBound"));
@@ -221,7 +221,7 @@ void FRangeStructCustomization<NumericType>::CustomizeHeader(TSharedRef<IPropert
 
 
 template <typename NumericType>
-void FRangeStructCustomization<NumericType>::CustomizeChildren(TSharedRef<IPropertyHandle> StructPropertyHandle, IDetailChildrenBuilder& StructBuilder, IStructCustomizationUtils& StructCustomizationUtils)
+void FRangeStructCustomization<NumericType>::CustomizeChildren(TSharedRef<IPropertyHandle> StructPropertyHandle, IDetailChildrenBuilder& StructBuilder, IPropertyTypeCustomizationUtils& StructCustomizationUtils)
 {
 	// Don't display children, as editing them directly can break the constraints
 }
diff --git a/Engine/Source/Editor/DetailCustomizations/Private/RangeStructCustomization.h b/Engine/Source/Editor/DetailCustomizations/Private/RangeStructCustomization.h
--- a/Engine/Source/Editor/DetailCustomizations/Private/RangeStructCustomization.h
+++ b/Engine/Source/Editor/DetailCustomizations/Private/RangeStructCustomization.h
@@ -19,9 +19,9 @@ public:
 
 	// Begin IPropertyTypeCustomization interface
 
-	virtual void CustomizeHeader(TSharedRef<IPropertyHandle> StructPropertyHandle, FDetailWidgetRow& HeaderRow, IStructCustomizationUtils& StructCustomizationUtils) OVERRIDE;
+	virtual void CustomizeHeader(TSharedRef<IPropertyHandle> StructPropertyHandle, FDetailWidgetRow& HeaderRow, IPropertyTypeCustomizationUtils& StructCustomizationUtils) OVERRIDE;
 
-	virtual void CustomizeChildren(TSharedRef<IPropertyHandle> StructPropertyHandle, IDetailChildrenBuilder& StructBuilder, IStructCustomizationUtils& StructCustomizationUtils) OVERRIDE;
+	virtual void CustomizeChildren(TSharedRef<IPropertyHandle> StructPropertyHandle, IDetailChildrenBuilder& StructBuilder, IPropertyTypeCustomizationUtils& StructCustomizationUtils) OVERRIDE;
 
 	// End IPropertyTypeCustomization interface
 
diff --git a/Engine/Source/Editor/DetailCustomizations/Private/SlateBrushCustomization.cpp b/Engine/Source/Editor/DetailCustomizations/Private/SlateBrushCustomization.cpp
--- a/Engine/Source/Editor/DetailCustomizations/Private/SlateBrushCustomization.cpp
+++ b/Engine/Source/Editor/DetailCustomizations/Private/SlateBrushCustomization.cpp
@@ -938,7 +938,7 @@ TSharedRef<IPropertyTypeCustomization> FSlateBrushStructCustomization::MakeInsta
 	return MakeShareable( new FSlateBrushStructCustomization() );
 }
 
-void FSlateBrushStructCustomization::CustomizeHeader( TSharedRef<IPropertyHandle> StructPropertyHandle, FDetailWidgetRow& HeaderRow, IStructCustomizationUtils& StructCustomizationUtils )
+void FSlateBrushStructCustomization::CustomizeHeader( TSharedRef<IPropertyHandle> StructPropertyHandle, FDetailWidgetRow& HeaderRow, IPropertyTypeCustomizationUtils& StructCustomizationUtils )
 {
 	bool ShowOnlyInnerProperties = StructPropertyHandle->GetProperty()->HasMetaData(TEXT("ShowOnlyInnerProperties"));
 
@@ -952,7 +952,7 @@ void FSlateBrushStructCustomization::CustomizeHeader( TSharedRef<IPropertyHandle
 	}
 }
 
-void FSlateBrushStructCustomization::CustomizeChildren( TSharedRef<IPropertyHandle> StructPropertyHandle, IDetailChildrenBuilder& StructBuilder, IStructCustomizationUtils& StructCustomizationUtils )
+void FSlateBrushStructCustomization::CustomizeChildren( TSharedRef<IPropertyHandle> StructPropertyHandle, IDetailChildrenBuilder& StructBuilder, IPropertyTypeCustomizationUtils& StructCustomizationUtils )
 {
 	// Add the child properties
 	TSharedPtr<IPropertyHandle> ImageSizeProperty = StructPropertyHandle->GetChildHandle( TEXT("ImageSize") );
diff --git a/Engine/Source/Editor/DetailCustomizations/Private/SlateBrushCustomization.h b/Engine/Source/Editor/DetailCustomizations/Private/SlateBrushCustomization.h
--- a/Engine/Source/Editor/DetailCustomizations/Private/SlateBrushCustomization.h
+++ b/Engine/Source/Editor/DetailCustomizations/Private/SlateBrushCustomization.h
@@ -9,9 +9,9 @@ public:
 	static TSharedRef<IPropertyTypeCustomization> MakeInstance();
 
 	/** IPropertyTypeCustomization interface */
-	virtual void CustomizeHeader( TSharedRef<class IPropertyHandle> StructPropertyHandle, class FDetailWidgetRow& HeaderRow, IStructCustomizationUtils& StructCustomizationUtils ) OVERRIDE;
+	virtual void CustomizeHeader( TSharedRef<class IPropertyHandle> StructPropertyHandle, class FDetailWidgetRow& HeaderRow, IPropertyTypeCustomizationUtils& StructCustomizationUtils ) OVERRIDE;
 
-	virtual void CustomizeChildren( TSharedRef<class IPropertyHandle> StructPropertyHandle, class IDetailChildrenBuilder& StructBuilder, IStructCustomizationUtils& StructCustomizationUtils ) OVERRIDE;              
+	virtual void CustomizeChildren( TSharedRef<class IPropertyHandle> StructPropertyHandle, class IDetailChildrenBuilder& StructBuilder, IPropertyTypeCustomizationUtils& StructCustomizationUtils ) OVERRIDE;              
 
 private:
 	/**
diff --git a/Engine/Source/Editor/DetailCustomizations/Private/SlateColorCustomization.cpp b/Engine/Source/Editor/DetailCustomizations/Private/SlateColorCustomization.cpp
--- a/Engine/Source/Editor/DetailCustomizations/Private/SlateColorCustomization.cpp
+++ b/Engine/Source/Editor/DetailCustomizations/Private/SlateColorCustomization.cpp
@@ -8,7 +8,7 @@ TSharedRef<IPropertyTypeCustomization> FSlateColorCustomization::MakeInstance()
 	return MakeShareable( new FSlateColorCustomization() );
 }
 
-void FSlateColorCustomization::CustomizeHeader( TSharedRef<class IPropertyHandle> StructPropertyHandle, class FDetailWidgetRow& HeaderRow, IStructCustomizationUtils& StructCustomizationUtils )
+void FSlateColorCustomization::CustomizeHeader( TSharedRef<class IPropertyHandle> StructPropertyHandle, class FDetailWidgetRow& HeaderRow, IPropertyTypeCustomizationUtils& StructCustomizationUtils )
 {
 	uint32 NumChildren;
 	StructPropertyHandle->GetNumChildren( NumChildren );
@@ -31,7 +31,7 @@ void FSlateColorCustomization::CustomizeHeader( TSharedRef<class IPropertyHandle
 	check( SpecifiedColorHandle.IsValid() );
 }
 
-void FSlateColorCustomization::CustomizeChildren( TSharedRef<class IPropertyHandle> StructPropertyHandle, class IDetailChildrenBuilder& StructBuilder, IStructCustomizationUtils& StructCustomizationUtils )
+void FSlateColorCustomization::CustomizeChildren( TSharedRef<class IPropertyHandle> StructPropertyHandle, class IDetailChildrenBuilder& StructBuilder, IPropertyTypeCustomizationUtils& StructCustomizationUtils )
 {
 	SpecifiedColorHandle->SetOnPropertyValueChanged( FSimpleDelegate::CreateSP( SharedThis( this ), &FSlateColorCustomization::OnValueChanged ) );
 	StructBuilder.AddChildProperty( SpecifiedColorHandle.ToSharedRef() )
diff --git a/Engine/Source/Editor/DetailCustomizations/Private/SlateColorCustomization.h b/Engine/Source/Editor/DetailCustomizations/Private/SlateColorCustomization.h
--- a/Engine/Source/Editor/DetailCustomizations/Private/SlateColorCustomization.h
+++ b/Engine/Source/Editor/DetailCustomizations/Private/SlateColorCustomization.h
@@ -8,9 +8,9 @@ public:
 	static TSharedRef<IPropertyTypeCustomization> MakeInstance();
 
 	/** IPropertyTypeCustomization interface */
-	virtual void CustomizeHeader( TSharedRef<class IPropertyHandle> StructPropertyHandle, class FDetailWidgetRow& HeaderRow, IStructCustomizationUtils& StructCustomizationUtils ) OVERRIDE;
+	virtual void CustomizeHeader( TSharedRef<class IPropertyHandle> StructPropertyHandle, class FDetailWidgetRow& HeaderRow, IPropertyTypeCustomizationUtils& StructCustomizationUtils ) OVERRIDE;
 
-	virtual void CustomizeChildren( TSharedRef<class IPropertyHandle> StructPropertyHandle, class IDetailChildrenBuilder& StructBuilder, IStructCustomizationUtils& StructCustomizationUtils ) OVERRIDE;              
+	virtual void CustomizeChildren( TSharedRef<class IPropertyHandle> StructPropertyHandle, class IDetailChildrenBuilder& StructBuilder, IPropertyTypeCustomizationUtils& StructCustomizationUtils ) OVERRIDE;              
 
 private:
 
diff --git a/Engine/Source/Editor/DetailCustomizations/Private/SlateSoundCustomization.cpp b/Engine/Source/Editor/DetailCustomizations/Private/SlateSoundCustomization.cpp
--- a/Engine/Source/Editor/DetailCustomizations/Private/SlateSoundCustomization.cpp
+++ b/Engine/Source/Editor/DetailCustomizations/Private/SlateSoundCustomization.cpp
@@ -8,7 +8,7 @@ TSharedRef<IPropertyTypeCustomization> FSlateSoundStructCustomization::MakeInsta
 	return MakeShareable(new FSlateSoundStructCustomization());
 }
 
-void FSlateSoundStructCustomization::CustomizeHeader(TSharedRef<IPropertyHandle> StructPropertyHandle, FDetailWidgetRow& HeaderRow, IStructCustomizationUtils& StructCustomizationUtils)
+void FSlateSoundStructCustomization::CustomizeHeader(TSharedRef<IPropertyHandle> StructPropertyHandle, FDetailWidgetRow& HeaderRow, IPropertyTypeCustomizationUtils& StructCustomizationUtils)
 {
 	TSharedPtr<IPropertyHandle> ResourceObjectProperty = StructPropertyHandle->GetChildHandle(TEXT("ResourceObject"));
 	check(ResourceObjectProperty.IsValid());
@@ -36,7 +36,7 @@ void FSlateSoundStructCustomization::CustomizeHeader(TSharedRef<IPropertyHandle>
 	];
 }
 
-void FSlateSoundStructCustomization::CustomizeChildren(TSharedRef<IPropertyHandle> StructPropertyHandle, IDetailChildrenBuilder& StructBuilder, IStructCustomizationUtils& StructCustomizationUtils)
+void FSlateSoundStructCustomization::CustomizeChildren(TSharedRef<IPropertyHandle> StructPropertyHandle, IDetailChildrenBuilder& StructBuilder, IPropertyTypeCustomizationUtils& StructCustomizationUtils)
 {
 }
 
diff --git a/Engine/Source/Editor/DetailCustomizations/Private/SlateSoundCustomization.h b/Engine/Source/Editor/DetailCustomizations/Private/SlateSoundCustomization.h
--- a/Engine/Source/Editor/DetailCustomizations/Private/SlateSoundCustomization.h
+++ b/Engine/Source/Editor/DetailCustomizations/Private/SlateSoundCustomization.h
@@ -9,8 +9,8 @@ public:
 	static TSharedRef<IPropertyTypeCustomization> MakeInstance();
 
 	/** IPropertyTypeCustomization interface */
-	virtual void CustomizeHeader(TSharedRef<class IPropertyHandle> StructPropertyHandle, class FDetailWidgetRow& HeaderRow, IStructCustomizationUtils& StructCustomizationUtils) OVERRIDE;
-	virtual void CustomizeChildren(TSharedRef<class IPropertyHandle> StructPropertyHandle, class IDetailChildrenBuilder& StructBuilder, IStructCustomizationUtils& StructCustomizationUtils) OVERRIDE;
+	virtual void CustomizeHeader(TSharedRef<class IPropertyHandle> StructPropertyHandle, class FDetailWidgetRow& HeaderRow, IPropertyTypeCustomizationUtils& StructCustomizationUtils) OVERRIDE;
+	virtual void CustomizeChildren(TSharedRef<class IPropertyHandle> StructPropertyHandle, class IDetailChildrenBuilder& StructBuilder, IPropertyTypeCustomizationUtils& StructCustomizationUtils) OVERRIDE;
 
 protected:
 	/** Called when the resource object used by this FSlateSound has been changed */
diff --git a/Engine/Source/Editor/DetailCustomizations/Private/StringAssetReferenceCustomization.cpp b/Engine/Source/Editor/DetailCustomizations/Private/StringAssetReferenceCustomization.cpp
--- a/Engine/Source/Editor/DetailCustomizations/Private/StringAssetReferenceCustomization.cpp
+++ b/Engine/Source/Editor/DetailCustomizations/Private/StringAssetReferenceCustomization.cpp
@@ -4,7 +4,7 @@
 #include "StringAssetReferenceCustomization.h"
 #include "AssetData.h"
 
-void FStringAssetReferenceCustomization::CustomizeHeader( TSharedRef<IPropertyHandle> InStructPropertyHandle, FDetailWidgetRow& HeaderRow, IStructCustomizationUtils& StructCustomizationUtils )
+void FStringAssetReferenceCustomization::CustomizeHeader( TSharedRef<IPropertyHandle> InStructPropertyHandle, FDetailWidgetRow& HeaderRow, IPropertyTypeCustomizationUtils& StructCustomizationUtils )
 {
 	StructPropertyHandle = InStructPropertyHandle;
 
@@ -70,7 +70,7 @@ void FStringAssetReferenceCustomization::CustomizeHeader( TSharedRef<IPropertyHa
 	];
 }
 
-void FStringAssetReferenceCustomization::CustomizeChildren(TSharedRef<IPropertyHandle> InStructPropertyHandle, IDetailChildrenBuilder& StructBuilder, IStructCustomizationUtils& StructCustomizationUtils)
+void FStringAssetReferenceCustomization::CustomizeChildren(TSharedRef<IPropertyHandle> InStructPropertyHandle, IDetailChildrenBuilder& StructBuilder, IPropertyTypeCustomizationUtils& StructCustomizationUtils)
 {
 }
 
diff --git a/Engine/Source/Editor/DetailCustomizations/Private/StringAssetReferenceCustomization.h b/Engine/Source/Editor/DetailCustomizations/Private/StringAssetReferenceCustomization.h
--- a/Engine/Source/Editor/DetailCustomizations/Private/StringAssetReferenceCustomization.h
+++ b/Engine/Source/Editor/DetailCustomizations/Private/StringAssetReferenceCustomization.h
@@ -14,8 +14,8 @@ public:
 	}
 
 	/** IPropertyTypeCustomization interface */
-	virtual void CustomizeHeader( TSharedRef<class IPropertyHandle> InStructPropertyHandle, class FDetailWidgetRow& HeaderRow, IStructCustomizationUtils& StructCustomizationUtils ) OVERRIDE;
-	virtual void CustomizeChildren( TSharedRef<class IPropertyHandle> InStructPropertyHandle, class IDetailChildrenBuilder& StructBuilder, IStructCustomizationUtils& StructCustomizationUtils ) OVERRIDE;
+	virtual void CustomizeHeader( TSharedRef<class IPropertyHandle> InStructPropertyHandle, class FDetailWidgetRow& HeaderRow, IPropertyTypeCustomizationUtils& StructCustomizationUtils ) OVERRIDE;
+	virtual void CustomizeChildren( TSharedRef<class IPropertyHandle> InStructPropertyHandle, class IDetailChildrenBuilder& StructBuilder, IPropertyTypeCustomizationUtils& StructCustomizationUtils ) OVERRIDE;
 
 private:
 	bool OnShouldFilterAsset( const class FAssetData& InAssetData ) const;
diff --git a/Engine/Source/Editor/DetailCustomizations/Private/StringClassReferenceCustomization.cpp b/Engine/Source/Editor/DetailCustomizations/Private/StringClassReferenceCustomization.cpp
--- a/Engine/Source/Editor/DetailCustomizations/Private/StringClassReferenceCustomization.cpp
+++ b/Engine/Source/Editor/DetailCustomizations/Private/StringClassReferenceCustomization.cpp
@@ -3,7 +3,7 @@
 #include "DetailCustomizationsPrivatePCH.h"
 #include "StringClassReferenceCustomization.h"
 
-void FStringClassReferenceCustomization::CustomizeHeader(TSharedRef<IPropertyHandle> InStructPropertyHandle, FDetailWidgetRow& HeaderRow, IStructCustomizationUtils& StructCustomizationUtils)
+void FStringClassReferenceCustomization::CustomizeHeader(TSharedRef<IPropertyHandle> InStructPropertyHandle, FDetailWidgetRow& HeaderRow, IPropertyTypeCustomizationUtils& StructCustomizationUtils)
 {
 	StructPropertyHandle = InStructPropertyHandle;
 
@@ -41,7 +41,7 @@ void FStringClassReferenceCustomization::CustomizeHeader(TSharedRef<IPropertyHan
 	];
 }
 
-void FStringClassReferenceCustomization::CustomizeChildren(TSharedRef<IPropertyHandle> InStructPropertyHandle, IDetailChildrenBuilder& StructBuilder, IStructCustomizationUtils& StructCustomizationUtils)
+void FStringClassReferenceCustomization::CustomizeChildren(TSharedRef<IPropertyHandle> InStructPropertyHandle, IDetailChildrenBuilder& StructBuilder, IPropertyTypeCustomizationUtils& StructCustomizationUtils)
 {
 }
 
diff --git a/Engine/Source/Editor/DetailCustomizations/Private/StringClassReferenceCustomization.h b/Engine/Source/Editor/DetailCustomizations/Private/StringClassReferenceCustomization.h
--- a/Engine/Source/Editor/DetailCustomizations/Private/StringClassReferenceCustomization.h
+++ b/Engine/Source/Editor/DetailCustomizations/Private/StringClassReferenceCustomization.h
@@ -14,8 +14,8 @@ public:
 	}
 
 	/** IPropertyTypeCustomization interface */
-	virtual void CustomizeHeader(TSharedRef<class IPropertyHandle> InStructPropertyHandle, class FDetailWidgetRow& HeaderRow, IStructCustomizationUtils& StructCustomizationUtils) OVERRIDE;
-	virtual void CustomizeChildren(TSharedRef<class IPropertyHandle> InStructPropertyHandle, class IDetailChildrenBuilder& StructBuilder, IStructCustomizationUtils& StructCustomizationUtils) OVERRIDE;
+	virtual void CustomizeHeader(TSharedRef<class IPropertyHandle> InStructPropertyHandle, class FDetailWidgetRow& HeaderRow, IPropertyTypeCustomizationUtils& StructCustomizationUtils) OVERRIDE;
+	virtual void CustomizeChildren(TSharedRef<class IPropertyHandle> InStructPropertyHandle, class IDetailChildrenBuilder& StructBuilder, IPropertyTypeCustomizationUtils& StructCustomizationUtils) OVERRIDE;
 
 private:
 	/** @return The class currently set on this reference */
diff --git a/Engine/Source/Editor/DetailCustomizations/Private/VehicleTransmissionDataCustomization.cpp b/Engine/Source/Editor/DetailCustomizations/Private/VehicleTransmissionDataCustomization.cpp
--- a/Engine/Source/Editor/DetailCustomizations/Private/VehicleTransmissionDataCustomization.cpp
+++ b/Engine/Source/Editor/DetailCustomizations/Private/VehicleTransmissionDataCustomization.cpp
@@ -166,7 +166,7 @@ void FVehicleTransmissionDataCustomization::CreateGearUIDelegate(TSharedRef<IPro
 	CreateGearUIHelper(ChildrenBuilder.AddChildContent(Label.ToString()), Label, GearProperty, ForwardGear);
 }
 
-void FVehicleTransmissionDataCustomization::CustomizeChildren(TSharedRef<class IPropertyHandle> StructPropertyHandle, class IDetailChildrenBuilder& StructBuilder, IStructCustomizationUtils& StructCustomizationUtils)
+void FVehicleTransmissionDataCustomization::CustomizeChildren(TSharedRef<class IPropertyHandle> StructPropertyHandle, class IDetailChildrenBuilder& StructBuilder, IPropertyTypeCustomizationUtils& StructCustomizationUtils)
 {
 	// copy all transmision instances I'm accessing right now
 	TArray<void*> StructPtrs;
@@ -268,7 +268,7 @@ void FVehicleTransmissionDataCustomization::RemoveGear(TSharedRef<IPropertyHandl
 }
 
 
-void FVehicleTransmissionDataCustomization::CustomizeHeader(TSharedRef<IPropertyHandle> StructPropertyHandle, FDetailWidgetRow& HeaderRow, IStructCustomizationUtils& StructCustomizationUtils)
+void FVehicleTransmissionDataCustomization::CustomizeHeader(TSharedRef<IPropertyHandle> StructPropertyHandle, FDetailWidgetRow& HeaderRow, IPropertyTypeCustomizationUtils& StructCustomizationUtils)
 {
 	HeaderRow.
 	NameContent()
diff --git a/Engine/Source/Editor/DetailCustomizations/Private/VehicleTransmissionDataCustomization.h b/Engine/Source/Editor/DetailCustomizations/Private/VehicleTransmissionDataCustomization.h
--- a/Engine/Source/Editor/DetailCustomizations/Private/VehicleTransmissionDataCustomization.h
+++ b/Engine/Source/Editor/DetailCustomizations/Private/VehicleTransmissionDataCustomization.h
@@ -12,8 +12,8 @@ public:
 	}
 
 	/** IPropertyTypeCustomization interface */
-	virtual void CustomizeHeader(TSharedRef<IPropertyHandle> StructPropertyHandle, class FDetailWidgetRow& HeaderRow, IStructCustomizationUtils& StructCustomizationUtils) OVERRIDE;
-	virtual void CustomizeChildren( TSharedRef<class IPropertyHandle> StructPropertyHandle, class IDetailChildrenBuilder& StructBuilder, IStructCustomizationUtils& StructCustomizationUtils ) OVERRIDE;
+	virtual void CustomizeHeader(TSharedRef<IPropertyHandle> StructPropertyHandle, class FDetailWidgetRow& HeaderRow, IPropertyTypeCustomizationUtils& StructCustomizationUtils) OVERRIDE;
+	virtual void CustomizeChildren( TSharedRef<class IPropertyHandle> StructPropertyHandle, class IDetailChildrenBuilder& StructBuilder, IPropertyTypeCustomizationUtils& StructCustomizationUtils ) OVERRIDE;
 
 private:
 	FVehicleTransmissionData * SelectedTransmission;
diff --git a/Engine/Source/Editor/EnvironmentQueryEditor/Private/DetailCustomizations/EnvDirectionCustomization.cpp b/Engine/Source/Editor/EnvironmentQueryEditor/Private/DetailCustomizations/EnvDirectionCustomization.cpp
--- a/Engine/Source/Editor/EnvironmentQueryEditor/Private/DetailCustomizations/EnvDirectionCustomization.cpp
+++ b/Engine/Source/Editor/EnvironmentQueryEditor/Private/DetailCustomizations/EnvDirectionCustomization.cpp
@@ -11,7 +11,7 @@ TSharedRef<IPropertyTypeCustomization> FEnvDirectionCustomization::MakeInstance(
 	return MakeShareable(new FEnvDirectionCustomization);
 }
 
-void FEnvDirectionCustomization::CustomizeHeader( TSharedRef<class IPropertyHandle> StructPropertyHandle, class FDetailWidgetRow& HeaderRow, IStructCustomizationUtils& StructCustomizationUtils )
+void FEnvDirectionCustomization::CustomizeHeader( TSharedRef<class IPropertyHandle> StructPropertyHandle, class FDetailWidgetRow& HeaderRow, IPropertyTypeCustomizationUtils& StructCustomizationUtils )
 {
 	// create struct header
 	HeaderRow.NameContent()
@@ -36,7 +36,7 @@ void FEnvDirectionCustomization::CustomizeHeader( TSharedRef<class IPropertyHand
 	OnModeChanged();
 }
 
-void FEnvDirectionCustomization::CustomizeChildren( TSharedRef<class IPropertyHandle> StructPropertyHandle, class IDetailChildrenBuilder& StructBuilder, IStructCustomizationUtils& StructCustomizationUtils )
+void FEnvDirectionCustomization::CustomizeChildren( TSharedRef<class IPropertyHandle> StructPropertyHandle, class IDetailChildrenBuilder& StructBuilder, IPropertyTypeCustomizationUtils& StructCustomizationUtils )
 {
 	StructBuilder.AddChildProperty(ModeProp.ToSharedRef());
 
diff --git a/Engine/Source/Editor/EnvironmentQueryEditor/Private/DetailCustomizations/EnvDirectionCustomization.h b/Engine/Source/Editor/EnvironmentQueryEditor/Private/DetailCustomizations/EnvDirectionCustomization.h
--- a/Engine/Source/Editor/EnvironmentQueryEditor/Private/DetailCustomizations/EnvDirectionCustomization.h
+++ b/Engine/Source/Editor/EnvironmentQueryEditor/Private/DetailCustomizations/EnvDirectionCustomization.h
@@ -10,8 +10,8 @@ class FEnvDirectionCustomization : public IPropertyTypeCustomization
 public:
 
 	// Begin IPropertyTypeCustomization interface
-	virtual void CustomizeHeader( TSharedRef<class IPropertyHandle> StructPropertyHandle, class FDetailWidgetRow& HeaderRow, IStructCustomizationUtils& StructCustomizationUtils ) OVERRIDE;
-	virtual void CustomizeChildren( TSharedRef<class IPropertyHandle> StructPropertyHandle, class IDetailChildrenBuilder& StructBuilder, IStructCustomizationUtils& StructCustomizationUtils ) OVERRIDE;
+	virtual void CustomizeHeader( TSharedRef<class IPropertyHandle> StructPropertyHandle, class FDetailWidgetRow& HeaderRow, IPropertyTypeCustomizationUtils& StructCustomizationUtils ) OVERRIDE;
+	virtual void CustomizeChildren( TSharedRef<class IPropertyHandle> StructPropertyHandle, class IDetailChildrenBuilder& StructBuilder, IPropertyTypeCustomizationUtils& StructCustomizationUtils ) OVERRIDE;
 	// End IPropertyTypeCustomization interface
 
 	static TSharedRef<IPropertyTypeCustomization> MakeInstance( );
diff --git a/Engine/Source/Editor/EnvironmentQueryEditor/Private/DetailCustomizations/EnvQueryParamSetupCustomization.cpp b/Engine/Source/Editor/EnvironmentQueryEditor/Private/DetailCustomizations/EnvQueryParamSetupCustomization.cpp
--- a/Engine/Source/Editor/EnvironmentQueryEditor/Private/DetailCustomizations/EnvQueryParamSetupCustomization.cpp
+++ b/Engine/Source/Editor/EnvironmentQueryEditor/Private/DetailCustomizations/EnvQueryParamSetupCustomization.cpp
@@ -10,7 +10,7 @@ TSharedRef<IPropertyTypeCustomization> FEnvQueryParamSetupCustomization::MakeIns
 	return MakeShareable(new FEnvQueryParamSetupCustomization);
 }
 
-void FEnvQueryParamSetupCustomization::CustomizeHeader( TSharedRef<class IPropertyHandle> StructPropertyHandle, class FDetailWidgetRow& HeaderRow, IStructCustomizationUtils& StructCustomizationUtils )
+void FEnvQueryParamSetupCustomization::CustomizeHeader( TSharedRef<class IPropertyHandle> StructPropertyHandle, class FDetailWidgetRow& HeaderRow, IPropertyTypeCustomizationUtils& StructCustomizationUtils )
 {
 	PropertyHandle = StructPropertyHandle;
 	CacheMyValues();
@@ -69,7 +69,7 @@ void FEnvQueryParamSetupCustomization::CustomizeHeader( TSharedRef<class IProper
 	OnModeChanged();
 }
 
-void FEnvQueryParamSetupCustomization::CustomizeChildren( TSharedRef<class IPropertyHandle> StructPropertyHandle, class IDetailChildrenBuilder& StructBuilder, IStructCustomizationUtils& StructCustomizationUtils )
+void FEnvQueryParamSetupCustomization::CustomizeChildren( TSharedRef<class IPropertyHandle> StructPropertyHandle, class IDetailChildrenBuilder& StructBuilder, IPropertyTypeCustomizationUtils& StructCustomizationUtils )
 {
 	// do nothing
 }
diff --git a/Engine/Source/Editor/EnvironmentQueryEditor/Private/DetailCustomizations/EnvQueryParamSetupCustomization.h b/Engine/Source/Editor/EnvironmentQueryEditor/Private/DetailCustomizations/EnvQueryParamSetupCustomization.h
--- a/Engine/Source/Editor/EnvironmentQueryEditor/Private/DetailCustomizations/EnvQueryParamSetupCustomization.h
+++ b/Engine/Source/Editor/EnvironmentQueryEditor/Private/DetailCustomizations/EnvQueryParamSetupCustomization.h
@@ -10,8 +10,8 @@ class FEnvQueryParamSetupCustomization : public IPropertyTypeCustomization
 public:
 
 	// Begin IPropertyTypeCustomization interface
-	virtual void CustomizeHeader( TSharedRef<class IPropertyHandle> StructPropertyHandle, class FDetailWidgetRow& HeaderRow, IStructCustomizationUtils& StructCustomizationUtils ) OVERRIDE;
-	virtual void CustomizeChildren( TSharedRef<class IPropertyHandle> StructPropertyHandle, class IDetailChildrenBuilder& StructBuilder, IStructCustomizationUtils& StructCustomizationUtils ) OVERRIDE;
+	virtual void CustomizeHeader( TSharedRef<class IPropertyHandle> StructPropertyHandle, class FDetailWidgetRow& HeaderRow, IPropertyTypeCustomizationUtils& StructCustomizationUtils ) OVERRIDE;
+	virtual void CustomizeChildren( TSharedRef<class IPropertyHandle> StructPropertyHandle, class IDetailChildrenBuilder& StructBuilder, IPropertyTypeCustomizationUtils& StructCustomizationUtils ) OVERRIDE;
 	// End IPropertyTypeCustomization interface
 
 	static TSharedRef<IPropertyTypeCustomization> MakeInstance( );
diff --git a/Engine/Source/Editor/EnvironmentQueryEditor/Private/DetailCustomizations/EnvTraceDataCustomization.cpp b/Engine/Source/Editor/EnvironmentQueryEditor/Private/DetailCustomizations/EnvTraceDataCustomization.cpp
--- a/Engine/Source/Editor/EnvironmentQueryEditor/Private/DetailCustomizations/EnvTraceDataCustomization.cpp
+++ b/Engine/Source/Editor/EnvironmentQueryEditor/Private/DetailCustomizations/EnvTraceDataCustomization.cpp
@@ -12,7 +12,7 @@ TSharedRef<IPropertyTypeCustomization> FEnvTraceDataCustomization::MakeInstance(
 }
 
 BEGIN_SLATE_FUNCTION_BUILD_OPTIMIZATION
-void FEnvTraceDataCustomization::CustomizeHeader( TSharedRef<class IPropertyHandle> StructPropertyHandle, class FDetailWidgetRow& HeaderRow, IStructCustomizationUtils& StructCustomizationUtils )
+void FEnvTraceDataCustomization::CustomizeHeader( TSharedRef<class IPropertyHandle> StructPropertyHandle, class FDetailWidgetRow& HeaderRow, IPropertyTypeCustomizationUtils& StructCustomizationUtils )
 {
 	HeaderRow.NameContent()
 	[
@@ -31,7 +31,7 @@ void FEnvTraceDataCustomization::CustomizeHeader( TSharedRef<class IPropertyHand
 	CacheTraceModes(StructPropertyHandle);
 }
 
-void FEnvTraceDataCustomization::CustomizeChildren( TSharedRef<class IPropertyHandle> StructPropertyHandle, class IDetailChildrenBuilder& StructBuilder, IStructCustomizationUtils& StructCustomizationUtils )
+void FEnvTraceDataCustomization::CustomizeChildren( TSharedRef<class IPropertyHandle> StructPropertyHandle, class IDetailChildrenBuilder& StructBuilder, IPropertyTypeCustomizationUtils& StructCustomizationUtils )
 {
 	if (TraceModes.Num() > 1)
 	{
diff --git a/Engine/Source/Editor/EnvironmentQueryEditor/Private/DetailCustomizations/EnvTraceDataCustomization.h b/Engine/Source/Editor/EnvironmentQueryEditor/Private/DetailCustomizations/EnvTraceDataCustomization.h
--- a/Engine/Source/Editor/EnvironmentQueryEditor/Private/DetailCustomizations/EnvTraceDataCustomization.h
+++ b/Engine/Source/Editor/EnvironmentQueryEditor/Private/DetailCustomizations/EnvTraceDataCustomization.h
@@ -12,8 +12,8 @@ public:
 	static TSharedRef<IPropertyTypeCustomization> MakeInstance();
 
 	// Begin IPropertyTypeCustomization interface
-	virtual void CustomizeHeader( TSharedRef<class IPropertyHandle> StructPropertyHandle, class FDetailWidgetRow& HeaderRow, IStructCustomizationUtils& StructCustomizationUtils ) OVERRIDE;
-	virtual void CustomizeChildren( TSharedRef<class IPropertyHandle> StructPropertyHandle, class IDetailChildrenBuilder& StructBuilder, IStructCustomizationUtils& StructCustomizationUtils ) OVERRIDE;
+	virtual void CustomizeHeader( TSharedRef<class IPropertyHandle> StructPropertyHandle, class FDetailWidgetRow& HeaderRow, IPropertyTypeCustomizationUtils& StructCustomizationUtils ) OVERRIDE;
+	virtual void CustomizeChildren( TSharedRef<class IPropertyHandle> StructPropertyHandle, class IDetailChildrenBuilder& StructBuilder, IPropertyTypeCustomizationUtils& StructCustomizationUtils ) OVERRIDE;
 	// End IPropertyTypeCustomization interface
 
 protected:
diff --git a/Engine/Source/Editor/GameplayTagsEditor/Private/GameplayTagContainerCustomization.cpp b/Engine/Source/Editor/GameplayTagsEditor/Private/GameplayTagContainerCustomization.cpp
--- a/Engine/Source/Editor/GameplayTagsEditor/Private/GameplayTagContainerCustomization.cpp
+++ b/Engine/Source/Editor/GameplayTagsEditor/Private/GameplayTagContainerCustomization.cpp
@@ -10,7 +10,7 @@
 
 #define LOCTEXT_NAMESPACE "GameplayTagContainerCustomization"
 
-void FGameplayTagContainerCustomization::CustomizeHeader(TSharedRef<class IPropertyHandle> InStructPropertyHandle, class FDetailWidgetRow& HeaderRow, IStructCustomizationUtils& StructCustomizationUtils)
+void FGameplayTagContainerCustomization::CustomizeHeader(TSharedRef<class IPropertyHandle> InStructPropertyHandle, class FDetailWidgetRow& HeaderRow, IPropertyTypeCustomizationUtils& StructCustomizationUtils)
 {
 	StructPropertyHandle = InStructPropertyHandle;
 
diff --git a/Engine/Source/Editor/GameplayTagsEditor/Private/GameplayTagContainerCustomization.h b/Engine/Source/Editor/GameplayTagsEditor/Private/GameplayTagContainerCustomization.h
--- a/Engine/Source/Editor/GameplayTagsEditor/Private/GameplayTagContainerCustomization.h
+++ b/Engine/Source/Editor/GameplayTagsEditor/Private/GameplayTagContainerCustomization.h
@@ -17,10 +17,10 @@ public:
 	~FGameplayTagContainerCustomization();
 
 	/** Overridden to show an edit button to launch the gameplay tag editor */
-	virtual void CustomizeHeader(TSharedRef<class IPropertyHandle> StructPropertyHandle, class FDetailWidgetRow& HeaderRow, IStructCustomizationUtils& StructCustomizationUtils) OVERRIDE;
+	virtual void CustomizeHeader(TSharedRef<class IPropertyHandle> StructPropertyHandle, class FDetailWidgetRow& HeaderRow, IPropertyTypeCustomizationUtils& StructCustomizationUtils) OVERRIDE;
 	
 	/** Overridden to do nothing */
-	virtual void CustomizeChildren(TSharedRef<IPropertyHandle> InStructPropertyHandle, class IDetailChildrenBuilder& ChildBuilder, IStructCustomizationUtils& StructCustomizationUtils) OVERRIDE {}
+	virtual void CustomizeChildren(TSharedRef<IPropertyHandle> InStructPropertyHandle, class IDetailChildrenBuilder& ChildBuilder, IPropertyTypeCustomizationUtils& StructCustomizationUtils) OVERRIDE {}
 
 	// Begin FEditorUndoClient Interface
 	virtual void PostUndo( bool bSuccess ) OVERRIDE;
diff --git a/Engine/Source/Editor/LandscapeEditor/Private/LandscapeEditorDetailCustomization_CopyPaste.cpp b/Engine/Source/Editor/LandscapeEditor/Private/LandscapeEditorDetailCustomization_CopyPaste.cpp
--- a/Engine/Source/Editor/LandscapeEditor/Private/LandscapeEditorDetailCustomization_CopyPaste.cpp
+++ b/Engine/Source/Editor/LandscapeEditor/Private/LandscapeEditorDetailCustomization_CopyPaste.cpp
@@ -462,11 +462,11 @@ TSharedRef<IPropertyTypeCustomization> FLandscapeEditorStructCustomization_FGizm
 	return MakeShareable(new FLandscapeEditorStructCustomization_FGizmoImportLayer);
 }
 
-void FLandscapeEditorStructCustomization_FGizmoImportLayer::CustomizeHeader(TSharedRef<IPropertyHandle> StructPropertyHandle, class FDetailWidgetRow& HeaderRow, IStructCustomizationUtils& StructCustomizationUtils)
+void FLandscapeEditorStructCustomization_FGizmoImportLayer::CustomizeHeader(TSharedRef<IPropertyHandle> StructPropertyHandle, class FDetailWidgetRow& HeaderRow, IPropertyTypeCustomizationUtils& StructCustomizationUtils)
 {
 }
 
-void FLandscapeEditorStructCustomization_FGizmoImportLayer::CustomizeChildren(TSharedRef<IPropertyHandle> StructPropertyHandle, class IDetailChildrenBuilder& ChildBuilder, IStructCustomizationUtils& StructCustomizationUtils)
+void FLandscapeEditorStructCustomization_FGizmoImportLayer::CustomizeChildren(TSharedRef<IPropertyHandle> StructPropertyHandle, class IDetailChildrenBuilder& ChildBuilder, IPropertyTypeCustomizationUtils& StructCustomizationUtils)
 {
 	TSharedRef<IPropertyHandle> PropertyHandle_LayerFilename = StructPropertyHandle->GetChildHandle(GET_MEMBER_NAME_CHECKED(FGizmoImportLayer, LayerFilename)).ToSharedRef();
 	ChildBuilder.AddChildProperty(PropertyHandle_LayerFilename)
diff --git a/Engine/Source/Editor/LandscapeEditor/Private/LandscapeEditorDetailCustomization_CopyPaste.h b/Engine/Source/Editor/LandscapeEditor/Private/LandscapeEditorDetailCustomization_CopyPaste.h
--- a/Engine/Source/Editor/LandscapeEditor/Private/LandscapeEditorDetailCustomization_CopyPaste.h
+++ b/Engine/Source/Editor/LandscapeEditor/Private/LandscapeEditorDetailCustomization_CopyPaste.h
@@ -35,8 +35,8 @@ class FLandscapeEditorStructCustomization_FGizmoImportLayer : public FLandscapeE
 public:
 	static TSharedRef<IPropertyTypeCustomization> MakeInstance();
 
-	virtual void CustomizeHeader(TSharedRef<IPropertyHandle> StructPropertyHandle, class FDetailWidgetRow& HeaderRow, IStructCustomizationUtils& StructCustomizationUtils) OVERRIDE;
-	virtual void CustomizeChildren(TSharedRef<IPropertyHandle> StructPropertyHandle, class IDetailChildrenBuilder& ChildBuilder, IStructCustomizationUtils& StructCustomizationUtils) OVERRIDE;
+	virtual void CustomizeHeader(TSharedRef<IPropertyHandle> StructPropertyHandle, class FDetailWidgetRow& HeaderRow, IPropertyTypeCustomizationUtils& StructCustomizationUtils) OVERRIDE;
+	virtual void CustomizeChildren(TSharedRef<IPropertyHandle> StructPropertyHandle, class IDetailChildrenBuilder& ChildBuilder, IPropertyTypeCustomizationUtils& StructCustomizationUtils) OVERRIDE;
 
 public:
 	static FReply OnGizmoImportLayerFilenameButtonClicked(TSharedRef<IPropertyHandle> PropertyHandle_LayerFilename);
diff --git a/Engine/Source/Editor/LandscapeEditor/Private/LandscapeEditorDetailCustomization_NewLandscape.cpp b/Engine/Source/Editor/LandscapeEditor/Private/LandscapeEditorDetailCustomization_NewLandscape.cpp
--- a/Engine/Source/Editor/LandscapeEditor/Private/LandscapeEditorDetailCustomization_NewLandscape.cpp
+++ b/Engine/Source/Editor/LandscapeEditor/Private/LandscapeEditorDetailCustomization_NewLandscape.cpp
@@ -1312,11 +1312,11 @@ TSharedRef<IPropertyTypeCustomization> FLandscapeEditorStructCustomization_FLand
 	return MakeShareable(new FLandscapeEditorStructCustomization_FLandscapeImportLayer);
 }
 
-void FLandscapeEditorStructCustomization_FLandscapeImportLayer::CustomizeHeader(TSharedRef<IPropertyHandle> StructPropertyHandle, class FDetailWidgetRow& HeaderRow, IStructCustomizationUtils& StructCustomizationUtils)
+void FLandscapeEditorStructCustomization_FLandscapeImportLayer::CustomizeHeader(TSharedRef<IPropertyHandle> StructPropertyHandle, class FDetailWidgetRow& HeaderRow, IPropertyTypeCustomizationUtils& StructCustomizationUtils)
 {
 }
 
-void FLandscapeEditorStructCustomization_FLandscapeImportLayer::CustomizeChildren(TSharedRef<IPropertyHandle> StructPropertyHandle, class IDetailChildrenBuilder& ChildBuilder, IStructCustomizationUtils& StructCustomizationUtils)
+void FLandscapeEditorStructCustomization_FLandscapeImportLayer::CustomizeChildren(TSharedRef<IPropertyHandle> StructPropertyHandle, class IDetailChildrenBuilder& ChildBuilder, IPropertyTypeCustomizationUtils& StructCustomizationUtils)
 {
 	TSharedRef<IPropertyHandle> PropertyHandle_LayerName = StructPropertyHandle->GetChildHandle(GET_MEMBER_NAME_CHECKED(FLandscapeImportLayer, LayerName)).ToSharedRef();
 	TSharedRef<IPropertyHandle> PropertyHandle_LayerInfo = StructPropertyHandle->GetChildHandle(GET_MEMBER_NAME_CHECKED(FLandscapeImportLayer, LayerInfo)).ToSharedRef();
diff --git a/Engine/Source/Editor/LandscapeEditor/Private/LandscapeEditorDetailCustomization_NewLandscape.h b/Engine/Source/Editor/LandscapeEditor/Private/LandscapeEditorDetailCustomization_NewLandscape.h
--- a/Engine/Source/Editor/LandscapeEditor/Private/LandscapeEditorDetailCustomization_NewLandscape.h
+++ b/Engine/Source/Editor/LandscapeEditor/Private/LandscapeEditorDetailCustomization_NewLandscape.h
@@ -86,8 +86,8 @@ class FLandscapeEditorStructCustomization_FLandscapeImportLayer : public FLandsc
 public:
 	static TSharedRef<IPropertyTypeCustomization> MakeInstance();
 
-	virtual void CustomizeHeader(TSharedRef<IPropertyHandle> StructPropertyHandle, class FDetailWidgetRow& HeaderRow, IStructCustomizationUtils& StructCustomizationUtils) OVERRIDE;
-	virtual void CustomizeChildren(TSharedRef<IPropertyHandle> StructPropertyHandle, class IDetailChildrenBuilder& ChildBuilder, IStructCustomizationUtils& StructCustomizationUtils) OVERRIDE;
+	virtual void CustomizeHeader(TSharedRef<IPropertyHandle> StructPropertyHandle, class FDetailWidgetRow& HeaderRow, IPropertyTypeCustomizationUtils& StructCustomizationUtils) OVERRIDE;
+	virtual void CustomizeChildren(TSharedRef<IPropertyHandle> StructPropertyHandle, class IDetailChildrenBuilder& ChildBuilder, IPropertyTypeCustomizationUtils& StructCustomizationUtils) OVERRIDE;
 
 public:
 	static FReply OnLayerFilenameButtonClicked(TSharedRef<IPropertyHandle> PropertyHandle_LayerFilename);
diff --git a/Engine/Source/Editor/Persona/Private/AnimGraphNodeDetails.cpp b/Engine/Source/Editor/Persona/Private/AnimGraphNodeDetails.cpp
--- a/Engine/Source/Editor/Persona/Private/AnimGraphNodeDetails.cpp
+++ b/Engine/Source/Editor/Persona/Private/AnimGraphNodeDetails.cpp
@@ -264,7 +264,7 @@ TSharedRef<IPropertyTypeCustomization> FInputScaleBiasCustomization::MakeInstanc
 }
 
 BEGIN_SLATE_FUNCTION_BUILD_OPTIMIZATION
-void FInputScaleBiasCustomization::CustomizeHeader(TSharedRef<class IPropertyHandle> StructPropertyHandle, class FDetailWidgetRow& HeaderRow, IStructCustomizationUtils& StructCustomizationUtils)
+void FInputScaleBiasCustomization::CustomizeHeader(TSharedRef<class IPropertyHandle> StructPropertyHandle, class FDetailWidgetRow& HeaderRow, IPropertyTypeCustomizationUtils& StructCustomizationUtils)
 {
 	HeaderRow.NameContent()
 	[
@@ -335,7 +335,7 @@ void FInputScaleBiasCustomization::CustomizeHeader(TSharedRef<class IPropertyHan
 END_SLATE_FUNCTION_BUILD_OPTIMIZATION
 
 /////////////////////////////////////////////////////
-void FInputScaleBiasCustomization::CustomizeChildren( TSharedRef<class IPropertyHandle> StructPropertyHandle, class IDetailChildrenBuilder& StructBuilder, IStructCustomizationUtils& StructCustomizationUtils )
+void FInputScaleBiasCustomization::CustomizeChildren( TSharedRef<class IPropertyHandle> StructPropertyHandle, class IDetailChildrenBuilder& StructBuilder, IPropertyTypeCustomizationUtils& StructCustomizationUtils )
 {
 	// nothing here
 }
@@ -407,7 +407,7 @@ TSharedRef<IPropertyTypeCustomization> FBoneReferenceCustomization::MakeInstance
 }
 
 BEGIN_SLATE_FUNCTION_BUILD_OPTIMIZATION
-void FBoneReferenceCustomization::CustomizeHeader( TSharedRef<IPropertyHandle> StructPropertyHandle, class FDetailWidgetRow& HeaderRow, IStructCustomizationUtils& StructCustomizationUtils )
+void FBoneReferenceCustomization::CustomizeHeader( TSharedRef<IPropertyHandle> StructPropertyHandle, class FDetailWidgetRow& HeaderRow, IPropertyTypeCustomizationUtils& StructCustomizationUtils )
 {
 	uint32 NumChildren = 0;
 	StructPropertyHandle->GetNumChildren(NumChildren);
@@ -464,7 +464,7 @@ void FBoneReferenceCustomization::CustomizeHeader( TSharedRef<IPropertyHandle> S
 }
 END_SLATE_FUNCTION_BUILD_OPTIMIZATION
 
-void FBoneReferenceCustomization::CustomizeChildren( TSharedRef<IPropertyHandle> StructPropertyHandle, class IDetailChildrenBuilder& ChildBuilder, IStructCustomizationUtils& StructCustomizationUtils )
+void FBoneReferenceCustomization::CustomizeChildren( TSharedRef<IPropertyHandle> StructPropertyHandle, class IDetailChildrenBuilder& ChildBuilder, IPropertyTypeCustomizationUtils& StructCustomizationUtils )
 {
 	// No child customisations as the properties are shown in the header
 }
diff --git a/Engine/Source/Editor/Persona/Private/AnimGraphNodeDetails.h b/Engine/Source/Editor/Persona/Private/AnimGraphNodeDetails.h
--- a/Engine/Source/Editor/Persona/Private/AnimGraphNodeDetails.h
+++ b/Engine/Source/Editor/Persona/Private/AnimGraphNodeDetails.h
@@ -42,8 +42,8 @@ public:
 	static TSharedRef<IPropertyTypeCustomization> MakeInstance();
 
 	// IPropertyTypeCustomization interface
-	virtual void CustomizeHeader(TSharedRef<IPropertyHandle> StructPropertyHandle, class FDetailWidgetRow& HeaderRow, IStructCustomizationUtils& StructCustomizationUtils) OVERRIDE;
-	virtual void CustomizeChildren(TSharedRef<IPropertyHandle> StructPropertyHandle, class IDetailChildrenBuilder& ChildBuilder, IStructCustomizationUtils& StructCustomizationUtils) OVERRIDE;
+	virtual void CustomizeHeader(TSharedRef<IPropertyHandle> StructPropertyHandle, class FDetailWidgetRow& HeaderRow, IPropertyTypeCustomizationUtils& StructCustomizationUtils) OVERRIDE;
+	virtual void CustomizeChildren(TSharedRef<IPropertyHandle> StructPropertyHandle, class IDetailChildrenBuilder& ChildBuilder, IPropertyTypeCustomizationUtils& StructCustomizationUtils) OVERRIDE;
 	// End of IPropertyTypeCustomization interface
 
 	FText GetMinValue(TSharedRef<class IPropertyHandle> StructPropertyHandle) const;
@@ -61,8 +61,8 @@ public:
 	static TSharedRef<IPropertyTypeCustomization> MakeInstance();
 
 	// IPropertyTypeCustomization interface
-	virtual void CustomizeHeader(TSharedRef<IPropertyHandle> StructPropertyHandle, class FDetailWidgetRow& HeaderRow, IStructCustomizationUtils& StructCustomizationUtils) OVERRIDE;
-	virtual void CustomizeChildren(TSharedRef<IPropertyHandle> StructPropertyHandle, class IDetailChildrenBuilder& ChildBuilder, IStructCustomizationUtils& StructCustomizationUtils) OVERRIDE;
+	virtual void CustomizeHeader(TSharedRef<IPropertyHandle> StructPropertyHandle, class FDetailWidgetRow& HeaderRow, IPropertyTypeCustomizationUtils& StructCustomizationUtils) OVERRIDE;
+	virtual void CustomizeChildren(TSharedRef<IPropertyHandle> StructPropertyHandle, class IDetailChildrenBuilder& ChildBuilder, IPropertyTypeCustomizationUtils& StructCustomizationUtils) OVERRIDE;
 
 private:
 	// Storage object for bone hierarchy
diff --git a/Engine/Source/Editor/WorldBrowser/Private/Tiles/WorldTileDetailsCustomization.cpp b/Engine/Source/Editor/WorldBrowser/Private/Tiles/WorldTileDetailsCustomization.cpp
--- a/Engine/Source/Editor/WorldBrowser/Private/Tiles/WorldTileDetailsCustomization.cpp
+++ b/Engine/Source/Editor/WorldBrowser/Private/Tiles/WorldTileDetailsCustomization.cpp
@@ -180,7 +180,7 @@ TSharedRef<IPropertyTypeCustomization> FStreamingLevelDetailsCustomization::Make
 
 void FStreamingLevelDetailsCustomization::CustomizeHeader(TSharedRef<class IPropertyHandle> StructPropertyHandle, 
 																class FDetailWidgetRow& HeaderRow, 
-																IStructCustomizationUtils& StructCustomizationUtils )
+																IPropertyTypeCustomizationUtils& StructCustomizationUtils )
 {
 	HeaderRow
 		.NameContent()
@@ -191,7 +191,7 @@ void FStreamingLevelDetailsCustomization::CustomizeHeader(TSharedRef<class IProp
 
 void FStreamingLevelDetailsCustomization::CustomizeChildren(TSharedRef<IPropertyHandle> StructPropertyHandle, 
 																	class IDetailChildrenBuilder& ChildBuilder, 
-																	IStructCustomizationUtils& StructCustomizationUtils )
+																	IPropertyTypeCustomizationUtils& StructCustomizationUtils )
 {
 	TSharedPtr<IPropertyHandle> StreamingModeProperty = StructPropertyHandle->GetChildHandle(TEXT("StreamingMode"));
 	TSharedPtr<IPropertyHandle> PackageNameProperty = StructPropertyHandle->GetChildHandle(TEXT("PackageName"));
@@ -238,7 +238,7 @@ TSharedRef<IPropertyTypeCustomization> FTileLODEntryDetailsCustomization::MakeIn
 
 void FTileLODEntryDetailsCustomization::CustomizeHeader(TSharedRef<class IPropertyHandle> StructPropertyHandle, 
 																class FDetailWidgetRow& HeaderRow, 
-																IStructCustomizationUtils& StructCustomizationUtils )
+																IPropertyTypeCustomizationUtils& StructCustomizationUtils )
 {
 	HeaderRow
 		.NameContent()
@@ -260,7 +260,7 @@ void FTileLODEntryDetailsCustomization::CustomizeHeader(TSharedRef<class IProper
 
 void FTileLODEntryDetailsCustomization::CustomizeChildren(TSharedRef<IPropertyHandle> StructPropertyHandle, 
 																	class IDetailChildrenBuilder& ChildBuilder, 
-																	IStructCustomizationUtils& StructCustomizationUtils )
+																	IPropertyTypeCustomizationUtils& StructCustomizationUtils )
 {
 	LODIndexHandle = StructPropertyHandle->GetChildHandle(
 		GET_MEMBER_NAME_CHECKED(FTileLODEntryDetails, LODIndex)
diff --git a/Engine/Source/Editor/WorldBrowser/Private/Tiles/WorldTileDetailsCustomization.h b/Engine/Source/Editor/WorldBrowser/Private/Tiles/WorldTileDetailsCustomization.h
--- a/Engine/Source/Editor/WorldBrowser/Private/Tiles/WorldTileDetailsCustomization.h
+++ b/Engine/Source/Editor/WorldBrowser/Private/Tiles/WorldTileDetailsCustomization.h
@@ -38,8 +38,8 @@ public:
 	static TSharedRef<IPropertyTypeCustomization> MakeInstance(TSharedRef<FWorldTileCollectionModel> InWorldData);
 
 	/** IPropertyTypeCustomization interface */
-	virtual void CustomizeHeader(TSharedRef<class IPropertyHandle> StructPropertyHandle, class FDetailWidgetRow& HeaderRow, IStructCustomizationUtils& StructCustomizationUtils ) OVERRIDE;
-	virtual void CustomizeChildren(TSharedRef<class IPropertyHandle> StructPropertyHandle, class IDetailChildrenBuilder& ChildBuilder, IStructCustomizationUtils& StructCustomizationUtils ) OVERRIDE;
+	virtual void CustomizeHeader(TSharedRef<class IPropertyHandle> StructPropertyHandle, class FDetailWidgetRow& HeaderRow, IPropertyTypeCustomizationUtils& StructCustomizationUtils ) OVERRIDE;
+	virtual void CustomizeChildren(TSharedRef<class IPropertyHandle> StructPropertyHandle, class IDetailChildrenBuilder& ChildBuilder, IPropertyTypeCustomizationUtils& StructCustomizationUtils ) OVERRIDE;
 
 private:
 	bool OnShouldFilterStreamingPackage(const FString& InPackageName) const;
@@ -56,8 +56,8 @@ public:
 	static TSharedRef<IPropertyTypeCustomization> MakeInstance(TSharedRef<FWorldTileCollectionModel> InWorldData);
 
 	/** IPropertyTypeCustomization interface */
-	virtual void CustomizeHeader( TSharedRef<class IPropertyHandle> StructPropertyHandle, class FDetailWidgetRow& HeaderRow, IStructCustomizationUtils& StructCustomizationUtils ) OVERRIDE;
-	virtual void CustomizeChildren( TSharedRef<IPropertyHandle> StructPropertyHandle, class IDetailChildrenBuilder& ChildBuilder, IStructCustomizationUtils& StructCustomizationUtils ) OVERRIDE;
+	virtual void CustomizeHeader( TSharedRef<class IPropertyHandle> StructPropertyHandle, class FDetailWidgetRow& HeaderRow, IPropertyTypeCustomizationUtils& StructCustomizationUtils ) OVERRIDE;
+	virtual void CustomizeChildren( TSharedRef<IPropertyHandle> StructPropertyHandle, class IDetailChildrenBuilder& ChildBuilder, IPropertyTypeCustomizationUtils& StructCustomizationUtils ) OVERRIDE;
 
 private:
 	FReply OnGenerateTile();
