Deprecated RegisterCustomPropertyLayout and fixed up use cases.  Use RegisterCustomClassLayout instead

diff --git a/Engine/Plugins/Experimental/Paper2D/Source/Paper2DEditor/Private/Paper2DEditorModule.cpp b/Engine/Plugins/Experimental/Paper2D/Source/Paper2DEditor/Private/Paper2DEditorModule.cpp
--- a/Engine/Plugins/Experimental/Paper2D/Source/Paper2DEditor/Private/Paper2DEditorModule.cpp
+++ b/Engine/Plugins/Experimental/Paper2D/Source/Paper2DEditor/Private/Paper2DEditorModule.cpp
@@ -95,8 +95,8 @@ public:
 		// Register the details customizations
 		{
 			FPropertyEditorModule& PropertyModule = FModuleManager::LoadModuleChecked<FPropertyEditorModule>("PropertyEditor");
-			PropertyModule.RegisterCustomPropertyLayout("PaperTileMapRenderComponent", FOnGetDetailCustomizationInstance::CreateStatic(&FPaperTileMapDetailsCustomization::MakeInstance));
-			PropertyModule.RegisterCustomPropertyLayout("PaperSprite", FOnGetDetailCustomizationInstance::CreateStatic(&FSpriteDetailsCustomization::MakeInstance));
+			PropertyModule.RegisterCustomClassLayout("PaperTileMapRenderComponent", FOnGetDetailCustomizationInstance::CreateStatic(&FPaperTileMapDetailsCustomization::MakeInstance));
+			PropertyModule.RegisterCustomClassLayout("PaperSprite", FOnGetDetailCustomizationInstance::CreateStatic(&FSpriteDetailsCustomization::MakeInstance));
 
 			//@TODO: Struct registration should happen using ::StaticStruct, not by string!!!
 			//PropertyModule.RegisterCustomPropertyTypeLayout( "SpritePolygonCollection", FOnGetPropertyTypeCustomizationInstance::CreateStatic( &FSpritePolygonCollectionCustomization::MakeInstance ) );
diff --git a/Engine/Source/Developer/Android/AndroidPlatformEditor/Private/AndroidPlatformEditorModule.cpp b/Engine/Source/Developer/Android/AndroidPlatformEditor/Private/AndroidPlatformEditorModule.cpp
--- a/Engine/Source/Developer/Android/AndroidPlatformEditor/Private/AndroidPlatformEditorModule.cpp
+++ b/Engine/Source/Developer/Android/AndroidPlatformEditor/Private/AndroidPlatformEditorModule.cpp
@@ -16,7 +16,7 @@ class FAndroidPlatformEditorModule : public IModuleInterface
 	{
 		// Register the settings detail panel customization
 		FPropertyEditorModule& PropertyModule = FModuleManager::LoadModuleChecked<FPropertyEditorModule>("PropertyEditor");
-		PropertyModule.RegisterCustomPropertyLayout(
+		PropertyModule.RegisterCustomClassLayout(
 			UAndroidRuntimeSettings::StaticClass()->GetFName(),
 			FOnGetDetailCustomizationInstance::CreateStatic(&FAndroidTargetSettingsCustomization::MakeInstance));
 		PropertyModule.NotifyCustomizationModuleChanged();
diff --git a/Engine/Source/Developer/iOS/IOSPlatformEditor/Private/IOSPlatformEditorModule.cpp b/Engine/Source/Developer/iOS/IOSPlatformEditor/Private/IOSPlatformEditorModule.cpp
--- a/Engine/Source/Developer/iOS/IOSPlatformEditor/Private/IOSPlatformEditorModule.cpp
+++ b/Engine/Source/Developer/iOS/IOSPlatformEditor/Private/IOSPlatformEditorModule.cpp
@@ -16,7 +16,7 @@ class FIOSPlatformEditorModule : public IModuleInterface
 	{
 		// Register the settings detail panel customization
 		FPropertyEditorModule& PropertyModule = FModuleManager::LoadModuleChecked<FPropertyEditorModule>("PropertyEditor");
-		PropertyModule.RegisterCustomPropertyLayout(
+		PropertyModule.RegisterCustomClassLayout(
 			"IOSRuntimeSettings",
 			FOnGetDetailCustomizationInstance::CreateStatic(&FIOSTargetSettingsCustomization::MakeInstance));
 		PropertyModule.NotifyCustomizationModuleChanged();
diff --git a/Engine/Source/Editor/BehaviorTreeEditor/Private/BehaviorTreeEditorModule.cpp b/Engine/Source/Editor/BehaviorTreeEditor/Private/BehaviorTreeEditorModule.cpp
--- a/Engine/Source/Editor/BehaviorTreeEditor/Private/BehaviorTreeEditorModule.cpp
+++ b/Engine/Source/Editor/BehaviorTreeEditor/Private/BehaviorTreeEditorModule.cpp
@@ -51,8 +51,8 @@ void FBehaviorTreeEditorModule::StartupModule()
 	// Register the details customizer
 	FPropertyEditorModule& PropertyModule = FModuleManager::LoadModuleChecked<FPropertyEditorModule>("PropertyEditor");
 	PropertyModule.RegisterCustomPropertyTypeLayout( "BlackboardKeySelector", FOnGetPropertyTypeCustomizationInstance::CreateStatic( &FBlackboardSelectorDetails::MakeInstance ) );
-	PropertyModule.RegisterCustomPropertyLayout( "BTDecorator_Blackboard", FOnGetDetailCustomizationInstance::CreateStatic( &FBlackboardDecoratorDetails::MakeInstance ) );
-	PropertyModule.RegisterCustomPropertyLayout( "BTDecorator", FOnGetDetailCustomizationInstance::CreateStatic( &FBehaviorDecoratorDetails::MakeInstance ) );
+	PropertyModule.RegisterCustomClassLayout( "BTDecorator_Blackboard", FOnGetDetailCustomizationInstance::CreateStatic( &FBlackboardDecoratorDetails::MakeInstance ) );
+	PropertyModule.RegisterCustomClassLayout( "BTDecorator", FOnGetDetailCustomizationInstance::CreateStatic( &FBehaviorDecoratorDetails::MakeInstance ) );
 	PropertyModule.NotifyCustomizationModuleChanged();
 }
 
@@ -89,8 +89,8 @@ void FBehaviorTreeEditorModule::ShutdownModule()
 	{
 		FPropertyEditorModule& PropertyModule = FModuleManager::LoadModuleChecked<FPropertyEditorModule>("PropertyEditor");
 		PropertyModule.UnregisterCustomPropertyTypeLayout( "BlackboardKeySelector" );
-		PropertyModule.UnregisterCustomPropertyLayout( "BTDecorator_Blackboard" );
-		PropertyModule.UnregisterCustomPropertyLayout( "BTDecorator" );
+		PropertyModule.UnregisterCustomClassLayout( "BTDecorator_Blackboard" );
+		PropertyModule.UnregisterCustomClassLayout( "BTDecorator" );
 		PropertyModule.NotifyCustomizationModuleChanged();
 	}
 }
diff --git a/Engine/Source/Editor/Blutility/Private/BlutilityModule.cpp b/Engine/Source/Editor/Blutility/Private/BlutilityModule.cpp
--- a/Engine/Source/Editor/Blutility/Private/BlutilityModule.cpp
+++ b/Engine/Source/Editor/Blutility/Private/BlutilityModule.cpp
@@ -37,8 +37,8 @@ public:
 		
 		// Register the details customizer
 		FPropertyEditorModule& PropertyModule = FModuleManager::LoadModuleChecked<FPropertyEditorModule>("PropertyEditor");
-		PropertyModule.RegisterCustomPropertyLayout("PlacedEditorUtilityBase", FOnGetDetailCustomizationInstance::CreateStatic(&FEditorUtilityInstanceDetails::MakeInstance));
-		PropertyModule.RegisterCustomPropertyLayout("GlobalEditorUtilityBase", FOnGetDetailCustomizationInstance::CreateStatic(&FEditorUtilityInstanceDetails::MakeInstance));
+		PropertyModule.RegisterCustomClassLayout("PlacedEditorUtilityBase", FOnGetDetailCustomizationInstance::CreateStatic(&FEditorUtilityInstanceDetails::MakeInstance));
+		PropertyModule.RegisterCustomClassLayout("GlobalEditorUtilityBase", FOnGetDetailCustomizationInstance::CreateStatic(&FEditorUtilityInstanceDetails::MakeInstance));
 		PropertyModule.NotifyCustomizationModuleChanged();
 
 		if (GetDefault<UEditorExperimentalSettings>()->bEnableEditorUtilityBlueprints)
@@ -72,8 +72,8 @@ public:
 		if (FModuleManager::Get().IsModuleLoaded("PropertyEditor"))
 		{
 			FPropertyEditorModule& PropertyModule = FModuleManager::LoadModuleChecked<FPropertyEditorModule>("PropertyEditor");
-			PropertyModule.UnregisterCustomPropertyLayout("PlacedEditorUtilityBase");
-			PropertyModule.UnregisterCustomPropertyLayout("GlobalEditorUtilityBase");
+			PropertyModule.UnregisterCustomClassLayout("PlacedEditorUtilityBase");
+			PropertyModule.UnregisterCustomClassLayout("GlobalEditorUtilityBase");
 			PropertyModule.NotifyCustomizationModuleChanged();
 		}
 	}
diff --git a/Engine/Source/Editor/DetailCustomizations/Private/DetailCustomizations.cpp b/Engine/Source/Editor/DetailCustomizations/Private/DetailCustomizations.cpp
--- a/Engine/Source/Editor/DetailCustomizations/Private/DetailCustomizations.cpp
+++ b/Engine/Source/Editor/DetailCustomizations/Private/DetailCustomizations.cpp
@@ -103,7 +103,7 @@ void FDetailCustomizationsModule::ShutdownModule()
 		{
 			if( It->IsValid() )
 			{
-				PropertyModule.UnregisterCustomPropertyLayout( *It );
+				PropertyModule.UnregisterCustomClassLayout( *It );
 			}
 		}
 
@@ -173,73 +173,73 @@ void FDetailCustomizationsModule::RegisterObjectCustomizations()
 	// in the order they are customized.  This is only relevant for inheritance where both a child and a parent have properties that are customized.
 	// In the order below, Actor will get a chance to display details first, followed by USceneComponent.
 
-	RegisterCustomPropertyLayout("Actor", FOnGetDetailCustomizationInstance::CreateStatic(&FActorDetails::MakeInstance));
-	RegisterCustomPropertyLayout("SceneComponent", FOnGetDetailCustomizationInstance::CreateStatic(&FSceneComponentDetails::MakeInstance));
-	RegisterCustomPropertyLayout("PrimitiveComponent", FOnGetDetailCustomizationInstance::CreateStatic(&FPrimitiveComponentDetails::MakeInstance));
-	RegisterCustomPropertyLayout("StaticMeshComponent", FOnGetDetailCustomizationInstance::CreateStatic(&FStaticMeshComponentDetails::MakeInstance));
-	RegisterCustomPropertyLayout("SkeletalMeshComponent", FOnGetDetailCustomizationInstance::CreateStatic(&FSkeletalMeshComponentDetails::MakeInstance));
-	RegisterCustomPropertyLayout("SkinnedMeshComponent", FOnGetDetailCustomizationInstance::CreateStatic(&FSkinnedMeshComponentDetails::MakeInstance));
-	RegisterCustomPropertyLayout("WheeledVehicleMovementComponent4W", FOnGetDetailCustomizationInstance::CreateStatic(&FWheeledVehicleMovementComponent4WDetails::MakeInstance));
-	RegisterCustomPropertyLayout("LightComponent", FOnGetDetailCustomizationInstance::CreateStatic(&FLightComponentDetails::MakeInstance));
-	RegisterCustomPropertyLayout("PointLightComponent", FOnGetDetailCustomizationInstance::CreateStatic(&FPointLightComponentDetails::MakeInstance));
-	RegisterCustomPropertyLayout("DirectionalLightComponent", FOnGetDetailCustomizationInstance::CreateStatic(&FDirectionalLightComponentDetails::MakeInstance));
-	RegisterCustomPropertyLayout("StaticMeshActor", FOnGetDetailCustomizationInstance::CreateStatic(&FStaticMeshActorDetails::MakeInstance));
-	RegisterCustomPropertyLayout("MeshComponent", FOnGetDetailCustomizationInstance::CreateStatic(&FMeshComponentDetails::MakeInstance));
-	RegisterCustomPropertyLayout("MatineeActor", FOnGetDetailCustomizationInstance::CreateStatic(&FMatineeActorDetails::MakeInstance));
-	RegisterCustomPropertyLayout("ReflectionCapture", FOnGetDetailCustomizationInstance::CreateStatic(&FReflectionCaptureDetails::MakeInstance));
-	RegisterCustomPropertyLayout("SkyLight", FOnGetDetailCustomizationInstance::CreateStatic(&FSkyLightComponentDetails::MakeInstance));
-	RegisterCustomPropertyLayout("Brush", FOnGetDetailCustomizationInstance::CreateStatic(&FBrushDetails::MakeInstance));
-	RegisterCustomPropertyLayout("AmbientSound", FOnGetDetailCustomizationInstance::CreateStatic(&FAmbientSoundDetails::MakeInstance));
-	RegisterCustomPropertyLayout("WorldSettings", FOnGetDetailCustomizationInstance::CreateStatic(&FWorldSettingsDetails::MakeInstance));
-	RegisterCustomPropertyLayout("GeneralProjectSettings", FOnGetDetailCustomizationInstance::CreateStatic(&FGeneralProjectSettingsDetails::MakeInstance));
+	RegisterCustomClassLayout("Actor", FOnGetDetailCustomizationInstance::CreateStatic(&FActorDetails::MakeInstance));
+	RegisterCustomClassLayout("SceneComponent", FOnGetDetailCustomizationInstance::CreateStatic(&FSceneComponentDetails::MakeInstance));
+	RegisterCustomClassLayout("PrimitiveComponent", FOnGetDetailCustomizationInstance::CreateStatic(&FPrimitiveComponentDetails::MakeInstance));
+	RegisterCustomClassLayout("StaticMeshComponent", FOnGetDetailCustomizationInstance::CreateStatic(&FStaticMeshComponentDetails::MakeInstance));
+	RegisterCustomClassLayout("SkeletalMeshComponent", FOnGetDetailCustomizationInstance::CreateStatic(&FSkeletalMeshComponentDetails::MakeInstance));
+	RegisterCustomClassLayout("SkinnedMeshComponent", FOnGetDetailCustomizationInstance::CreateStatic(&FSkinnedMeshComponentDetails::MakeInstance));
+	RegisterCustomClassLayout("WheeledVehicleMovementComponent4W", FOnGetDetailCustomizationInstance::CreateStatic(&FWheeledVehicleMovementComponent4WDetails::MakeInstance));
+	RegisterCustomClassLayout("LightComponent", FOnGetDetailCustomizationInstance::CreateStatic(&FLightComponentDetails::MakeInstance));
+	RegisterCustomClassLayout("PointLightComponent", FOnGetDetailCustomizationInstance::CreateStatic(&FPointLightComponentDetails::MakeInstance));
+	RegisterCustomClassLayout("DirectionalLightComponent", FOnGetDetailCustomizationInstance::CreateStatic(&FDirectionalLightComponentDetails::MakeInstance));
+	RegisterCustomClassLayout("StaticMeshActor", FOnGetDetailCustomizationInstance::CreateStatic(&FStaticMeshActorDetails::MakeInstance));
+	RegisterCustomClassLayout("MeshComponent", FOnGetDetailCustomizationInstance::CreateStatic(&FMeshComponentDetails::MakeInstance));
+	RegisterCustomClassLayout("MatineeActor", FOnGetDetailCustomizationInstance::CreateStatic(&FMatineeActorDetails::MakeInstance));
+	RegisterCustomClassLayout("ReflectionCapture", FOnGetDetailCustomizationInstance::CreateStatic(&FReflectionCaptureDetails::MakeInstance));
+	RegisterCustomClassLayout("SkyLight", FOnGetDetailCustomizationInstance::CreateStatic(&FSkyLightComponentDetails::MakeInstance));
+	RegisterCustomClassLayout("Brush", FOnGetDetailCustomizationInstance::CreateStatic(&FBrushDetails::MakeInstance));
+	RegisterCustomClassLayout("AmbientSound", FOnGetDetailCustomizationInstance::CreateStatic(&FAmbientSoundDetails::MakeInstance));
+	RegisterCustomClassLayout("WorldSettings", FOnGetDetailCustomizationInstance::CreateStatic(&FWorldSettingsDetails::MakeInstance));
+	RegisterCustomClassLayout("GeneralProjectSettings", FOnGetDetailCustomizationInstance::CreateStatic(&FGeneralProjectSettingsDetails::MakeInstance));
 
 	//@TODO: A2REMOVAL: Rename FSkeletalControlNodeDetails to something more generic
-	RegisterCustomPropertyLayout("K2Node_StructMemberGet", FOnGetDetailCustomizationInstance::CreateStatic(&FSkeletalControlNodeDetails::MakeInstance));
-	RegisterCustomPropertyLayout("K2Node_StructMemberSet", FOnGetDetailCustomizationInstance::CreateStatic(&FSkeletalControlNodeDetails::MakeInstance));
-
-	RegisterCustomPropertyLayout("AnimSequence", FOnGetDetailCustomizationInstance::CreateStatic(&FAnimSequenceDetails::MakeInstance));
-
-	RegisterCustomPropertyLayout("EditorAnimSegment", FOnGetDetailCustomizationInstance::CreateStatic(&FAnimMontageSegmentDetails::MakeInstance));
-	RegisterCustomPropertyLayout("EditorAnimCompositeSegment", FOnGetDetailCustomizationInstance::CreateStatic(&FAnimMontageSegmentDetails::MakeInstance));
-	RegisterCustomPropertyLayout("EditorSkeletonNotifyObj", FOnGetDetailCustomizationInstance::CreateStatic(&FSkeletonNotifyDetails::MakeInstance));
-	RegisterCustomPropertyLayout("AnimStateNode", FOnGetDetailCustomizationInstance::CreateStatic(&FAnimStateNodeDetails::MakeInstance));
-	RegisterCustomPropertyLayout("AnimStateTransitionNode", FOnGetDetailCustomizationInstance::CreateStatic(&FAnimTransitionNodeDetails::MakeInstance));
-	RegisterCustomPropertyLayout("AnimGraphNode_TransitionPoseEvaluator", FOnGetDetailCustomizationInstance::CreateStatic(&FTransitionPoseEvaluatorNodeDetails::MakeInstance));
-
-	RegisterCustomPropertyLayout("LandscapeSplineControlPoint", FOnGetDetailCustomizationInstance::CreateStatic(&FLandscapeSplineDetails::MakeInstance));
-	RegisterCustomPropertyLayout("LandscapeSplineSegment", FOnGetDetailCustomizationInstance::CreateStatic(&FLandscapeSplineDetails::MakeInstance));
-
-	RegisterCustomPropertyLayout("DialogueWave", FOnGetDetailCustomizationInstance::CreateStatic(&FDialogueWaveDetails::MakeInstance));
-	RegisterCustomPropertyLayout("BodySetup", FOnGetDetailCustomizationInstance::CreateStatic(&FBodySetupDetails::MakeInstance));
-	RegisterCustomPropertyLayout("PhysicsConstraintTemplate", FOnGetDetailCustomizationInstance::CreateStatic(&FPhysicsConstraintComponentDetails::MakeInstance));
-	RegisterCustomPropertyLayout("PhysicsConstraintComponent", FOnGetDetailCustomizationInstance::CreateStatic(&FPhysicsConstraintComponentDetails::MakeInstance));
-	RegisterCustomPropertyLayout("CollisionProfile", FOnGetDetailCustomizationInstance::CreateStatic(&FCollisionProfileDetails::MakeInstance));
-	RegisterCustomPropertyLayout("PhysicsSettings", FOnGetDetailCustomizationInstance::CreateStatic(&FPhysicsSettingsDetails::MakeInstance));
-
-	RegisterCustomPropertyLayout("ParticleModuleRequired", FOnGetDetailCustomizationInstance::CreateStatic(&FParticleModuleRequiredDetails::MakeInstance));
-	RegisterCustomPropertyLayout("ParticleModuleSubUV", FOnGetDetailCustomizationInstance::CreateStatic(&FParticleModuleSubUVDetails::MakeInstance));
-	RegisterCustomPropertyLayout("ParticleModuleAccelerationDrag", FOnGetDetailCustomizationInstance::CreateStatic(&FParticleModuleAccelerationDragDetails::MakeInstance));
-	RegisterCustomPropertyLayout("ParticleModuleAcceleration", FOnGetDetailCustomizationInstance::CreateStatic(&FParticleModuleAccelerationDetails::MakeInstance));
-	RegisterCustomPropertyLayout("ParticleModuleAccelerationDragScaleOverLife", FOnGetDetailCustomizationInstance::CreateStatic(&FParticleModuleAccelerationDragScaleOverLifeDetails::MakeInstance));
-	RegisterCustomPropertyLayout("ParticleModuleCollisionGPU", FOnGetDetailCustomizationInstance::CreateStatic(&FParticleModuleCollisionGPUDetails::MakeInstance));
-	RegisterCustomPropertyLayout("ParticleModuleOrbit", FOnGetDetailCustomizationInstance::CreateStatic(&FParticleModuleOrbitDetails::MakeInstance));
-	RegisterCustomPropertyLayout("ParticleModuleSizeMultiplyLife", FOnGetDetailCustomizationInstance::CreateStatic(&FParticleModuleSizeMultiplyLifeDetails::MakeInstance));
-	RegisterCustomPropertyLayout("ParticleModuleSizeScale", FOnGetDetailCustomizationInstance::CreateStatic(&FParticleModuleSizeScaleDetails::MakeInstance));
-	RegisterCustomPropertyLayout("ParticleModuleVectorFieldScale", FOnGetDetailCustomizationInstance::CreateStatic(&FParticleModuleVectorFieldScaleDetails::MakeInstance));
-	RegisterCustomPropertyLayout("ParticleModuleVectorFieldScaleOverLife", FOnGetDetailCustomizationInstance::CreateStatic(&FParticleModuleVectorFieldScaleOverLifeDetails::MakeInstance));
-
-	RegisterCustomPropertyLayout("CameraComponent", FOnGetDetailCustomizationInstance::CreateStatic(&FCameraDetails::MakeInstance));
-	RegisterCustomPropertyLayout("DeviceProfile", FOnGetDetailCustomizationInstance::CreateStatic(&FDeviceProfileDetails::MakeInstance));
-	RegisterCustomPropertyLayout("InternationalizationSettingsModel", FOnGetDetailCustomizationInstance::CreateStatic(&FInternationalizationSettingsModelDetails::MakeInstance));
-	RegisterCustomPropertyLayout("InputSettings", FOnGetDetailCustomizationInstance::CreateStatic(&FInputSettingsDetails::MakeInstance));
-	RegisterCustomPropertyLayout("WindowsTargetSettings", FOnGetDetailCustomizationInstance::CreateStatic(&FWindowsTargetSettingsDetails::MakeInstance));
-	RegisterCustomPropertyLayout("MoviePlayerSettings", FOnGetDetailCustomizationInstance::CreateStatic(&FMoviePlayerSettingsDetails::MakeInstance));
-
-	RegisterCustomPropertyLayout("SourceCodeAccessSettings", FOnGetDetailCustomizationInstance::CreateStatic(&FSourceCodeAccessSettingsDetails::MakeInstance));
-	RegisterCustomPropertyLayout("ParticleSystemComponent", FOnGetDetailCustomizationInstance::CreateStatic(&FParticleSystemComponentDetails::MakeInstance));
+	RegisterCustomClassLayout("K2Node_StructMemberGet", FOnGetDetailCustomizationInstance::CreateStatic(&FSkeletalControlNodeDetails::MakeInstance));
+	RegisterCustomClassLayout("K2Node_StructMemberSet", FOnGetDetailCustomizationInstance::CreateStatic(&FSkeletalControlNodeDetails::MakeInstance));
+
+	RegisterCustomClassLayout("AnimSequence", FOnGetDetailCustomizationInstance::CreateStatic(&FAnimSequenceDetails::MakeInstance));
+
+	RegisterCustomClassLayout("EditorAnimSegment", FOnGetDetailCustomizationInstance::CreateStatic(&FAnimMontageSegmentDetails::MakeInstance));
+	RegisterCustomClassLayout("EditorAnimCompositeSegment", FOnGetDetailCustomizationInstance::CreateStatic(&FAnimMontageSegmentDetails::MakeInstance));
+	RegisterCustomClassLayout("EditorSkeletonNotifyObj", FOnGetDetailCustomizationInstance::CreateStatic(&FSkeletonNotifyDetails::MakeInstance));
+	RegisterCustomClassLayout("AnimStateNode", FOnGetDetailCustomizationInstance::CreateStatic(&FAnimStateNodeDetails::MakeInstance));
+	RegisterCustomClassLayout("AnimStateTransitionNode", FOnGetDetailCustomizationInstance::CreateStatic(&FAnimTransitionNodeDetails::MakeInstance));
+	RegisterCustomClassLayout("AnimGraphNode_TransitionPoseEvaluator", FOnGetDetailCustomizationInstance::CreateStatic(&FTransitionPoseEvaluatorNodeDetails::MakeInstance));
+
+	RegisterCustomClassLayout("LandscapeSplineControlPoint", FOnGetDetailCustomizationInstance::CreateStatic(&FLandscapeSplineDetails::MakeInstance));
+	RegisterCustomClassLayout("LandscapeSplineSegment", FOnGetDetailCustomizationInstance::CreateStatic(&FLandscapeSplineDetails::MakeInstance));
+
+	RegisterCustomClassLayout("DialogueWave", FOnGetDetailCustomizationInstance::CreateStatic(&FDialogueWaveDetails::MakeInstance));
+	RegisterCustomClassLayout("BodySetup", FOnGetDetailCustomizationInstance::CreateStatic(&FBodySetupDetails::MakeInstance));
+	RegisterCustomClassLayout("PhysicsConstraintTemplate", FOnGetDetailCustomizationInstance::CreateStatic(&FPhysicsConstraintComponentDetails::MakeInstance));
+	RegisterCustomClassLayout("PhysicsConstraintComponent", FOnGetDetailCustomizationInstance::CreateStatic(&FPhysicsConstraintComponentDetails::MakeInstance));
+	RegisterCustomClassLayout("CollisionProfile", FOnGetDetailCustomizationInstance::CreateStatic(&FCollisionProfileDetails::MakeInstance));
+	RegisterCustomClassLayout("PhysicsSettings", FOnGetDetailCustomizationInstance::CreateStatic(&FPhysicsSettingsDetails::MakeInstance));
+
+	RegisterCustomClassLayout("ParticleModuleRequired", FOnGetDetailCustomizationInstance::CreateStatic(&FParticleModuleRequiredDetails::MakeInstance));
+	RegisterCustomClassLayout("ParticleModuleSubUV", FOnGetDetailCustomizationInstance::CreateStatic(&FParticleModuleSubUVDetails::MakeInstance));
+	RegisterCustomClassLayout("ParticleModuleAccelerationDrag", FOnGetDetailCustomizationInstance::CreateStatic(&FParticleModuleAccelerationDragDetails::MakeInstance));
+	RegisterCustomClassLayout("ParticleModuleAcceleration", FOnGetDetailCustomizationInstance::CreateStatic(&FParticleModuleAccelerationDetails::MakeInstance));
+	RegisterCustomClassLayout("ParticleModuleAccelerationDragScaleOverLife", FOnGetDetailCustomizationInstance::CreateStatic(&FParticleModuleAccelerationDragScaleOverLifeDetails::MakeInstance));
+	RegisterCustomClassLayout("ParticleModuleCollisionGPU", FOnGetDetailCustomizationInstance::CreateStatic(&FParticleModuleCollisionGPUDetails::MakeInstance));
+	RegisterCustomClassLayout("ParticleModuleOrbit", FOnGetDetailCustomizationInstance::CreateStatic(&FParticleModuleOrbitDetails::MakeInstance));
+	RegisterCustomClassLayout("ParticleModuleSizeMultiplyLife", FOnGetDetailCustomizationInstance::CreateStatic(&FParticleModuleSizeMultiplyLifeDetails::MakeInstance));
+	RegisterCustomClassLayout("ParticleModuleSizeScale", FOnGetDetailCustomizationInstance::CreateStatic(&FParticleModuleSizeScaleDetails::MakeInstance));
+	RegisterCustomClassLayout("ParticleModuleVectorFieldScale", FOnGetDetailCustomizationInstance::CreateStatic(&FParticleModuleVectorFieldScaleDetails::MakeInstance));
+	RegisterCustomClassLayout("ParticleModuleVectorFieldScaleOverLife", FOnGetDetailCustomizationInstance::CreateStatic(&FParticleModuleVectorFieldScaleOverLifeDetails::MakeInstance));
+
+	RegisterCustomClassLayout("CameraComponent", FOnGetDetailCustomizationInstance::CreateStatic(&FCameraDetails::MakeInstance));
+	RegisterCustomClassLayout("DeviceProfile", FOnGetDetailCustomizationInstance::CreateStatic(&FDeviceProfileDetails::MakeInstance));
+	RegisterCustomClassLayout("InternationalizationSettingsModel", FOnGetDetailCustomizationInstance::CreateStatic(&FInternationalizationSettingsModelDetails::MakeInstance));
+	RegisterCustomClassLayout("InputSettings", FOnGetDetailCustomizationInstance::CreateStatic(&FInputSettingsDetails::MakeInstance));
+	RegisterCustomClassLayout("WindowsTargetSettings", FOnGetDetailCustomizationInstance::CreateStatic(&FWindowsTargetSettingsDetails::MakeInstance));
+	RegisterCustomClassLayout("MoviePlayerSettings", FOnGetDetailCustomizationInstance::CreateStatic(&FMoviePlayerSettingsDetails::MakeInstance));
+
+	RegisterCustomClassLayout("SourceCodeAccessSettings", FOnGetDetailCustomizationInstance::CreateStatic(&FSourceCodeAccessSettingsDetails::MakeInstance));
+	RegisterCustomClassLayout("ParticleSystemComponent", FOnGetDetailCustomizationInstance::CreateStatic(&FParticleSystemComponentDetails::MakeInstance));
 }
 
-void FDetailCustomizationsModule::RegisterCustomPropertyLayout(FName ClassName, FOnGetDetailCustomizationInstance DetailLayoutDelegate )
+void FDetailCustomizationsModule::RegisterCustomClassLayout(FName ClassName, FOnGetDetailCustomizationInstance DetailLayoutDelegate )
 {
 	check( ClassName != NAME_None );
 
@@ -250,7 +250,7 @@ void FDetailCustomizationsModule::RegisterCustomPropertyLayout(FName ClassName,
 
 	FPropertyEditorModule& PropertyModule = FModuleManager::GetModuleChecked<FPropertyEditorModule>(PropertyEditor);
 
-	PropertyModule.RegisterCustomPropertyLayout( ClassName, DetailLayoutDelegate );
+	PropertyModule.RegisterCustomClassLayout( ClassName, DetailLayoutDelegate );
 }
 
 
diff --git a/Engine/Source/Editor/DetailCustomizations/Public/DetailCustomizations.h b/Engine/Source/Editor/DetailCustomizations/Public/DetailCustomizations.h
--- a/Engine/Source/Editor/DetailCustomizations/Public/DetailCustomizations.h
+++ b/Engine/Source/Editor/DetailCustomizations/Public/DetailCustomizations.h
@@ -37,7 +37,7 @@ private:
 	 * @param ClassName				The class name to register for property customization
 	 * @param DetailLayoutDelegate	The delegate to call to get the custom detail layout instance
 	 */
-	void RegisterCustomPropertyLayout(FName ClassName, FOnGetDetailCustomizationInstance DetailLayoutDelegate );
+	void RegisterCustomClassLayout(FName ClassName, FOnGetDetailCustomizationInstance DetailLayoutDelegate );
 
 	/**
 	* Registers a custom struct
diff --git a/Engine/Source/Editor/EnvironmentQueryEditor/Private/EnvironmentQueryEditorModule.cpp b/Engine/Source/Editor/EnvironmentQueryEditor/Private/EnvironmentQueryEditorModule.cpp
--- a/Engine/Source/Editor/EnvironmentQueryEditor/Private/EnvironmentQueryEditorModule.cpp
+++ b/Engine/Source/Editor/EnvironmentQueryEditor/Private/EnvironmentQueryEditorModule.cpp
@@ -51,7 +51,7 @@ void FEnvironmentQueryEditorModule::StartupModule()
 	PropertyModule.RegisterCustomPropertyTypeLayout( "EnvBoolParam", FOnGetPropertyTypeCustomizationInstance::CreateStatic( &FEnvQueryParamSetupCustomization::MakeInstance ) );
 	PropertyModule.RegisterCustomPropertyTypeLayout( "EnvDirection", FOnGetPropertyTypeCustomizationInstance::CreateStatic( &FEnvDirectionCustomization::MakeInstance ) );
 	PropertyModule.RegisterCustomPropertyTypeLayout( "EnvTraceData", FOnGetPropertyTypeCustomizationInstance::CreateStatic( &FEnvTraceDataCustomization::MakeInstance ) );
-	PropertyModule.RegisterCustomPropertyLayout( "EnvQueryTest", FOnGetDetailCustomizationInstance::CreateStatic( &FEnvQueryTestDetails::MakeInstance ) );
+	PropertyModule.RegisterCustomClassLayout( "EnvQueryTest", FOnGetDetailCustomizationInstance::CreateStatic( &FEnvQueryTestDetails::MakeInstance ) );
 	PropertyModule.NotifyCustomizationModuleChanged();
 }
 
@@ -92,7 +92,7 @@ void FEnvironmentQueryEditorModule::ShutdownModule()
 		PropertyModule.UnregisterCustomPropertyTypeLayout( "EnvDirection" );
 		PropertyModule.UnregisterCustomPropertyTypeLayout( "EnvTraceData" );
 
-		PropertyModule.UnregisterCustomPropertyLayout( "EnvQueryTest" );
+		PropertyModule.UnregisterCustomClassLayout( "EnvQueryTest" );
 		PropertyModule.NotifyCustomizationModuleChanged();
 	}
 }
diff --git a/Engine/Source/Editor/GameplayAbilitiesEditor/Private/AttributeDetails.cpp b/Engine/Source/Editor/GameplayAbilitiesEditor/Private/AttributeDetails.cpp
--- a/Engine/Source/Editor/GameplayAbilitiesEditor/Private/AttributeDetails.cpp
+++ b/Engine/Source/Editor/GameplayAbilitiesEditor/Private/AttributeDetails.cpp
@@ -23,7 +23,7 @@ TSharedRef<IPropertyTypeCustomization> FAttributePropertyDetails::MakeInstance()
 }
 
 BEGIN_SLATE_FUNCTION_BUILD_OPTIMIZATION
-void FAttributePropertyDetails::CustomizeHeader( TSharedRef<IPropertyHandle> StructPropertyHandle, class FDetailWidgetRow& HeaderRow, IStructCustomizationUtils& StructCustomizationUtils )
+void FAttributePropertyDetails::CustomizeHeader( TSharedRef<IPropertyHandle> StructPropertyHandle, class FDetailWidgetRow& HeaderRow, IPropertyTypeCustomizationUtils& StructCustomizationUtils )
 {
 	MyProperty = StructPropertyHandle->GetChildHandle(GET_MEMBER_NAME_CHECKED(FGameplayAttribute,Attribute));
 
@@ -78,7 +78,7 @@ void FAttributePropertyDetails::CustomizeHeader( TSharedRef<IPropertyHandle> Str
 }
 END_SLATE_FUNCTION_BUILD_OPTIMIZATION
 
-void FAttributePropertyDetails::CustomizeChildren( TSharedRef<IPropertyHandle> StructPropertyHandle, class IDetailChildrenBuilder& StructBuilder, IStructCustomizationUtils& StructCustomizationUtils )
+void FAttributePropertyDetails::CustomizeChildren( TSharedRef<IPropertyHandle> StructPropertyHandle, class IDetailChildrenBuilder& StructBuilder, IPropertyTypeCustomizationUtils& StructCustomizationUtils )
 {
 
 
@@ -238,7 +238,7 @@ TSharedRef<IPropertyTypeCustomization> FScalableFloatDetails::MakeInstance()
 }
 
 BEGIN_SLATE_FUNCTION_BUILD_OPTIMIZATION
-void FScalableFloatDetails::CustomizeHeader( TSharedRef<class IPropertyHandle> StructPropertyHandle, class FDetailWidgetRow& HeaderRow, IStructCustomizationUtils& StructCustomizationUtils )
+void FScalableFloatDetails::CustomizeHeader( TSharedRef<class IPropertyHandle> StructPropertyHandle, class FDetailWidgetRow& HeaderRow, IPropertyTypeCustomizationUtils& StructCustomizationUtils )
 {
 	uint32 NumChildren = 0;
 	StructPropertyHandle->GetNumChildren(NumChildren);
@@ -471,7 +471,7 @@ bool FScalableFloatDetails::IsEditable( ) const
 	return true;
 }
 
-void FScalableFloatDetails::CustomizeChildren( TSharedRef<class IPropertyHandle> StructPropertyHandle, class IDetailChildrenBuilder& StructBuilder, IStructCustomizationUtils& StructCustomizationUtils )
+void FScalableFloatDetails::CustomizeChildren( TSharedRef<class IPropertyHandle> StructPropertyHandle, class IDetailChildrenBuilder& StructBuilder, IPropertyTypeCustomizationUtils& StructCustomizationUtils )
 {
 	
 }
diff --git a/Engine/Source/Editor/GameplayAbilitiesEditor/Private/AttributeDetails.h b/Engine/Source/Editor/GameplayAbilitiesEditor/Private/AttributeDetails.h
--- a/Engine/Source/Editor/GameplayAbilitiesEditor/Private/AttributeDetails.h
+++ b/Engine/Source/Editor/GameplayAbilitiesEditor/Private/AttributeDetails.h
@@ -36,8 +36,8 @@ public:
 	static TSharedRef<IPropertyTypeCustomization> MakeInstance();
 
 	/** IPropertyTypeCustomization interface */
-	virtual void CustomizeHeader( TSharedRef<class IPropertyHandle> StructPropertyHandle, class FDetailWidgetRow& HeaderRow, IStructCustomizationUtils& StructCustomizationUtils ) OVERRIDE;
-	virtual void CustomizeChildren( TSharedRef<class IPropertyHandle> StructPropertyHandle, class IDetailChildrenBuilder& StructBuilder, IStructCustomizationUtils& StructCustomizationUtils ) OVERRIDE;              
+	virtual void CustomizeHeader( TSharedRef<class IPropertyHandle> StructPropertyHandle, class FDetailWidgetRow& HeaderRow, IPropertyTypeCustomizationUtils& StructCustomizationUtils ) OVERRIDE;
+	virtual void CustomizeChildren( TSharedRef<class IPropertyHandle> StructPropertyHandle, class IDetailChildrenBuilder& StructBuilder, IPropertyTypeCustomizationUtils& StructCustomizationUtils ) OVERRIDE;              
 
 private:
 
@@ -57,8 +57,8 @@ public:
 
 private:
 
-	virtual void CustomizeHeader( TSharedRef<class IPropertyHandle> StructPropertyHandle, class FDetailWidgetRow& HeaderRow, IStructCustomizationUtils& StructCustomizationUtils ) OVERRIDE;
-	virtual void CustomizeChildren( TSharedRef<class IPropertyHandle> StructPropertyHandle, class IDetailChildrenBuilder& StructBuilder, IStructCustomizationUtils& StructCustomizationUtils ) OVERRIDE;              
+	virtual void CustomizeHeader( TSharedRef<class IPropertyHandle> StructPropertyHandle, class FDetailWidgetRow& HeaderRow, IPropertyTypeCustomizationUtils& StructCustomizationUtils ) OVERRIDE;
+	virtual void CustomizeChildren( TSharedRef<class IPropertyHandle> StructPropertyHandle, class IDetailChildrenBuilder& StructBuilder, IPropertyTypeCustomizationUtils& StructCustomizationUtils ) OVERRIDE;              
 
 	bool IsEditable( ) const;
 
diff --git a/Engine/Source/Editor/GameplayAbilitiesEditor/Private/GameplayAbilitiesEditorModule.cpp b/Engine/Source/Editor/GameplayAbilitiesEditor/Private/GameplayAbilitiesEditorModule.cpp
--- a/Engine/Source/Editor/GameplayAbilitiesEditor/Private/GameplayAbilitiesEditorModule.cpp
+++ b/Engine/Source/Editor/GameplayAbilitiesEditor/Private/GameplayAbilitiesEditorModule.cpp
@@ -22,7 +22,7 @@ void FGameplayAbilitiesEditorModule::StartupModule()
 	PropertyModule.RegisterCustomPropertyTypeLayout( "GameplayAttribute", FOnGetPropertyTypeCustomizationInstance::CreateStatic( &FAttributePropertyDetails::MakeInstance ) );
 	PropertyModule.RegisterCustomPropertyTypeLayout( "ScalableFloat", FOnGetPropertyTypeCustomizationInstance::CreateStatic( &FScalableFloatDetails::MakeInstance ) );
 
-	PropertyModule.RegisterCustomPropertyLayout( "AttributeSet", FOnGetDetailCustomizationInstance::CreateStatic( &FAttributeDetails::MakeInstance ) );
+	PropertyModule.RegisterCustomClassLayout( "AttributeSet", FOnGetDetailCustomizationInstance::CreateStatic( &FAttributeDetails::MakeInstance ) );
 }
 
 void FGameplayAbilitiesEditorModule::ShutdownModule()
diff --git a/Engine/Source/Editor/LandscapeEditor/Private/LandscapeEditorModule.cpp b/Engine/Source/Editor/LandscapeEditor/Private/LandscapeEditorModule.cpp
--- a/Engine/Source/Editor/LandscapeEditor/Private/LandscapeEditorModule.cpp
+++ b/Engine/Source/Editor/LandscapeEditor/Private/LandscapeEditorModule.cpp
@@ -37,7 +37,7 @@ public:
 		EdModeLandscape = NewEditorMode;
 
 		FPropertyEditorModule& PropertyModule = FModuleManager::GetModuleChecked<FPropertyEditorModule>("PropertyEditor");
-		PropertyModule.RegisterCustomPropertyLayout("LandscapeEditorObject", FOnGetDetailCustomizationInstance::CreateStatic(&FLandscapeEditorDetails::MakeInstance));
+		PropertyModule.RegisterCustomClassLayout("LandscapeEditorObject", FOnGetDetailCustomizationInstance::CreateStatic(&FLandscapeEditorDetails::MakeInstance));
 		PropertyModule.RegisterCustomPropertyTypeLayout("GizmoImportLayer", FOnGetPropertyTypeCustomizationInstance::CreateStatic(&FLandscapeEditorStructCustomization_FGizmoImportLayer::MakeInstance));
 		PropertyModule.RegisterCustomPropertyTypeLayout("LandscapeImportLayer", FOnGetPropertyTypeCustomizationInstance::CreateStatic(&FLandscapeEditorStructCustomization_FLandscapeImportLayer::MakeInstance));
 
diff --git a/Engine/Source/Editor/Persona/Private/PersonaModule.cpp b/Engine/Source/Editor/Persona/Private/PersonaModule.cpp
--- a/Engine/Source/Editor/Persona/Private/PersonaModule.cpp
+++ b/Engine/Source/Editor/Persona/Private/PersonaModule.cpp
@@ -55,10 +55,10 @@ void FPersonaModule::StartupModule()
 	}
 	{
 		FPropertyEditorModule& PropertyModule = FModuleManager::LoadModuleChecked<FPropertyEditorModule>("PropertyEditor");
-		PropertyModule.RegisterCustomPropertyLayout( "SkeletalMeshSocket", FOnGetDetailCustomizationInstance::CreateStatic( &FSkeletalMeshSocketDetails::MakeInstance ) );
-		PropertyModule.RegisterCustomPropertyLayout( "AnimNotify", FOnGetDetailCustomizationInstance::CreateStatic( &FAnimNotifyDetails::MakeInstance ) );
-		PropertyModule.RegisterCustomPropertyLayout( "AnimNotifyState", FOnGetDetailCustomizationInstance::CreateStatic(&FAnimNotifyDetails::MakeInstance));
-		PropertyModule.RegisterCustomPropertyLayout( "AnimGraphNode_Base", FOnGetDetailCustomizationInstance::CreateStatic( &FAnimGraphNodeDetails::MakeInstance ) );
+		PropertyModule.RegisterCustomClassLayout( "SkeletalMeshSocket", FOnGetDetailCustomizationInstance::CreateStatic( &FSkeletalMeshSocketDetails::MakeInstance ) );
+		PropertyModule.RegisterCustomClassLayout( "AnimNotify", FOnGetDetailCustomizationInstance::CreateStatic( &FAnimNotifyDetails::MakeInstance ) );
+		PropertyModule.RegisterCustomClassLayout( "AnimNotifyState", FOnGetDetailCustomizationInstance::CreateStatic(&FAnimNotifyDetails::MakeInstance));
+		PropertyModule.RegisterCustomClassLayout( "AnimGraphNode_Base", FOnGetDetailCustomizationInstance::CreateStatic( &FAnimGraphNodeDetails::MakeInstance ) );
 
 		PropertyModule.RegisterCustomPropertyTypeLayout( "InputScaleBias", FOnGetPropertyTypeCustomizationInstance::CreateStatic( &FInputScaleBiasCustomization::MakeInstance ) );
 		PropertyModule.RegisterCustomPropertyTypeLayout( "BoneReference", FOnGetPropertyTypeCustomizationInstance::CreateStatic( &FBoneReferenceCustomization::MakeInstance ) );
@@ -73,10 +73,10 @@ void FPersonaModule::ShutdownModule()
 	if(FModuleManager::Get().IsModuleLoaded("PropertyEditor"))
 	{
 		FPropertyEditorModule& PropertyModule = FModuleManager::LoadModuleChecked<FPropertyEditorModule>("PropertyEditor");
-		PropertyModule.UnregisterCustomPropertyLayout("SkeletalMeshSocket");
-		PropertyModule.UnregisterCustomPropertyLayout("AnimNotify");
-		PropertyModule.UnregisterCustomPropertyLayout("AnimNotifyState");
-		PropertyModule.UnregisterCustomPropertyLayout("AnimGraphNode_Base");
+		PropertyModule.UnregisterCustomClassLayout("SkeletalMeshSocket");
+		PropertyModule.UnregisterCustomClassLayout("AnimNotify");
+		PropertyModule.UnregisterCustomClassLayout("AnimNotifyState");
+		PropertyModule.UnregisterCustomClassLayout("AnimGraphNode_Base");
 
 		PropertyModule.UnregisterCustomPropertyTypeLayout("InputScaleBias");
 		PropertyModule.UnregisterCustomPropertyTypeLayout("BoneReference");
diff --git a/Engine/Source/Editor/PropertyEditor/Private/PropertyEditorModule.cpp b/Engine/Source/Editor/PropertyEditor/Private/PropertyEditorModule.cpp
--- a/Engine/Source/Editor/PropertyEditor/Private/PropertyEditorModule.cpp
+++ b/Engine/Source/Editor/PropertyEditor/Private/PropertyEditorModule.cpp
@@ -367,21 +367,25 @@ TSharedRef<IPropertyChangeListener> FPropertyEditorModule::CreatePropertyChangeL
 	return MakeShareable( new FPropertyChangeListener );
 }
 
-
-void FPropertyEditorModule::RegisterCustomPropertyLayout( FName ClassName, FOnGetDetailCustomizationInstance DetailLayoutDelegate )
+void FPropertyEditorModule::RegisterCustomClassLayout( FName ClassName, FOnGetDetailCustomizationInstance DetailLayoutDelegate )
 {
-	if( ClassName != NAME_None )
+	if (ClassName != NAME_None)
 	{
 		FDetailLayoutCallback Callback;
 		Callback.DetailLayoutDelegate = DetailLayoutDelegate;
 		// @todo: DetailsView: Fix me: this specifies the order in which detail layouts should be queried
 		Callback.Order = ClassNameToDetailLayoutNameMap.Num();
 
-		ClassNameToDetailLayoutNameMap.Add( ClassName, Callback );
+		ClassNameToDetailLayoutNameMap.Add(ClassName, Callback);
 	}
 }
 
-void FPropertyEditorModule::UnregisterCustomPropertyLayout(FName ClassName)
+void FPropertyEditorModule::RegisterCustomPropertyLayout( FName ClassName, FOnGetDetailCustomizationInstance DetailLayoutDelegate )
+{
+	RegisterCustomClassLayout( ClassName, DetailLayoutDelegate );
+}
+
+void FPropertyEditorModule::UnregisterCustomClassLayout( FName ClassName )
 {
 	if (ClassName != NAME_None)
 	{
@@ -389,6 +393,11 @@ void FPropertyEditorModule::UnregisterCustomPropertyLayout(FName ClassName)
 	}
 }
 
+void FPropertyEditorModule::UnregisterCustomPropertyLayout( FName ClassName )
+{
+	UnregisterCustomClassLayout(ClassName);
+}
+
 void FPropertyEditorModule::RegisterStructPropertyLayout( FName StructTypeName, FOnGetStructCustomizationInstance StructLayoutDelegate )
 {
 	if (StructTypeName != NAME_None)
diff --git a/Engine/Source/Editor/PropertyEditor/Public/PropertyEditorModule.h b/Engine/Source/Editor/PropertyEditor/Public/PropertyEditorModule.h
--- a/Engine/Source/Editor/PropertyEditor/Public/PropertyEditorModule.h
+++ b/Engine/Source/Editor/PropertyEditor/Public/PropertyEditorModule.h
@@ -162,6 +162,9 @@ public:
 	 * @param ClassName	The name of the class that the custom detail layout is for
 	 * @param DetailLayoutDelegate	The delegate to call when querying for custom detail layouts for the classes properties
 	 */
+	virtual void RegisterCustomClassLayout( FName ClassName, FOnGetDetailCustomizationInstance DetailLayoutDelegate );
+
+	DEPRECATED(4.3, "RegisterCustomPropertyLayout is deprecated use RegisterCustomClassLayout instead.")
 	virtual void RegisterCustomPropertyLayout( FName ClassName, FOnGetDetailCustomizationInstance DetailLayoutDelegate );
 
 	/**
@@ -169,6 +172,9 @@ public:
 	 *
 	 * @param ClassName	The class name with the custom detail layout delegate to remove
 	 */
+	virtual void UnregisterCustomClassLayout( FName ClassName );
+
+	DEPRECATED(4.3, "UnregisterCustomPropertyLayout is deprecated use UnregisterCustomClassLayout instead.")
 	virtual void UnregisterCustomPropertyLayout( FName ClassName );
 
 	/**
diff --git a/Engine/Source/Editor/UnrealEd/Private/UnrealEdEngine.cpp b/Engine/Source/Editor/UnrealEd/Private/UnrealEdEngine.cpp
--- a/Engine/Source/Editor/UnrealEd/Private/UnrealEdEngine.cpp
+++ b/Engine/Source/Editor/UnrealEd/Private/UnrealEdEngine.cpp
@@ -95,10 +95,10 @@ void UUnrealEdEngine::Init(IEngineLoop* InEngineLoop)
 	{
 		FPropertyEditorModule& PropertyModule = FModuleManager::LoadModuleChecked<FPropertyEditorModule>("PropertyEditor");
 
-		PropertyModule.RegisterCustomPropertyLayout("EditorLoadingSavingSettings", FOnGetDetailCustomizationInstance::CreateStatic(&FEditorLoadingSavingSettingsCustomization::MakeInstance));
-		PropertyModule.RegisterCustomPropertyLayout("GameMapsSettings", FOnGetDetailCustomizationInstance::CreateStatic(&FGameMapsSettingsCustomization::MakeInstance));
-		PropertyModule.RegisterCustomPropertyLayout("LevelEditorPlaySettings", FOnGetDetailCustomizationInstance::CreateStatic(&FLevelEditorPlaySettingsCustomization::MakeInstance));
-		PropertyModule.RegisterCustomPropertyLayout("ProjectPackagingSettings", FOnGetDetailCustomizationInstance::CreateStatic(&FProjectPackagingSettingsCustomization::MakeInstance));
+		PropertyModule.RegisterCustomClassLayout("EditorLoadingSavingSettings", FOnGetDetailCustomizationInstance::CreateStatic(&FEditorLoadingSavingSettingsCustomization::MakeInstance));
+		PropertyModule.RegisterCustomClassLayout("GameMapsSettings", FOnGetDetailCustomizationInstance::CreateStatic(&FGameMapsSettingsCustomization::MakeInstance));
+		PropertyModule.RegisterCustomClassLayout("LevelEditorPlaySettings", FOnGetDetailCustomizationInstance::CreateStatic(&FLevelEditorPlaySettingsCustomization::MakeInstance));
+		PropertyModule.RegisterCustomClassLayout("ProjectPackagingSettings", FOnGetDetailCustomizationInstance::CreateStatic(&FProjectPackagingSettingsCustomization::MakeInstance));
 	}
 
 	bool bCookOnTheSide = FParse::Param( FCommandLine::Get(),TEXT("COOKONTHESIDE"));
