Fixed not being able to convert volumes to other volume types

diff --git a/Engine/Source/Editor/DetailCustomizations/Private/ActorDetails.cpp b/Engine/Source/Editor/DetailCustomizations/Private/ActorDetails.cpp
--- a/Engine/Source/Editor/DetailCustomizations/Private/ActorDetails.cpp
+++ b/Engine/Source/Editor/DetailCustomizations/Private/ActorDetails.cpp
@@ -25,7 +25,7 @@
 #include "ComponentsTree.h"
 #include "IPropertyUtilities.h"
 #include "IDocumentation.h"
-
+#include "Runtime/Engine/Classes/Engine/BrushShape.h"
 #include "ActorDetailsDelegates.h"
 
 #define LOCTEXT_NAMESPACE "ActorDetails"
@@ -248,8 +248,17 @@ void FActorDetails::CreateClassPickerConvertActorFilter(const TWeakObjectPtr<AAc
 
 	// Never convert to the same class
 	Filter->DisallowedClasses.Add(ConvertClass);
+
+	if( ConvertActor->IsA<AVolume>() )
+	{
+		// Volumes cannot be converted to brushes or brush shapes or the abstract type
+		Filter->DisallowedClasses.Add(ABrush::StaticClass());
+		Filter->DisallowedClasses.Add(ABrushShape::StaticClass());
+		Filter->DisallowedClasses.Add(AVolume::StaticClass());
+	}
 }
 
+
 TSharedRef<SWidget> FActorDetails::OnGetConvertContent()
 {
 	// Build a class picker widget
@@ -295,31 +304,19 @@ EVisibility FActorDetails::GetConvertMenuVisibility() const
 
 TSharedRef<SWidget> FActorDetails::MakeConvertMenu( const FSelectedActorInfo& SelectedActorInfo )
 {
-	if( SelectedActorInfo.HasConvertableAsset() )
-	{
-		UClass* RootConversionClass = GetConversionRoot(SelectedActorInfo.SelectionClass);
-		const FSlateFontInfo& FontInfo = FEditorStyle::GetFontStyle( TEXT("ExpandableArea.NormalFont") );
-		return
-			SNew( SBox )
-			.HAlign( HAlign_Left )
-			[
-				SNew( SComboButton )
-					.IsEnabled(RootConversionClass != NULL)
-					.Visibility( this, &FActorDetails::GetConvertMenuVisibility )
-					.OnGetMenuContent( this, &FActorDetails::OnGetConvertContent )
-					.ButtonContent()
-				[
-					SNew( STextBlock )
-					.Text( LOCTEXT("ConvertButton", "Convert") ) 
-					.ToolTipText( LOCTEXT("ConvertButton_ToolTip", "Convert actors to different types") )
-					.Font( FontInfo )
-				]
-			];
-
-	}
-
-	// Nothing to convert
-	return SNullWidget::NullWidget;
+	UClass* RootConversionClass = GetConversionRoot(SelectedActorInfo.SelectionClass);
+	return
+		SNew(SComboButton)
+		.ContentPadding(2)
+		.IsEnabled(RootConversionClass != NULL)
+		.Visibility(this, &FActorDetails::GetConvertMenuVisibility)
+		.OnGetMenuContent(this, &FActorDetails::OnGetConvertContent)
+		.ButtonContent()
+		[
+			SNew(STextBlock)
+			.Text(LOCTEXT("ConvertButton", "Select a Type"))
+			.Font(IDetailLayoutBuilder::GetDetailFont())
+		];
 
 }
 
@@ -480,11 +477,18 @@ void FActorDetails::AddActorCategory( IDetailLayoutBuilder& DetailBuilder, const
 
 	}
 
-	// WorldSettings should never convert to another class type.
-	if(SelectedActorInfo.SelectionClass != AWorldSettings::StaticClass())
+	// WorldSettings should never convert to another class type
+	if( SelectedActorInfo.SelectionClass != AWorldSettings::StaticClass() && SelectedActorInfo.HasConvertableAsset() )
 	{
-		ActorCategory.AddCustomRow(  LOCTEXT("SelectionFilter", "Selected").ToString() )
+		ActorCategory.AddCustomRow( LOCTEXT("ConvertMenu", "Convert").ToString() )
 		.NameContent()
+		[
+			SNew(STextBlock)
+			.Text(LOCTEXT("ConvertButton", "Convert Actor"))
+			.ToolTipText(LOCTEXT("ConvertButton_ToolTip", "Convert actors to different types"))
+			.Font(IDetailLayoutBuilder::GetDetailFont())
+		]
+		.ValueContent()
 		[
 			MakeConvertMenu( SelectedActorInfo )
 		];
diff --git a/Engine/Source/Editor/UnrealEd/Private/AssetSelection.cpp b/Engine/Source/Editor/UnrealEd/Private/AssetSelection.cpp
--- a/Engine/Source/Editor/UnrealEd/Private/AssetSelection.cpp
+++ b/Engine/Source/Editor/UnrealEd/Private/AssetSelection.cpp
@@ -81,6 +81,7 @@ namespace AssetSelectionUtils
 							ActorInfo.bHaveBuilderBrush = FActorEditorUtils::IsABuilderBrush(Brush);
 						}
 						ActorInfo.bHaveBrush |= true;
+						ActorInfo.bHaveVolume |= Brush->IsVolumeBrush();
 					}
 
 					UClass* CurrentClass = CurrentActor->GetClass();
diff --git a/Engine/Source/Editor/UnrealEd/Public/AssetSelection.h b/Engine/Source/Editor/UnrealEd/Public/AssetSelection.h
--- a/Engine/Source/Editor/UnrealEd/Public/AssetSelection.h
+++ b/Engine/Source/Editor/UnrealEd/Public/AssetSelection.h
@@ -82,6 +82,8 @@ struct FSelectedActorInfo
 	uint32 bAllSelectedStaticMeshesHaveCollisionModels:1;
 	/** true if a brush is in the selection */
 	uint32 bHaveBrush:1;
+	/** true if a volume brush is in the selection */
+	uint32 bHaveVolume:1;
 	/** true if a builder brush is in the selection */
 	uint32 bHaveBuilderBrush:1;
 	/** true if an actor in the selection has a static mesh component */
@@ -122,6 +124,7 @@ struct FSelectedActorInfo
 		, bAllSelectedActorsBelongToSameWorld(true)
 		, bAllSelectedStaticMeshesHaveCollisionModels(true)
 		, bHaveBrush(false)
+		, bHaveVolume(false)
 		, bHaveBuilderBrush(false)
 		, bHaveStaticMeshComponent(false)
 		, bHaveStaticMesh(false)
@@ -140,7 +143,7 @@ struct FSelectedActorInfo
 	 */
 	bool HasConvertableAsset() const
 	{
-		return NumSelected && !bHaveBuilderBrush && !bHaveBrush;
+		return NumSelected && !bHaveBuilderBrush && (!bHaveBrush || bHaveVolume);
 	}
 
 };
diff --git a/Engine/Source/Runtime/Engine/Classes/GameFramework/Volume.h b/Engine/Source/Runtime/Engine/Classes/GameFramework/Volume.h
--- a/Engine/Source/Runtime/Engine/Classes/GameFramework/Volume.h
+++ b/Engine/Source/Runtime/Engine/Classes/GameFramework/Volume.h
@@ -8,7 +8,7 @@
 #include "Engine/Brush.h"
 #include "Volume.generated.h"
 
-UCLASS(MinimalAPI, showcategories=Collision, hidecategories=(Brush, Physics), abstract)
+UCLASS(MinimalAPI, showcategories=Collision, hidecategories=(Brush, Physics), abstract, ConversionRoot)
 class AVolume : public ABrush
 {
 	GENERATED_UCLASS_BODY()
