Fixed brushes not being able to convert to BSP

diff --git a/Engine/Source/Editor/DetailCustomizations/Private/ActorDetails.cpp b/Engine/Source/Editor/DetailCustomizations/Private/ActorDetails.cpp
--- a/Engine/Source/Editor/DetailCustomizations/Private/ActorDetails.cpp
+++ b/Engine/Source/Editor/DetailCustomizations/Private/ActorDetails.cpp
@@ -249,7 +249,7 @@ void FActorDetails::CreateClassPickerConvertActorFilter(const TWeakObjectPtr<AAc
 	// Never convert to the same class
 	Filter->DisallowedClasses.Add(ConvertClass);
 
-	if( ConvertActor->IsA<AVolume>() )
+	if( ConvertActor->IsA<ABrush>() )
 	{
 		// Volumes cannot be converted to brushes or brush shapes or the abstract type
 		Filter->DisallowedClasses.Add(ABrush::StaticClass());
diff --git a/Engine/Source/Editor/UnrealEd/Public/AssetSelection.h b/Engine/Source/Editor/UnrealEd/Public/AssetSelection.h
--- a/Engine/Source/Editor/UnrealEd/Public/AssetSelection.h
+++ b/Engine/Source/Editor/UnrealEd/Public/AssetSelection.h
@@ -143,7 +143,7 @@ struct FSelectedActorInfo
 	 */
 	bool HasConvertableAsset() const
 	{
-		return NumSelected && !bHaveBuilderBrush && (!bHaveBrush || bHaveVolume);
+		return NumSelected && !bHaveBuilderBrush;
 	}
 
 };
