337764	LIVE: Collision Channel: New object channels not saving after closing and reopening

diff --git a/Engine/Source/Editor/DetailCustomizations/Private/CollisionProfileDetails.cpp b/Engine/Source/Editor/DetailCustomizations/Private/CollisionProfileDetails.cpp
--- a/Engine/Source/Editor/DetailCustomizations/Private/CollisionProfileDetails.cpp
+++ b/Engine/Source/Editor/DetailCustomizations/Private/CollisionProfileDetails.cpp
@@ -1345,6 +1345,9 @@ void FCollisionProfileDetails::CustomizeDetails( IDetailLayoutBuilder& DetailBui
 	CollisionProfile= UCollisionProfile::Get();
 	check(CollisionProfile);
 
+	ProfilesPropertyHandle = DetailBuilder.GetProperty(TEXT("Profiles"));
+	check(ProfilesPropertyHandle->IsValidHandle());
+
 	// save currently loaded data
 	SavedData.Save(CollisionProfile);
 
@@ -1815,6 +1818,7 @@ void FCollisionProfileDetails::UpdateChannel(bool bTraceType)
 	}
 
 	CollisionProfile->LoadProfileConfig(true);
+	ProfilesPropertyHandle->NotifyPostChange();
 
 	// update profile list if the name changed, this has to be refreshed
 	RefreshProfileList();
@@ -1826,6 +1830,7 @@ void FCollisionProfileDetails::UpdateProfile()
 	RefreshProfileList();
 	ProfileListView->RequestListRefresh();
 	CollisionProfile->LoadProfileConfig(true);
+	ProfilesPropertyHandle->NotifyPostChange();
 }
 
 void FCollisionProfileDetails::RefreshChannelList(bool bTraceType)
diff --git a/Engine/Source/Editor/DetailCustomizations/Private/CollisionProfileDetails.h b/Engine/Source/Editor/DetailCustomizations/Private/CollisionProfileDetails.h
--- a/Engine/Source/Editor/DetailCustomizations/Private/CollisionProfileDetails.h
+++ b/Engine/Source/Editor/DetailCustomizations/Private/CollisionProfileDetails.h
@@ -173,5 +173,7 @@ private:
 	};
 
 	FCollsiionProfileData SavedData;
+
+	TSharedPtr<IPropertyHandle> ProfilesPropertyHandle;
 };
 
diff --git a/Engine/Source/Runtime/Engine/Classes/Engine/CollisionProfile.h b/Engine/Source/Runtime/Engine/Classes/Engine/CollisionProfile.h
--- a/Engine/Source/Runtime/Engine/Classes/Engine/CollisionProfile.h
+++ b/Engine/Source/Runtime/Engine/Classes/Engine/CollisionProfile.h
@@ -129,7 +129,9 @@ class UCollisionProfile
 	GENERATED_UCLASS_BODY()
 
 private:
-	UPROPERTY(globalconfig)
+	// This is hacky, but without this edit tag, we can't get valid property handle
+	// and we can't save them properly to config, so we need this tag. 
+	UPROPERTY(globalconfig,Category=Collision, VisibleAnywhere)
 	TArray<FCollisionResponseTemplate>	Profiles;
 
 	UPROPERTY(globalconfig)
diff --git a/Engine/Source/Runtime/Engine/Private/Collision/CollisionProfile.cpp b/Engine/Source/Runtime/Engine/Private/Collision/CollisionProfile.cpp
--- a/Engine/Source/Runtime/Engine/Private/Collision/CollisionProfile.cpp
+++ b/Engine/Source/Runtime/Engine/Private/Collision/CollisionProfile.cpp
@@ -47,7 +47,7 @@ UCollisionProfile* UCollisionProfile::Get()
 {
 	static bool bInitialized = false;
 	// this is singletone. Use default object
-	UCollisionProfile* CollisionProfile = CastChecked<UCollisionProfile>(UCollisionProfile::StaticClass()->GetDefaultObject());
+	UCollisionProfile* CollisionProfile = GetMutableDefault<UCollisionProfile>();
 
 	if (!bInitialized)
 	{
