Fixed so that it can just call SaveConfig.

diff --git a/Engine/Source/Editor/DetailCustomizations/Private/CollisionProfileDetails.cpp b/Engine/Source/Editor/DetailCustomizations/Private/CollisionProfileDetails.cpp
--- a/Engine/Source/Editor/DetailCustomizations/Private/CollisionProfileDetails.cpp
+++ b/Engine/Source/Editor/DetailCustomizations/Private/CollisionProfileDetails.cpp
@@ -1345,9 +1345,6 @@ void FCollisionProfileDetails::CustomizeDetails( IDetailLayoutBuilder& DetailBui
 	CollisionProfile= UCollisionProfile::Get();
 	check(CollisionProfile);
 
-	ProfilesPropertyHandle = DetailBuilder.GetProperty(TEXT("Profiles"));
-	check(ProfilesPropertyHandle->IsValidHandle());
-
 	// save currently loaded data
 	SavedData.Save(CollisionProfile);
 
@@ -1818,7 +1815,7 @@ void FCollisionProfileDetails::UpdateChannel(bool bTraceType)
 	}
 
 	CollisionProfile->LoadProfileConfig(true);
-	ProfilesPropertyHandle->NotifyPostChange();
+	CollisionProfile->SaveConfig();
 
 	// update profile list if the name changed, this has to be refreshed
 	RefreshProfileList();
@@ -1830,7 +1827,7 @@ void FCollisionProfileDetails::UpdateProfile()
 	RefreshProfileList();
 	ProfileListView->RequestListRefresh();
 	CollisionProfile->LoadProfileConfig(true);
-	ProfilesPropertyHandle->NotifyPostChange();
+	CollisionProfile->SaveConfig();
 }
 
 void FCollisionProfileDetails::RefreshChannelList(bool bTraceType)
diff --git a/Engine/Source/Editor/DetailCustomizations/Private/CollisionProfileDetails.h b/Engine/Source/Editor/DetailCustomizations/Private/CollisionProfileDetails.h
--- a/Engine/Source/Editor/DetailCustomizations/Private/CollisionProfileDetails.h
+++ b/Engine/Source/Editor/DetailCustomizations/Private/CollisionProfileDetails.h
@@ -173,7 +173,5 @@ private:
 	};
 
 	FCollsiionProfileData SavedData;
-
-	TSharedPtr<IPropertyHandle> ProfilesPropertyHandle;
 };
 
diff --git a/Engine/Source/Runtime/Engine/Classes/Engine/CollisionProfile.h b/Engine/Source/Runtime/Engine/Classes/Engine/CollisionProfile.h
--- a/Engine/Source/Runtime/Engine/Classes/Engine/CollisionProfile.h
+++ b/Engine/Source/Runtime/Engine/Classes/Engine/CollisionProfile.h
@@ -131,7 +131,7 @@ class UCollisionProfile
 private:
 	// This is hacky, but without this edit tag, we can't get valid property handle
 	// and we can't save them properly to config, so we need this tag. 
-	UPROPERTY(globalconfig,Category=Collision, VisibleAnywhere)
+	UPROPERTY(globalconfig)
 	TArray<FCollisionResponseTemplate>	Profiles;
 
 	UPROPERTY(globalconfig)
