PropertyEditor: PropertyUtilities has new DontUpdateWhileEditing function. Used by standalone customization widgets.

diff --git a/Engine/Source/Editor/DetailCustomizations/Private/MathStructCustomizations.cpp b/Engine/Source/Editor/DetailCustomizations/Private/MathStructCustomizations.cpp
--- a/Engine/Source/Editor/DetailCustomizations/Private/MathStructCustomizations.cpp
+++ b/Engine/Source/Editor/DetailCustomizations/Private/MathStructCustomizations.cpp
@@ -3,6 +3,7 @@
 #include "DetailCustomizationsPrivatePCH.h"
 #include "Materials/MaterialExpressionConstant3Vector.h"
 #include "MathStructCustomizations.h"
+#include "IPropertyUtilities.h"
 
 #define LOCTEXT_NAMESPACE "FMathStructCustomization"
 
@@ -344,6 +345,9 @@ void FColorStructCustomization::CustomizeHeader( TSharedRef<class IPropertyHandl
 	bIsLinearColor = CastChecked<UStructProperty>( StructPropertyHandle->GetProperty() )->Struct->GetFName() == NAME_LinearColor;
 	bIgnoreAlpha = StructPropertyHandle->GetProperty()->HasMetaData(TEXT("HideAlphaChannel"));
 
+	auto PropertyUtils = StructCustomizationUtils.GetPropertyUtilities();
+	bDontUpdateWhileEditing = PropertyUtils.IsValid() ? PropertyUtils->DontUpdateValueWhileEditing() : false;
+
 	FMathStructCustomization::CustomizeHeader( InStructPropertyHandle, InHeaderRow, StructCustomizationUtils );
 }
 
@@ -478,7 +482,7 @@ void FColorStructCustomization::CreateColorPicker( bool bUseAlpha, bool bOnlyRef
 	GetColorAsLinear(InitialColor);
 
 	// This needs to be meta data.  Other colors could benefit from this
-	const bool bRefreshOnlyOnOk = StructPropertyHandle->GetProperty()->GetOwnerClass()->IsChildOf(UMaterialExpressionConstant3Vector::StaticClass());
+	const bool bRefreshOnlyOnOk = bOnlyRefreshOnOk || StructPropertyHandle->GetProperty()->GetOwnerClass()->IsChildOf(UMaterialExpressionConstant3Vector::StaticClass());
 
 	FColorPickerArgs PickerArgs;
 	PickerArgs.bUseAlpha = !bIgnoreAlpha;
@@ -523,7 +527,7 @@ TSharedRef<SColorPicker> FColorStructCustomization::CreateInlineColorPicker()
 	GetColorAsLinear(InitialColor);
 
 	// This needs to be meta data.  Other colors could benefit from this
-	const bool bRefreshOnlyOnOk = StructPropertyHandle->GetProperty()->GetOwnerClass()->IsChildOf(UMaterialExpressionConstant3Vector::StaticClass());
+	const bool bRefreshOnlyOnOk = bDontUpdateWhileEditing || StructPropertyHandle->GetProperty()->GetOwnerClass()->IsChildOf(UMaterialExpressionConstant3Vector::StaticClass());
 	
 	return SNew(SColorPicker)
 		.Visibility(this, &FColorStructCustomization::GetInlineColorPickerVisibility)
@@ -590,10 +594,13 @@ void FColorStructCustomization::OnColorPickerInteractiveEnd()
 {
 	bIsInteractive = false;
 
-	// pushes the last value from the interactive change without the interactive flag
-	FString ColorString;
-	StructPropertyHandle->GetValueAsFormattedString(ColorString);
-	StructPropertyHandle->SetValueFromFormattedString(ColorString);
+	if (!bDontUpdateWhileEditing)
+	{
+		// pushes the last value from the interactive change without the interactive flag
+		FString ColorString;
+		StructPropertyHandle->GetValueAsFormattedString(ColorString);
+		StructPropertyHandle->SetValueFromFormattedString(ColorString);
+	}
 
 	GEditor->EndTransaction();
 }
@@ -663,7 +670,7 @@ FReply FColorStructCustomization::OnMouseButtonDownColorBlock(const FGeometry& M
 		return FReply::Unhandled();
 	}
 	
-	CreateColorPicker( /*bUseAlpha*/ true, /*bRefreshOnlyOnOk*/ false );
+	CreateColorPicker( /*bUseAlpha*/ true, bDontUpdateWhileEditing);
 
 	//bIsInlineColorPickerVisible = !bIsInlineColorPickerVisible;
 
@@ -677,7 +684,7 @@ EVisibility FColorStructCustomization::GetInlineColorPickerVisibility() const
 
 FReply FColorStructCustomization::OnOpenFullColorPickerClicked()
 {
-	CreateColorPicker( /*bUseAlpha*/ true, /*bRefreshOnlyOnOk*/ false );
+	CreateColorPicker( /*bUseAlpha*/ true, bDontUpdateWhileEditing);
 
 	bIsInlineColorPickerVisible = false;
 
diff --git a/Engine/Source/Editor/DetailCustomizations/Private/MathStructCustomizations.h b/Engine/Source/Editor/DetailCustomizations/Private/MathStructCustomizations.h
--- a/Engine/Source/Editor/DetailCustomizations/Private/MathStructCustomizations.h
+++ b/Engine/Source/Editor/DetailCustomizations/Private/MathStructCustomizations.h
@@ -148,6 +148,7 @@ private:
 		: bIgnoreAlpha( false )
 		, bIsInlineColorPickerVisible(false)
 		, bIsInteractive(false)
+		, bDontUpdateWhileEditing(false)
 	{}
 
 	/** FMathStructCustomization interface */
@@ -235,5 +236,7 @@ private:
 	bool bIsInteractive;
 	/** Cached widget for the color picker to use as a parent */
 	TSharedPtr<SWidget> ColorPickerParentWidget;
+	/** The value won;t be updated while editing */
+	bool bDontUpdateWhileEditing;
 };
 
diff --git a/Engine/Source/Editor/PropertyEditor/Private/IDetailsViewPrivate.h b/Engine/Source/Editor/PropertyEditor/Private/IDetailsViewPrivate.h
--- a/Engine/Source/Editor/PropertyEditor/Private/IDetailsViewPrivate.h
+++ b/Engine/Source/Editor/PropertyEditor/Private/IDetailsViewPrivate.h
@@ -106,4 +106,7 @@ public:
 	 * Called when properties have finished changing (after PostEditChange is called)
 	 */
 	virtual void NotifyFinishedChangingProperties(const FPropertyChangedEvent& PropertyChangedEvent) = 0;
+
+	/** If a customization standalone widget is used, the value should be update only once, when its window is closed */
+	virtual bool DontUpdateValueWhileEditing() const = 0;
 };
diff --git a/Engine/Source/Editor/PropertyEditor/Private/Presentation/PropertyTable/PropertyTable.h b/Engine/Source/Editor/PropertyEditor/Private/Presentation/PropertyTable/PropertyTable.h
--- a/Engine/Source/Editor/PropertyEditor/Private/Presentation/PropertyTable/PropertyTable.h
+++ b/Engine/Source/Editor/PropertyEditor/Private/Presentation/PropertyTable/PropertyTable.h
@@ -127,6 +127,11 @@ public:
 	virtual FRootPathChanged* OnRootPathChanged() override { return &RootPathChanged; }
 
 	virtual bool IsPropertyEditingEnabled() const override { return true; }
+
+	virtual bool DontUpdateValueWhileEditing() const override
+	{
+		return false;
+	}
 private:
 
 	void UpdateColumns();
diff --git a/Engine/Source/Editor/PropertyEditor/Private/SDetailsView.h b/Engine/Source/Editor/PropertyEditor/Private/SDetailsView.h
--- a/Engine/Source/Editor/PropertyEditor/Private/SDetailsView.h
+++ b/Engine/Source/Editor/PropertyEditor/Private/SDetailsView.h
@@ -101,6 +101,11 @@ public:
 	{
 		return RootPropertyNode;
 	}
+
+	virtual bool DontUpdateValueWhileEditing() const override
+	{ 
+		return false; 
+	}
 private:
 	void RegisterInstancedCustomPropertyLayout( UClass* Class, FOnGetDetailCustomizationInstance DetailLayoutDelegate );
 	void UnregisterInstancedCustomPropertyLayout( UClass* Class );
diff --git a/Engine/Source/Editor/PropertyEditor/Private/SPropertyTreeViewImpl.cpp b/Engine/Source/Editor/PropertyEditor/Private/SPropertyTreeViewImpl.cpp
--- a/Engine/Source/Editor/PropertyEditor/Private/SPropertyTreeViewImpl.cpp
+++ b/Engine/Source/Editor/PropertyEditor/Private/SPropertyTreeViewImpl.cpp
@@ -69,6 +69,11 @@ public:
 	}
 
 	virtual void NotifyFinishedChangingProperties(const FPropertyChangedEvent& PropertyChangedEvent) override {}
+
+	virtual bool DontUpdateValueWhileEditing() const override
+	{
+		return false;
+	}
 private:
 
 	SPropertyTreeViewImpl& View;
diff --git a/Engine/Source/Editor/PropertyEditor/Private/SSingleProperty.cpp b/Engine/Source/Editor/PropertyEditor/Private/SSingleProperty.cpp
--- a/Engine/Source/Editor/PropertyEditor/Private/SSingleProperty.cpp
+++ b/Engine/Source/Editor/PropertyEditor/Private/SSingleProperty.cpp
@@ -68,6 +68,11 @@ public:
 	virtual void NotifyFinishedChangingProperties(const FPropertyChangedEvent& PropertyChangedEvent)
 	{}
 
+	virtual bool DontUpdateValueWhileEditing() const override
+	{
+		return false;
+	}
+
 private:
 	TWeakPtr< SSingleProperty > View;
 };
diff --git a/Engine/Source/Editor/PropertyEditor/Private/SStructureDetailsView.h b/Engine/Source/Editor/PropertyEditor/Private/SStructureDetailsView.h
--- a/Engine/Source/Editor/PropertyEditor/Private/SStructureDetailsView.h
+++ b/Engine/Source/Editor/PropertyEditor/Private/SStructureDetailsView.h
@@ -30,6 +30,11 @@ public:
 
 	virtual bool IsConnected() const override;
 
+	virtual bool DontUpdateValueWhileEditing() const override
+	{
+		return true;
+	}
+
 	/** IStructureDetailsView interface */
 	virtual TSharedPtr<SWidget> GetWidget() override
 	{
diff --git a/Engine/Source/Editor/PropertyEditor/Private/UserInterface/PropertyDetails/PropertyDetailsUtilities.cpp b/Engine/Source/Editor/PropertyEditor/Private/UserInterface/PropertyDetails/PropertyDetailsUtilities.cpp
--- a/Engine/Source/Editor/PropertyEditor/Private/UserInterface/PropertyDetails/PropertyDetailsUtilities.cpp
+++ b/Engine/Source/Editor/PropertyEditor/Private/UserInterface/PropertyDetails/PropertyDetailsUtilities.cpp
@@ -52,4 +52,9 @@ TSharedPtr<class FAssetThumbnailPool> FPropertyDetailsUtilities::GetThumbnailPoo
 void FPropertyDetailsUtilities::NotifyFinishedChangingProperties(const FPropertyChangedEvent& PropertyChangedEvent)
 {
 	DetailsView.NotifyFinishedChangingProperties(PropertyChangedEvent);
+}
+
+bool FPropertyDetailsUtilities::DontUpdateValueWhileEditing() const
+{
+	return DetailsView.DontUpdateValueWhileEditing();
 }
\ No newline at end of file
diff --git a/Engine/Source/Editor/PropertyEditor/Private/UserInterface/PropertyDetails/PropertyDetailsUtilities.h b/Engine/Source/Editor/PropertyEditor/Private/UserInterface/PropertyDetails/PropertyDetailsUtilities.h
--- a/Engine/Source/Editor/PropertyEditor/Private/UserInterface/PropertyDetails/PropertyDetailsUtilities.h
+++ b/Engine/Source/Editor/PropertyEditor/Private/UserInterface/PropertyDetails/PropertyDetailsUtilities.h
@@ -17,6 +17,7 @@ public:
 	virtual void RequestRefresh() override;
 	virtual TSharedPtr<class FAssetThumbnailPool> GetThumbnailPool() const override;
 	virtual void NotifyFinishedChangingProperties(const FPropertyChangedEvent& PropertyChangedEvent) override;
+	virtual bool DontUpdateValueWhileEditing() const override;
 private:
 	IDetailsViewPrivate& DetailsView;
 };
diff --git a/Engine/Source/Editor/PropertyEditor/Public/IPropertyUtilities.h b/Engine/Source/Editor/PropertyEditor/Public/IPropertyUtilities.h
--- a/Engine/Source/Editor/PropertyEditor/Public/IPropertyUtilities.h
+++ b/Engine/Source/Editor/PropertyEditor/Public/IPropertyUtilities.h
@@ -18,4 +18,6 @@ public:
 	virtual void RequestRefresh() = 0;
 	virtual TSharedPtr<class FAssetThumbnailPool> GetThumbnailPool() const = 0;
 	virtual void NotifyFinishedChangingProperties(const FPropertyChangedEvent& PropertyChangedEvent) = 0;
+	/** If a customization standalone widget is used, the value should be update only once, when its window is closed */
+	virtual bool DontUpdateValueWhileEditing() const = 0;
 };
