Fix the margin customization creating 4 transactions and calling PostEditChange 4 times when changing the value

diff --git a/Engine/Source/Editor/DetailCustomizations/Private/MarginCustomization.cpp b/Engine/Source/Editor/DetailCustomizations/Private/MarginCustomization.cpp
--- a/Engine/Source/Editor/DetailCustomizations/Private/MarginCustomization.cpp
+++ b/Engine/Source/Editor/DetailCustomizations/Private/MarginCustomization.cpp
@@ -2,6 +2,7 @@
 
 #include "DetailCustomizationsPrivatePCH.h"
 #include "MarginCustomization.h"
+#include "ScopedTransaction.h"
 
 TSharedRef<IPropertyTypeCustomization> FMarginStructCustomization::MakeInstance() 
 {
@@ -10,6 +11,8 @@ TSharedRef<IPropertyTypeCustomization> FMarginStructCustomization::MakeInstance(
 
 void FMarginStructCustomization::CustomizeHeader( TSharedRef<class IPropertyHandle> StructPropertyHandle, class FDetailWidgetRow& HeaderRow, IPropertyTypeCustomizationUtils& StructCustomizationUtils )
 {
+	this->StructPropertyHandle = StructPropertyHandle;
+
 	const FString UVSpaceString( StructPropertyHandle->GetProperty()->GetMetaData( TEXT( "UVSpace" ) ) );
 	bIsMarginUsingUVSpace = UVSpaceString.Len() > 0 && UVSpaceString == TEXT( "true" );
 
@@ -38,15 +41,10 @@ void FMarginStructCustomization::CustomizeHeader( TSharedRef<class IPropertyHand
 	[
 		MakePropertyWidget()
 	];
-
-	UpdateMarginTextFromProperties();
 }
 
 void FMarginStructCustomization::CustomizeChildren( TSharedRef<class IPropertyHandle> StructPropertyHandle, class IDetailChildrenBuilder& StructBuilder, IPropertyTypeCustomizationUtils& StructCustomizationUtils )
 {
-	FSimpleDelegate OnPropertyChangedDelegate = FSimpleDelegate::CreateSP( this, &FMarginStructCustomization::OnPropertyChanged );
-	StructPropertyHandle->SetOnPropertyValueChanged( OnPropertyChangedDelegate );
-
 	const uint32 NumChildren = ChildPropertyHandles.Num();
 
 	for( uint32 ChildIndex = 0; ChildIndex < NumChildren; ++ChildIndex )
@@ -63,8 +61,6 @@ void FMarginStructCustomization::CustomizeChildren( TSharedRef<class IPropertyHa
 		[
 			MakeChildPropertyWidget( ChildIndex, false )
 		];
-
-		ChildHandle->SetOnPropertyValueChanged( OnPropertyChangedDelegate );
 	}
 }
 
@@ -108,7 +104,7 @@ TSharedRef<SWidget> FMarginStructCustomization::MakeChildPropertyWidget( int32 P
 
 FText FMarginStructCustomization::GetMarginText() const
 {
-	return FText::FromString( MarginText );
+	return FText::FromString( GetMarginTextFromProperties() );
 }
 
 void FMarginStructCustomization::OnMarginTextCommitted( const FText& InText, ETextCommit::Type InCommitType )
@@ -153,46 +149,64 @@ void FMarginStructCustomization::OnMarginTextCommitted( const FText& InText, ETe
 			}
 		}
 
+		
 		if( !bError )
 		{
+
+			FMargin NewMargin;
 			// Update the property values
 			if( PropertyValues.Num() == 1 )
 			{
-				const float Value = PropertyValues[ 0 ];
-				PropertyValues.Add( Value );
-				PropertyValues.Add( Value );
-				PropertyValues.Add( Value );
+				// Uniform margin
+				NewMargin = FMargin( PropertyValues[0] );
 			}
 			else if( PropertyValues.Num() == 2 )
 			{
-				float Value = PropertyValues[ 0 ];
-				PropertyValues.Add( Value );
-				Value = PropertyValues[ 1 ];
-				PropertyValues.Add( Value );
+				// Uniform on the two axes
+				NewMargin = FMargin( PropertyValues[0], PropertyValues[1] );
 			}
-			else if( PropertyValues.Num() != 4 )
+			else if( PropertyValues.Num() == 4 )
+			{
+				NewMargin.Left = PropertyValues[0];
+				NewMargin.Top = PropertyValues[1];
+				NewMargin.Right = PropertyValues[2];
+				NewMargin.Bottom = PropertyValues[3];
+			}
+			else
 			{
 				bError = true;
 			}
 
+
 			if( !bError )
 			{
-				for( int32 PropertyIndex = 0; PropertyIndex < 4; ++PropertyIndex )
+				if (bIsMarginUsingUVSpace)
 				{
-					TWeakPtr<IPropertyHandle> WeakHandlePtr = ChildPropertyHandles[ PropertyIndex ];
-					WeakHandlePtr.Pin()->SetValue( PropertyValues[ PropertyIndex ] );
+					if (NewMargin.Left + NewMargin.Right > 1.0f)
+					{
+						NewMargin.Left = 1.0f - NewMargin.Right;
+					}
+
+					if (NewMargin.Top + NewMargin.Bottom > 1.0f)
+					{
+						NewMargin.Top = 1.0f - NewMargin.Bottom;
+					}
 				}
 
-				if( bIsMarginUsingUVSpace )
+				TArray<void*> RawData;
+				StructPropertyHandle->AccessRawData( RawData );
+
 				{
-					// Using UV space so constrain the paired margin values to be between 0.0 and 1.0
-					for( int32 PropertyIndex = 0; PropertyIndex < 4; ++PropertyIndex )
+					FScopedTransaction Transaction( FText::Format( NSLOCTEXT("FMarginStructCustomization", "SetMarginProperty", "Edit {0}"), FText::FromString( StructPropertyHandle->GetPropertyDisplayName() ) ) );
+					StructPropertyHandle->NotifyPreChange();
+
+					for (void* Data : RawData)
 					{
-						ConstrainPropertyValues( PropertyIndex );
+						*(FMargin*)Data = NewMargin;
 					}
-				}
 
-				UpdateMarginTextFromProperties();
+					StructPropertyHandle->NotifyPostChange();
+				}
 
 				MarginEditableTextBox->SetError( FString() );
 			}
@@ -205,8 +219,9 @@ void FMarginStructCustomization::OnMarginTextCommitted( const FText& InText, ETe
 	}
 }
 
-void FMarginStructCustomization::UpdateMarginTextFromProperties()
+FString FMarginStructCustomization::GetMarginTextFromProperties() const
 {
+	FString MarginText;
 	float PropertyValues[ 4 ];
 	bool bMultipleValues = false;
 
@@ -242,13 +257,11 @@ void FMarginStructCustomization::UpdateMarginTextFromProperties()
 				FString::SanitizeFloat( PropertyValues[ 2 ] ) + FString( ", " ) + FString::SanitizeFloat( PropertyValues[ 3 ] );
 		}
 	}
-}
 
-void FMarginStructCustomization::OnPropertyChanged()
-{
-	UpdateMarginTextFromProperties();
+	return MarginText;
 }
 
+
 TOptional<float> FMarginStructCustomization::OnGetValue( int32 PropertyIndex ) const
 {
 	TWeakPtr<IPropertyHandle> WeakHandlePtr = ChildPropertyHandles[ PropertyIndex ];
@@ -265,7 +278,7 @@ void FMarginStructCustomization::OnBeginSliderMovement()
 {
 	bIsUsingSlider = true;
 
-	GEditor->BeginTransaction( NSLOCTEXT("FMarginStructCustomization", "SetMarginProperty", "Set Margin Property") );
+	GEditor->BeginTransaction( FText::Format( NSLOCTEXT("FMarginStructCustomization", "SetMarginProperty", "Edit {0}"), FText::FromString( StructPropertyHandle->GetPropertyDisplayName() ) ) );
 }
 
 void FMarginStructCustomization::OnEndSliderMovement( float NewValue )
@@ -279,13 +292,6 @@ void FMarginStructCustomization::OnValueCommitted( float NewValue, ETextCommit::
 {
 	TWeakPtr<IPropertyHandle> WeakHandlePtr = ChildPropertyHandles[ PropertyIndex ];
 	WeakHandlePtr.Pin()->SetValue( NewValue );
-
-	if( bIsMarginUsingUVSpace )
-	{
-		ConstrainPropertyValues( PropertyIndex );
-	}
-
-	UpdateMarginTextFromProperties();
 }	
 
 void FMarginStructCustomization::OnValueChanged( float NewValue, int32 PropertyIndex )
@@ -295,38 +301,6 @@ void FMarginStructCustomization::OnValueChanged( float NewValue, int32 PropertyI
 		TWeakPtr<IPropertyHandle> WeakHandlePtr = ChildPropertyHandles[ PropertyIndex ];
 		EPropertyValueSetFlags::Type Flags = EPropertyValueSetFlags::InteractiveChange;
 		WeakHandlePtr.Pin()->SetValue( NewValue, Flags );
-
-		if( bIsMarginUsingUVSpace )
-		{
-			ConstrainPropertyValues( PropertyIndex );
-		}
-
-		UpdateMarginTextFromProperties();
 	}
 }
 
-void FMarginStructCustomization::ConstrainPropertyValues( int32 PropertyIndex )
-{
-	check( bIsMarginUsingUVSpace );
-	const int32 PropertyIndex2 = PropertyIndex ^ 2;
-
-	float PropertyValues[ 2 ];
-	ChildPropertyHandles[ PropertyIndex ]->GetValue( PropertyValues[ 0 ] );
-	TArray<FString> ObjectValues;
-	ChildPropertyHandles[ PropertyIndex2 ]->GetPerObjectValues( ObjectValues );
-
-	for( int32 OutputIndex = 0; OutputIndex < ObjectValues.Num(); ++OutputIndex )
-	{
-		TTypeFromString<float>::FromString( PropertyValues[ 1 ], *ObjectValues[ OutputIndex ] );
-
-		if( PropertyValues[ 0 ] + PropertyValues[ 1 ] > 1.0f )
-		{
-			PropertyValues[ 1 ] = 1.0f - PropertyValues[ 0 ];
-		}
-
-		ChildPropertyHandles[ PropertyIndex ]->SetValue( PropertyValues[ 0 ] );
-		ObjectValues[ OutputIndex ] = FString::SanitizeFloat( PropertyValues[ 1 ] );
-	}
-
-	ChildPropertyHandles[ PropertyIndex2 ]->SetPerObjectValues( ObjectValues );
-}
diff --git a/Engine/Source/Editor/DetailCustomizations/Private/MarginCustomization.h b/Engine/Source/Editor/DetailCustomizations/Private/MarginCustomization.h
--- a/Engine/Source/Editor/DetailCustomizations/Private/MarginCustomization.h
+++ b/Engine/Source/Editor/DetailCustomizations/Private/MarginCustomization.h
@@ -40,12 +40,7 @@ private:
 	/**
 	 * Update the margin text from the margin property values
 	 */
-	void UpdateMarginTextFromProperties();
-
-	/**
-	 * Called on change of property value
-	 */
-	void OnPropertyChanged();
+	FString GetMarginTextFromProperties() const;
 
 	/**
 	 * Gets the value as a float for the provided property handle
@@ -81,11 +76,9 @@ private:
 	 * Called when a value stops being changed by a slider
 	 */
 	void OnEndSliderMovement( float NewValue );
-
-	/**
-	 * Constrain the specified property value pair so that the sum is less than 1.0
-	 */
-	void ConstrainPropertyValues( int32 PropertyIndex );
+	
+	/** Handle to the margin being customized */
+	TSharedPtr<IPropertyHandle> StructPropertyHandle;
 
 	/** Child property handles */
 	TArray< TSharedRef<IPropertyHandle> > ChildPropertyHandles;
@@ -98,8 +91,5 @@ private:
 
 	/** Margin text editable text box */
 	TSharedPtr<SEditableTextBox> MarginEditableTextBox;
-
-	/** Margin text */
-	FString MarginText;
 };
 
