Fixed Mass in KG field to be disabled in Slate according to whether it's read only or not.
#ttp 338694 - CRASH: Highlighting static mesh Mass in KG field, then pressing enter
#branch UE4
#proj Editor.DetailCustomizations
#add Added a delegate binding for IsEnabled on the FDetailWidgetRow used for the Mass in KG property.

diff --git a/Engine/Source/Editor/DetailCustomizations/Private/BodySetupDetails.cpp b/Engine/Source/Editor/DetailCustomizations/Private/BodySetupDetails.cpp
--- a/Engine/Source/Editor/DetailCustomizations/Private/BodySetupDetails.cpp
+++ b/Engine/Source/Editor/DetailCustomizations/Private/BodySetupDetails.cpp
@@ -36,6 +36,7 @@ void FBodySetupDetails::CustomizeDetails( IDetailLayoutBuilder& DetailBuilder )
 			DetailBuilder.GetObjectsBeingCustomized(ObjectsCustomized);
 
 			PhysicsCategory.AddCustomRow(TEXT("Mass"), false)
+				.IsEnabled(TAttribute<bool>(this, &FBodySetupDetails::IsBodyMassReadOnly))
 				.NameContent()
 				[
 					SNew (STextBlock)
diff --git a/Engine/Source/Editor/DetailCustomizations/Private/BodySetupDetails.h b/Engine/Source/Editor/DetailCustomizations/Private/BodySetupDetails.h
--- a/Engine/Source/Editor/DetailCustomizations/Private/BodySetupDetails.h
+++ b/Engine/Source/Editor/DetailCustomizations/Private/BodySetupDetails.h
@@ -13,6 +13,7 @@ public:
 
 	FText OnGetBodyMass() const;
 	bool IsBodyMassReadOnly() const { return true; }
+	bool IsBodyMassEnabled() const { return !IsBodyMassReadOnly(); }
 
 private:
 	TArray< TWeakObjectPtr<UObject> > ObjectsCustomized;
diff --git a/Engine/Source/Editor/DetailCustomizations/Private/PrimitiveComponentDetails.cpp b/Engine/Source/Editor/DetailCustomizations/Private/PrimitiveComponentDetails.cpp
--- a/Engine/Source/Editor/DetailCustomizations/Private/PrimitiveComponentDetails.cpp
+++ b/Engine/Source/Editor/DetailCustomizations/Private/PrimitiveComponentDetails.cpp
@@ -156,6 +156,7 @@ void FPrimitiveComponentDetails::CustomizeDetails( IDetailLayoutBuilder& DetailB
 			if (bDisplayMass)
 			{
 				PhysicsCategory.AddCustomRow(TEXT("Mass"), false)
+					.IsEnabled(TAttribute<bool>(this, &FPrimitiveComponentDetails::IsBodyMassEnabled))
 					.NameContent()
 					[
 						SNew(STextBlock)
diff --git a/Engine/Source/Editor/DetailCustomizations/Private/PrimitiveComponentDetails.h b/Engine/Source/Editor/DetailCustomizations/Private/PrimitiveComponentDetails.h
--- a/Engine/Source/Editor/DetailCustomizations/Private/PrimitiveComponentDetails.h
+++ b/Engine/Source/Editor/DetailCustomizations/Private/PrimitiveComponentDetails.h
@@ -29,6 +29,7 @@ private:
 
 	FText OnGetBodyMass() const;
 	bool IsBodyMassReadOnly() const { return true; }
+	bool IsBodyMassEnabled() const { return !IsBodyMassReadOnly(); }
 
 private:
 	/** Objects being customized so we can update the 'Simulate Physics' state if physics geometry is added/removed */
