Editor: defaulting file path customization to 'All files' if no file extension filter was specified in the property's meta data

diff --git a/Engine/Source/Editor/DetailCustomizations/Private/FilePathStructCustomization.cpp b/Engine/Source/Editor/DetailCustomizations/Private/FilePathStructCustomization.cpp
--- a/Engine/Source/Editor/DetailCustomizations/Private/FilePathStructCustomization.cpp
+++ b/Engine/Source/Editor/DetailCustomizations/Private/FilePathStructCustomization.cpp
@@ -65,10 +65,20 @@ private:
 
 			TSharedPtr<SWindow> ParentWindow = FSlateApplication::Get().FindWidgetWindow(AsShared());
 			void* ParentWindowHandle = (ParentWindow.IsValid() && ParentWindow->GetNativeWindow().IsValid()) ? ParentWindow->GetNativeWindow()->GetOSWindowHandle() : nullptr;
+			FString Filter;
+
+			if (FileFilterExtension.IsEmpty())
+			{
+				Filter = TEXT("All files (*.*)|*.*");
+			}
+			else
+			{
+				Filter = FString::Printf(TEXT("%s files (*.%s)|*.%s"), *FileFilterExtension, *FileFilterExtension, *FileFilterExtension);
+			}
 
 			TArray<FString> OutFiles;
-			const FString filter = FString::Printf(TEXT("%s files (*.%s)|*.%s"), *FileFilterExtension, *FileFilterExtension, *FileFilterExtension);
-			if (DesktopPlatform->OpenFileDialog(ParentWindowHandle, LOCTEXT("PropertyEditorTitle", "File picker...").ToString(), DefaultPath, TEXT(""), filter, EFileDialogFlags::None, OutFiles))
+		
+			if (DesktopPlatform->OpenFileDialog(ParentWindowHandle, LOCTEXT("PropertyEditorTitle", "File picker...").ToString(), DefaultPath, TEXT(""), Filter, EFileDialogFlags::None, OutFiles))
 			{
 				FEditorDirectories::Get().SetLastDirectory(ELastDirectory::GENERIC_OPEN, FPaths::GetPath(OutFiles[0]));
 
