Fixing inconsistent expansion of input settings properties

TTP#335353 - EDITOR: Event and axis mapping group expansion behavior is unintuitive

Expansion state was based on the name of the action/axis mapping so when one was renamed it had no state for the new name. Added code to record the state of expansion before renaming and reset it once the mapping groups have been recreated. This required a new function adding to the IDetailGroup class so that I could record its expansion state and I also needed to be sure I only recorded the state when the commit callback actually changed something.

diff --git a/Engine/Source/Editor/DetailCustomizations/Private/InputSettingsDetails.cpp b/Engine/Source/Editor/DetailCustomizations/Private/InputSettingsDetails.cpp
--- a/Engine/Source/Editor/DetailCustomizations/Private/InputSettingsDetails.cpp
+++ b/Engine/Source/Editor/DetailCustomizations/Private/InputSettingsDetails.cpp
@@ -27,6 +27,7 @@ void FActionMappingsNodeBuilder::Tick( float DeltaTime )
 	if (GroupsRequireRebuild())
 	{
 		RebuildChildren();
+		HandleRenamedGroupExpansion();
 	}
 }
 
@@ -71,7 +72,7 @@ void FActionMappingsNodeBuilder::GenerateChildContent( IDetailChildrenBuilder& C
 
 	for (int32 Index = 0; Index < GroupedMappings.Num(); ++Index)
 	{
-		const FMappingSet& MappingSet = GroupedMappings[Index];
+		FMappingSet& MappingSet = GroupedMappings[Index];
 
 		TSharedRef<SWidget> AddButton = PropertyCustomizationHelpers::MakeAddButton( FSimpleDelegate::CreateSP( this, &FActionMappingsNodeBuilder::AddActionMappingToGroupButton_OnClick, MappingSet), 
 			LOCTEXT("AddActionMappingToGroupToolTip", "Adds Action Mapping to Group") );
@@ -83,6 +84,7 @@ void FActionMappingsNodeBuilder::GenerateChildContent( IDetailChildrenBuilder& C
 		MappingSet.SharedName.AppendString(GroupNameString);
 		FName GroupName(*GroupNameString);
 		IDetailGroup& ActionMappingGroup = ChildrenBuilder.AddChildGroup(GroupName, MappingSet.SharedName.ToString());
+		MappingSet.DetailGroup = &ActionMappingGroup;
 		ActionMappingGroup.HeaderRow()
 		[
 			SNew( SHorizontalBox )
@@ -175,12 +177,28 @@ void FActionMappingsNodeBuilder::OnActionMappingNameCommitted(const FText& InNam
 	const FScopedTransaction Transaction(LOCTEXT("RenameActionMapping_Transaction", "Rename Action Mapping"));
 
 	FName NewName = FName(*InName.ToString());
+	FName CurrentName = NewName;
 
-	TSharedPtr<IPropertyHandleArray> ActionMappingsArrayHandle = ActionMappingsPropertyHandle->AsArray();
+	if (MappingSet.Mappings.Num() > 0)
+	{
+		MappingSet.Mappings[0]->GetChildHandle(GET_MEMBER_NAME_CHECKED(FInputActionKeyMapping, ActionName))->GetValue(CurrentName);
+	}
 
-	for (int32 Index = 0; Index < MappingSet.Mappings.Num(); ++Index)
+	if (NewName != CurrentName)
 	{
-		MappingSet.Mappings[Index]->GetChildHandle(GET_MEMBER_NAME_CHECKED(FInputActionKeyMapping, ActionName))->SetValue(NewName);
+		for (int32 Index = 0; Index < MappingSet.Mappings.Num(); ++Index)
+		{
+			MappingSet.Mappings[Index]->GetChildHandle(GET_MEMBER_NAME_CHECKED(FInputActionKeyMapping, ActionName))->SetValue(NewName);
+		}
+
+		if (MappingSet.DetailGroup)
+		{
+			RenamedGroupExpansionState.Key = NewName;
+			RenamedGroupExpansionState.Value = MappingSet.DetailGroup->GetExpansionState();
+
+			// Don't want to save expansion state of old name
+			MappingSet.DetailGroup->ToggleExpansion(false);
+		}
 	}
 }
 
@@ -277,6 +295,22 @@ void FActionMappingsNodeBuilder::RebuildGroupedMappings()
 	}
 }
 
+void FActionMappingsNodeBuilder::HandleRenamedGroupExpansion()
+{
+	if (RenamedGroupExpansionState.Key != NAME_None)
+	{
+		for (auto& MappingSet : GroupedMappings)
+		{
+			if (MappingSet.SharedName == RenamedGroupExpansionState.Key)
+			{
+				MappingSet.DetailGroup->ToggleExpansion(RenamedGroupExpansionState.Value);
+				break;
+			}
+		}
+		RenamedGroupExpansionState.Key = NAME_None;
+	}
+}
+
 //////////////////////////////
 // FAxisMappingsNodeBuilder //
 //////////////////////////////
@@ -296,6 +330,7 @@ void FAxisMappingsNodeBuilder::Tick( float DeltaTime )
 	if (GroupsRequireRebuild())
 	{
 		RebuildChildren();
+		HandleRenamedGroupExpansion();
 	}
 }
 
@@ -340,7 +375,7 @@ void FAxisMappingsNodeBuilder::GenerateChildContent( IDetailChildrenBuilder& Chi
 
 	for (int32 Index = 0; Index < GroupedMappings.Num(); ++Index)
 	{
-		const FMappingSet& MappingSet = GroupedMappings[Index];
+		FMappingSet& MappingSet = GroupedMappings[Index];
 
 		TSharedRef<SWidget> AddButton = PropertyCustomizationHelpers::MakeAddButton( FSimpleDelegate::CreateSP( this, &FAxisMappingsNodeBuilder::AddAxisMappingToGroupButton_OnClick, MappingSet), 
 			LOCTEXT("AddAxisMappingToGroupToolTip", "Adds Axis Mapping to Group") );
@@ -352,6 +387,7 @@ void FAxisMappingsNodeBuilder::GenerateChildContent( IDetailChildrenBuilder& Chi
 		MappingSet.SharedName.AppendString(GroupNameString);
 		FName GroupName(*GroupNameString);
 		IDetailGroup& AxisMappingGroup = ChildrenBuilder.AddChildGroup(GroupName, MappingSet.SharedName.ToString());
+		MappingSet.DetailGroup = &AxisMappingGroup;
 		AxisMappingGroup.HeaderRow()
 		[
 			SNew( SHorizontalBox )
@@ -444,12 +480,28 @@ void FAxisMappingsNodeBuilder::OnAxisMappingNameCommitted(const FText& InName, E
 	const FScopedTransaction Transaction(LOCTEXT("RenameAxisMapping_Transaction", "Rename Axis Mapping"));
 
 	FName NewName = FName(*InName.ToString());
+	FName CurrentName = NewName;
 
-	TSharedPtr<IPropertyHandleArray> AxisMappingsArrayHandle = AxisMappingsPropertyHandle->AsArray();
+	if (MappingSet.Mappings.Num() > 0)
+	{
+		MappingSet.Mappings[0]->GetChildHandle(GET_MEMBER_NAME_CHECKED(FInputAxisKeyMapping, AxisName))->GetValue(CurrentName);
+	}
 
-	for (int32 Index = 0; Index < MappingSet.Mappings.Num(); ++Index)
+	if (NewName != CurrentName)
 	{
-		MappingSet.Mappings[Index]->GetChildHandle(GET_MEMBER_NAME_CHECKED(FInputAxisKeyMapping, AxisName))->SetValue(NewName);
+		for (int32 Index = 0; Index < MappingSet.Mappings.Num(); ++Index)
+		{
+			MappingSet.Mappings[Index]->GetChildHandle(GET_MEMBER_NAME_CHECKED(FInputAxisKeyMapping, AxisName))->SetValue(NewName);
+		}
+
+		if (MappingSet.DetailGroup)
+		{
+			RenamedGroupExpansionState.Key = NewName;
+			RenamedGroupExpansionState.Value = MappingSet.DetailGroup->GetExpansionState();
+
+			// Don't want to save expansion state of old name
+			MappingSet.DetailGroup->ToggleExpansion(false);
+		}
 	}
 }
 
@@ -546,6 +598,22 @@ void FAxisMappingsNodeBuilder::RebuildGroupedMappings()
 	}
 }
 
+void FAxisMappingsNodeBuilder::HandleRenamedGroupExpansion()
+{
+	if (RenamedGroupExpansionState.Key != NAME_None)
+	{
+		for (auto& MappingSet : GroupedMappings)
+		{
+			if (MappingSet.SharedName == RenamedGroupExpansionState.Key)
+			{
+				MappingSet.DetailGroup->ToggleExpansion(RenamedGroupExpansionState.Value);
+				break;
+			}
+		}
+		RenamedGroupExpansionState.Key = NAME_None;
+	}
+}
+
 /////////////////////////
 // FInputSettingsDetails //
 /////////////////////////
diff --git a/Engine/Source/Editor/DetailCustomizations/Private/InputSettingsDetails.h b/Engine/Source/Editor/DetailCustomizations/Private/InputSettingsDetails.h
--- a/Engine/Source/Editor/DetailCustomizations/Private/InputSettingsDetails.h
+++ b/Engine/Source/Editor/DetailCustomizations/Private/InputSettingsDetails.h
@@ -12,6 +12,7 @@ namespace InputConstants
 struct FMappingSet
 {
 	FName SharedName;
+	IDetailGroup* DetailGroup;
 	TArray<TSharedRef<IPropertyHandle>> Mappings;
 };
 
@@ -27,7 +28,7 @@ public:
 	virtual void GenerateHeaderRowContent( FDetailWidgetRow& NodeRow ) override;
 	virtual void GenerateChildContent( IDetailChildrenBuilder& ChildrenBuilder ) override;
 	virtual bool InitiallyCollapsed() const override { return true; };
-	virtual FName GetName() const override { return NAME_None; }
+	virtual FName GetName() const override { return FName(TEXT("ActionMappings")); }
 
 private:
 	void AddActionMappingButton_OnClick();
@@ -42,6 +43,8 @@ private:
 	{
 		OnRebuildChildren.ExecuteIfBound();
 	}
+	/** Makes sure that renamed groups have their expansion set to their previous state */
+	void HandleRenamedGroupExpansion();
 
 private:
 	/** Called to rebuild the children of the detail tree */
@@ -54,6 +57,8 @@ private:
 	TSharedPtr<IPropertyHandle> ActionMappingsPropertyHandle;
 
 	TArray<FMappingSet> GroupedMappings;
+
+	TPair<FName,bool> RenamedGroupExpansionState;
 };
 
 class FAxisMappingsNodeBuilder : public IDetailCustomNodeBuilder, public TSharedFromThis<FAxisMappingsNodeBuilder>
@@ -68,7 +73,7 @@ public:
 	virtual void GenerateHeaderRowContent( FDetailWidgetRow& NodeRow ) override;
 	virtual void GenerateChildContent( IDetailChildrenBuilder& ChildrenBuilder ) override;
 	virtual bool InitiallyCollapsed() const override { return true; };
-	virtual FName GetName() const override { return NAME_None; }
+	virtual FName GetName() const override { return FName(TEXT("AxisMappings")); }
 
 private:
 	void AddAxisMappingButton_OnClick();
@@ -83,6 +88,8 @@ private:
 	{
 		OnRebuildChildren.ExecuteIfBound();
 	}
+	/** Makes sure that renamed groups have their expansion set to their previous state */
+	void HandleRenamedGroupExpansion();
 
 private:
 	/** Called to rebuild the children of the detail tree */
@@ -95,6 +102,8 @@ private:
 	TSharedPtr<IPropertyHandle> AxisMappingsPropertyHandle;
 
 	TArray<FMappingSet> GroupedMappings;
+
+	TPair<FName,bool> RenamedGroupExpansionState;
 };
 
 class FInputSettingsDetails : public IDetailCustomization
diff --git a/Engine/Source/Editor/PropertyEditor/Private/DetailGroup.cpp b/Engine/Source/Editor/PropertyEditor/Private/DetailGroup.cpp
--- a/Engine/Source/Editor/PropertyEditor/Private/DetailGroup.cpp
+++ b/Engine/Source/Editor/PropertyEditor/Private/DetailGroup.cpp
@@ -66,6 +66,15 @@ void FDetailGroup::ToggleExpansion( bool bExpand )
 	}
 }
 
+bool FDetailGroup::GetExpansionState() const
+{
+	if (ParentCategory.IsValid() && OwnerTreeNode.IsValid())
+	{
+		return ParentCategory.Pin()->GetSavedExpansionState(*OwnerTreeNode.Pin().Get());
+	}
+	return false;
+}
+
 bool FDetailGroup::HasColumns() const
 {
 	if( HeaderCustomization.IsValid() && HeaderCustomization->HasPropertyNode() && HeaderCustomization->PropertyRow->HasColumns() )
diff --git a/Engine/Source/Editor/PropertyEditor/Private/DetailGroup.h b/Engine/Source/Editor/PropertyEditor/Private/DetailGroup.h
--- a/Engine/Source/Editor/PropertyEditor/Private/DetailGroup.h
+++ b/Engine/Source/Editor/PropertyEditor/Private/DetailGroup.h
@@ -13,6 +13,7 @@ public:
 	virtual FDetailWidgetRow& AddWidgetRow() override;
 	virtual IDetailPropertyRow& AddPropertyRow( TSharedRef<IPropertyHandle> PropertyHandle ) override;
 	virtual void ToggleExpansion( bool bExpand ) override;
+	virtual bool GetExpansionState() const override;
 
 	/** @return The name of the group */
 	FName GetGroupName() const { return GroupName; }
diff --git a/Engine/Source/Editor/PropertyEditor/Public/IDetailGroup.h b/Engine/Source/Editor/PropertyEditor/Public/IDetailGroup.h
--- a/Engine/Source/Editor/PropertyEditor/Public/IDetailGroup.h
+++ b/Engine/Source/Editor/PropertyEditor/Public/IDetailGroup.h
@@ -46,4 +46,9 @@ public:
 	 * @param bExpand	true to expand the group, false to collapse
 	 */
 	virtual void ToggleExpansion( bool bExpand ) = 0;
+
+	/**
+	 * Gets the current state of expansion for the group
+	 */
+	virtual bool GetExpansionState() const = 0;
 };
