Editor: Minor improvements to the actor 'early access'/'experimental' display

- Removed a bogus additional row
- Changed hit testing on surrounding box to make the tooltip easier to trigger

diff --git a/Engine/Source/Editor/DetailCustomizations/Private/ActorDetails.cpp b/Engine/Source/Editor/DetailCustomizations/Private/ActorDetails.cpp
--- a/Engine/Source/Editor/DetailCustomizations/Private/ActorDetails.cpp
+++ b/Engine/Source/Editor/DetailCustomizations/Private/ActorDetails.cpp
@@ -436,31 +436,32 @@ void FActorDetails::AddExperimentalWarningCategory( IDetailLayoutBuilder& Detail
 
 		IDetailCategoryBuilder& WarningCategory = DetailBuilder.EditCategory(CategoryName, CategoryDisplayName, ECategoryPriority::Transform);
 
-		FDetailWidgetRow& WarningRow = WarningCategory.AddCustomRow(SearchString);
-
-		WarningCategory.AddCustomRow(SearchString)
-		.WholeRowContent()
-		[
-			SNew(SHorizontalBox)
-			.ToolTip(IDocumentation::Get()->CreateToolTip(Tooltip, nullptr, TEXT("Shared/LevelEditor"), ExcerptName))
-			+ SHorizontalBox::Slot()
-			.VAlign(VAlign_Center)
-			.AutoWidth()
-			.Padding(4.0f, 0.0f, 0.0f, 0.0f)
+		FDetailWidgetRow& WarningRow = WarningCategory.AddCustomRow(SearchString)
+			.WholeRowContent()
 			[
-				SNew(SImage)
-				.Image(WarningIcon)
-			]
-			+SHorizontalBox::Slot()
-			.VAlign(VAlign_Center)
-			.AutoWidth()
-			.Padding(4.0f, 0.0f, 0.0f, 0.0f)
-			[
-				SNew(STextBlock)
-				.Text(WarningText)
-				.Font(IDetailLayoutBuilder::GetDetailFont())
-			]
-		];
+				SNew(SHorizontalBox)
+				.ToolTip(IDocumentation::Get()->CreateToolTip(Tooltip, nullptr, TEXT("Shared/LevelEditor"), ExcerptName))
+				.Visibility(EVisibility::Visible)
+
+				+ SHorizontalBox::Slot()
+				.VAlign(VAlign_Center)
+				.AutoWidth()
+				.Padding(4.0f, 0.0f, 0.0f, 0.0f)
+				[
+					SNew(SImage)
+					.Image(WarningIcon)
+				]
+
+				+SHorizontalBox::Slot()
+				.VAlign(VAlign_Center)
+				.AutoWidth()
+				.Padding(4.0f, 0.0f, 0.0f, 0.0f)
+				[
+					SNew(STextBlock)
+					.Text(WarningText)
+					.Font(IDetailLayoutBuilder::GetDetailFont())
+				]
+			];
 	}
 }
 
