Expose constraint angular reference frame offset

diff --git a/Engine/Source/Editor/ComponentVisualizers/Private/ConstraintComponentVisualizer.cpp b/Engine/Source/Editor/ComponentVisualizers/Private/ConstraintComponentVisualizer.cpp
--- a/Engine/Source/Editor/ComponentVisualizers/Private/ConstraintComponentVisualizer.cpp
+++ b/Engine/Source/Editor/ComponentVisualizers/Private/ConstraintComponentVisualizer.cpp
@@ -33,6 +33,7 @@ void FConstraintComponentVisualizer::DrawVisualization( const UActorComponent* C
 			Con1Frame = ConstraintComp->ComponentToWorld;
 			Con1Frame.RemoveScaling();
 			Con2Frame = ConstraintComp->ComponentToWorld;
+			Con2Frame.SetRotation(ConstraintComp->ConstraintInstance.AngularRotationOffset.Quaternion() * Con2Frame.GetRotation());
 			Con1Frame.RemoveScaling();
 		}
 
diff --git a/Engine/Source/Editor/DetailCustomizations/Private/PhysicsConstraintComponentDetails.cpp b/Engine/Source/Editor/DetailCustomizations/Private/PhysicsConstraintComponentDetails.cpp
--- a/Engine/Source/Editor/DetailCustomizations/Private/PhysicsConstraintComponentDetails.cpp
+++ b/Engine/Source/Editor/DetailCustomizations/Private/PhysicsConstraintComponentDetails.cpp
@@ -22,6 +22,8 @@ void FPhysicsConstraintComponentDetails::CustomizeDetails( IDetailLayoutBuilder&
 	UPhysicsConstraintComponent* ConstraintComp = NULL;
 	APhysicsConstraintActor* OwningConstraintActor = NULL;
 
+	bool bInPhAT = false;
+
 	for (int32 i=0; i < Objects.Num(); ++i)
 	{
 		if (!Objects[i].IsValid()) { continue; }
@@ -29,6 +31,7 @@ void FPhysicsConstraintComponentDetails::CustomizeDetails( IDetailLayoutBuilder&
 		if (Objects[i]->IsA(UPhysicsConstraintTemplate::StaticClass()))
 		{
 			ConstraintInstance = DetailBuilder.GetProperty("DefaultInstance");
+			bInPhAT = true;
 			break;
 		}
 		else if (Objects[i]->IsA(UPhysicsConstraintComponent::StaticClass()))
@@ -289,6 +292,13 @@ void FPhysicsConstraintComponentDetails::CustomizeDetails( IDetailLayoutBuilder&
 		AngularLimitCat.AddProperty(ConstraintInstance->GetChildHandle("TwistLimitDamping").ToSharedRef())
 			.Visibility(TAttribute<EVisibility>::Create(TAttribute<EVisibility>::FGetter::CreateSP(this, &FPhysicsConstraintComponentDetails::IsPropertyVisible, EPropertyType::AngularTwistLimit)));
 		
+		if (bInPhAT == false)
+		{
+			AngularLimitCat.AddProperty(ConstraintInstance->GetChildHandle("AngularRotationOffset").ToSharedRef())
+				.Visibility(TAttribute<EVisibility>::Create(TAttribute<EVisibility>::FGetter::CreateSP(this, &FPhysicsConstraintComponentDetails::IsPropertyVisible, EPropertyType::AngularAnyLimit)));
+
+		}
+		
 		AngularLimitCat.AddProperty(ConstraintInstance->GetChildHandle("bAngularBreakable").ToSharedRef());
 		AngularLimitCat.AddProperty(ConstraintInstance->GetChildHandle("AngularBreakThreshold").ToSharedRef());
 	}
diff --git a/Engine/Source/Runtime/Engine/Classes/PhysicsEngine/ConstraintInstance.h b/Engine/Source/Runtime/Engine/Classes/PhysicsEngine/ConstraintInstance.h
--- a/Engine/Source/Runtime/Engine/Classes/PhysicsEngine/ConstraintInstance.h
+++ b/Engine/Source/Runtime/Engine/Classes/PhysicsEngine/ConstraintInstance.h
@@ -240,6 +240,11 @@ struct ENGINE_API FConstraintInstance
 	/** Damping of the twist limit constraint if soft limit is used for twist motions. */
 	UPROPERTY(EditAnywhere, AdvancedDisplay, Category=Angular, meta=(editcondition = "bTwistLimitSoft", ClampMin = "0.0"))
 	float TwistLimitDamping;
+
+	/** Specifies the angular offset between the two frames of reference. By default limit goes from (-Angle, +Angle)
+	  * This allows you to bias the limit for swing1 swing2 and twist. */
+	UPROPERTY(EditAnywhere, Category = Angular)
+	FRotator AngularRotationOffset;
 	
 	/** Whether it is possible to break the joint with angular force. */
 	UPROPERTY(EditAnywhere, AdvancedDisplay, Category=Angular)
diff --git a/Engine/Source/Runtime/Engine/Private/PhysicsEngine/PhysicsConstraintComponent.cpp b/Engine/Source/Runtime/Engine/Private/PhysicsEngine/PhysicsConstraintComponent.cpp
--- a/Engine/Source/Runtime/Engine/Private/PhysicsEngine/PhysicsConstraintComponent.cpp
+++ b/Engine/Source/Runtime/Engine/Private/PhysicsEngine/PhysicsConstraintComponent.cpp
@@ -401,8 +401,8 @@ void UPhysicsConstraintComponent::UpdateConstraintFrames()
 	ConstraintInstance.SecAxis1 = A1Transform.InverseTransformVectorNoScale(WOrth);
 
 	ConstraintInstance.Pos2 = A2Transform.InverseTransformPosition(WPos);
-	ConstraintInstance.PriAxis2 = A2Transform.InverseTransformVectorNoScale(WPri);
-	ConstraintInstance.SecAxis2 = A2Transform.InverseTransformVectorNoScale(WOrth);
+	ConstraintInstance.PriAxis2 = A2Transform.InverseTransformVectorNoScale(ConstraintInstance.AngularRotationOffset.RotateVector(WPri));
+	ConstraintInstance.SecAxis2 = A2Transform.InverseTransformVectorNoScale(ConstraintInstance.AngularRotationOffset.RotateVector(WOrth));
 }
 
 #if WITH_EDITOR
