UMG - No longer spinning the FMargin except when modifying UV ranges for brushes.

diff --git a/Engine/Source/Editor/DetailCustomizations/Private/MarginCustomization.cpp b/Engine/Source/Editor/DetailCustomizations/Private/MarginCustomization.cpp
--- a/Engine/Source/Editor/DetailCustomizations/Private/MarginCustomization.cpp
+++ b/Engine/Source/Editor/DetailCustomizations/Private/MarginCustomization.cpp
@@ -88,11 +88,11 @@ TSharedRef<SWidget> FMarginStructCustomization::MakeChildPropertyWidget( int32 P
 		.OnBeginSliderMovement( this, &FMarginStructCustomization::OnBeginSliderMovement )
 		.OnEndSliderMovement( this, &FMarginStructCustomization::OnEndSliderMovement )
 		.LabelVAlign( VAlign_Center )
-		.AllowSpin( true )
-		.MinValue( 0.0f )
-		.MaxValue( bIsMarginUsingUVSpace ? 1.0f : 100.0f )
-		.MinSliderValue( 0.0f )
-		.MaxSliderValue( bIsMarginUsingUVSpace ? 1.0f : 100.0f )
+		.AllowSpin( bIsMarginUsingUVSpace ? true : false )
+		.MinValue( bIsMarginUsingUVSpace ? 0.0f : TNumericLimits<float>::Lowest() )
+		.MaxValue( bIsMarginUsingUVSpace ? 1.0f : TNumericLimits<float>::Max() )
+		.MinSliderValue( bIsMarginUsingUVSpace ? 0.0f : TNumericLimits<float>::Lowest() )
+		.MaxSliderValue( bIsMarginUsingUVSpace ? 1.0f : TNumericLimits<float>::Max()  )
 		.Label()
 		[
 			SNew( STextBlock )
