Forgot these files for Locked Axis customization

diff --git a/Engine/Source/Editor/DetailCustomizations/Private/PrimitiveComponentDetails.cpp b/Engine/Source/Editor/DetailCustomizations/Private/PrimitiveComponentDetails.cpp
--- a/Engine/Source/Editor/DetailCustomizations/Private/PrimitiveComponentDetails.cpp
+++ b/Engine/Source/Editor/DetailCustomizations/Private/PrimitiveComponentDetails.cpp
@@ -5,10 +5,12 @@
 #include "ScopedTransaction.h"
 #include "ObjectEditorUtils.h"
 #include "PhysicsEngine/BodySetup.h"
+#include "PhysicsEngine//BodyInstance.h"
 #include "Components/DestructibleComponent.h"
 #include "Components/InstancedStaticMeshComponent.h"
 #include "IDocumentation.h"
 
+
 #define LOCTEXT_NAMESPACE "PrimitiveComponentDetails"
 
 //////////////////////////////////////////////////////////////
@@ -121,6 +123,21 @@ bool FPrimitiveComponentDetails::IsUseAsyncEditable() const
 	return bEnableUseAsyncScene;
 }
 
+EVisibility FPrimitiveComponentDetails::IsCustomLockedAxisSelected() const
+{
+	bool bVisible = false;
+	if (LockedAxisProperty.IsValid())
+	{
+		uint8 LockedAxis;
+		if (LockedAxisProperty->GetValue(LockedAxis) == FPropertyAccess::Success)
+		{
+			bVisible = static_cast<ELockedAxis::Type>(LockedAxis) == ELockedAxis::Custom;
+		}
+	}
+
+	return bVisible ? EVisibility::Visible : EVisibility::Collapsed;
+}
+
 void FPrimitiveComponentDetails::CustomizeDetails( IDetailLayoutBuilder& DetailBuilder )
 {
 	TSharedRef<IPropertyHandle> MobilityHandle = DetailBuilder.GetProperty("Mobility", USceneComponent::StaticClass());
@@ -190,6 +207,16 @@ void FPrimitiveComponentDetails::CustomizeDetails( IDetailLayoutBuilder& DetailB
 						//we only enable bUseAsyncScene if the project uses an AsyncScene
 						PhysicsCategory.AddProperty(ChildProperty).EditCondition(TAttribute<bool>(this, &FPrimitiveComponentDetails::IsUseAsyncEditable), NULL);
 					}
+					else if (ChildProperty->GetProperty()->GetName() == TEXT("LockedAxisMode"))
+					{
+						LockedAxisProperty = ChildProperty;
+						PhysicsCategory.AddProperty(ChildProperty);
+					}
+					else if (ChildProperty->GetProperty()->GetName() == TEXT("CustomLockedAxis"))
+					{
+						//we only enable bUseAsyncScene if the project uses an AsyncScene
+						PhysicsCategory.AddProperty(ChildProperty).Visibility(TAttribute<EVisibility>(this, &FPrimitiveComponentDetails::IsCustomLockedAxisSelected));
+					}
 					else
 					{
 						PhysicsCategory.AddProperty(ChildProperty);
diff --git a/Engine/Source/Editor/DetailCustomizations/Private/PrimitiveComponentDetails.h b/Engine/Source/Editor/DetailCustomizations/Private/PrimitiveComponentDetails.h
--- a/Engine/Source/Editor/DetailCustomizations/Private/PrimitiveComponentDetails.h
+++ b/Engine/Source/Editor/DetailCustomizations/Private/PrimitiveComponentDetails.h
@@ -34,5 +34,7 @@ private:
 private:
 	/** Objects being customized so we can update the 'Simulate Physics' state if physics geometry is added/removed */
 	TArray< TWeakObjectPtr<UObject> > ObjectsCustomized;
+	TSharedPtr<IPropertyHandle> LockedAxisProperty;
+	EVisibility IsCustomLockedAxisSelected() const;
 };
 
