Modified struct customization for FMatrix and FTransform to work in the same way as the component transform details customization.
#ttp 339291 - Level Editor Details panel transform copy/paste not compatible with Blueprint Details panel transform
#branch UE4
#proj Editor.DetailCustomizations

diff --git a/Engine/Source/Editor/DetailCustomizations/Private/ComponentTransformDetails.cpp b/Engine/Source/Editor/DetailCustomizations/Private/ComponentTransformDetails.cpp
--- a/Engine/Source/Editor/DetailCustomizations/Private/ComponentTransformDetails.cpp
+++ b/Engine/Source/Editor/DetailCustomizations/Private/ComponentTransformDetails.cpp
@@ -302,13 +302,13 @@ void FComponentTransformDetails::OnCopy( ETransformField::Type TransformField )
 	switch (TransformField)
 	{
 	case ETransformField::Location:
-		CopyStr = FString::Printf(TEXT("X=%f Y=%f Z=%f"), CachedLocation.X.GetValue(), CachedLocation.Y.GetValue(), CachedLocation.Z.GetValue());
+		CopyStr = FString::Printf(TEXT("(X=%f,Y=%f,Z=%f)"), CachedLocation.X.GetValue(), CachedLocation.Y.GetValue(), CachedLocation.Z.GetValue());
 		break;
 	case ETransformField::Rotation:
-		CopyStr = FString::Printf(TEXT("P=%f Y=%f R=%f"), CachedRotation.Y.GetValue(), CachedRotation.Z.GetValue(), CachedRotation.X.GetValue());
+		CopyStr = FString::Printf(TEXT("(Pitch=%f,Yaw=%f,Roll=%f)"), CachedRotation.Y.GetValue(), CachedRotation.Z.GetValue(), CachedRotation.X.GetValue());
 		break;
 	case ETransformField::Scale:
-		CopyStr = FString::Printf(TEXT("X=%f Y=%f Z=%f"), CachedScale.X.GetValue(), CachedScale.Y.GetValue(), CachedScale.Z.GetValue());
+		CopyStr = FString::Printf(TEXT("(X=%f,Y=%f,Z=%f)"), CachedScale.X.GetValue(), CachedScale.Y.GetValue(), CachedScale.Z.GetValue());
 		break;
 	default:
 		break;
@@ -342,6 +342,9 @@ void FComponentTransformDetails::OnPaste( ETransformField::Type TransformField )
 	case ETransformField::Rotation:
 		{
 			FRotator Rotation;
+			PastedText.ReplaceInline(TEXT("Pitch="), TEXT("P="));
+			PastedText.ReplaceInline(TEXT("Yaw="), TEXT("Y="));
+			PastedText.ReplaceInline(TEXT("Roll="), TEXT("R="));
 			if (Rotation.InitFromString(PastedText))
 			{
 				FScopedTransaction Transaction(LOCTEXT("PasteRotation", "Paste Rotation"));
diff --git a/Engine/Source/Editor/DetailCustomizations/Private/MathStructProxyCustomizations.cpp b/Engine/Source/Editor/DetailCustomizations/Private/MathStructProxyCustomizations.cpp
--- a/Engine/Source/Editor/DetailCustomizations/Private/MathStructProxyCustomizations.cpp
+++ b/Engine/Source/Editor/DetailCustomizations/Private/MathStructProxyCustomizations.cpp
@@ -2,6 +2,7 @@
 
 #include "DetailCustomizationsPrivatePCH.h"
 #include "MathStructProxyCustomizations.h"
+#include "ScopedTransaction.h"
 
 void FMathStructProxyCustomization::CustomizeChildren( TSharedRef<class IPropertyHandle> StructPropertyHandle, class IDetailChildrenBuilder& StructBuilder, IPropertyTypeCustomizationUtils& StructCustomizationUtils )
 {
@@ -85,118 +86,239 @@ void FMathStructProxyCustomization::OnEndSliderMovement( NumericType NewValue )
 }
 
 
-TSharedRef<IPropertyTypeCustomization> FMatrixStructCustomization::MakeInstance() 
+#define LOCTEXT_NAMESPACE "MatrixStructCustomization"
+
+TSharedRef<IPropertyTypeCustomization> FMatrixStructCustomization::MakeInstance()
 {
 	return MakeShareable( new FMatrixStructCustomization );
 }
 
-void FMatrixStructCustomization::MakeHeaderRow( TSharedRef<class IPropertyHandle>& StructPropertyHandle, FDetailWidgetRow& Row )
+void FMatrixStructCustomization::MakeHeaderRow(TSharedRef<class IPropertyHandle>& StructPropertyHandle, FDetailWidgetRow& Row)
 {
-#define LOCTEXT_NAMESPACE "MatrixStructCustomization"
-
-	// We'll set up reset to default ourselves
-	const bool bDisplayResetToDefault = false;
-
-	Row.NameContent()
+	Row
+	.NameContent()
 	[
-		SNew(SVerticalBox)
-		+SVerticalBox::Slot()
-		.Padding( FMargin(0.0f, 2.0f, 0.0f, 2.0f ) )
-		[
-			StructPropertyHandle->CreatePropertyNameWidget( LOCTEXT("LocationLabel", "Location").ToString(), bDisplayResetToDefault )
-		]
-		+SVerticalBox::Slot()
-		.Padding( FMargin(0.0f, 2.0f, 0.0f, 2.0f ) )
-		[
-			StructPropertyHandle->CreatePropertyNameWidget( LOCTEXT("RotationLabel", "Rotation").ToString(), bDisplayResetToDefault )
-		]
-		+SVerticalBox::Slot()
-		.Padding( FMargin(0.0f, 2.0f, 0.0f, 2.0f ) )
-		[
-			StructPropertyHandle->CreatePropertyNameWidget( LOCTEXT("ScaleLabel", "Scale").ToString(), bDisplayResetToDefault )
-		]
+		StructPropertyHandle->CreatePropertyNameWidget()
 	]
 	.ValueContent()
-	// Make enough space for each child handle
-	.MinDesiredWidth(125.0f * SortedChildHandles.Num() )
-	.MaxDesiredWidth(125.0f * SortedChildHandles.Num() )
+	.MinDesiredWidth(0.0f)
+	.MaxDesiredWidth(0.0f)
 	[
-		SNew(SVerticalBox)
-		+SVerticalBox::Slot()
-		.VAlign(VAlign_Center)
+		SNullWidget::NullWidget
+	];
+}
+
+void FMatrixStructCustomization::CustomizeChildren(TSharedRef<class IPropertyHandle> StructPropertyHandle, class IDetailChildrenBuilder& StructBuilder, IPropertyTypeCustomizationUtils& StructCustomizationUtils)
+{
+	TWeakPtr<IPropertyHandle> WeakHandlePtr = StructPropertyHandle;
+
+	StructBuilder.AddChildContent(LOCTEXT("LocationLabel", "Location").ToString())
+		.CopyAction(FUIAction(FExecuteAction::CreateSP(this, &FMatrixStructCustomization::OnCopy, FTransformField::Location, WeakHandlePtr)))
+		.PasteAction(FUIAction(FExecuteAction::CreateSP(this, &FMatrixStructCustomization::OnPaste, FTransformField::Location, WeakHandlePtr)))
+		.NameContent()
+		[
+			StructPropertyHandle->CreatePropertyNameWidget(LOCTEXT("LocationLabel", "Location").ToString())
+		]
+		.ValueContent()
+		.MinDesiredWidth(375.0f)
+		.MaxDesiredWidth(375.0f)
 		[
-			SNew( SHorizontalBox )
-			+SHorizontalBox::Slot()
-			.Padding( FMargin(0.0f, 2.0f, 3.0f, 2.0f ) )
+			SNew(SHorizontalBox)
+			+ SHorizontalBox::Slot()
+			.Padding(FMargin(0.0f, 2.0f, 3.0f, 2.0f))
 			[
 				MakeNumericProxyWidget<FVector, float>(StructPropertyHandle, CachedTranslationX, LOCTEXT("TranslationX", "X"), FLinearColor::White, SNumericEntryBox<float>::RedLabelBackgroundColor)
 			]
-			+SHorizontalBox::Slot()
-			.Padding( FMargin(0.0f, 2.0f, 3.0f, 2.0f ) )
+			+ SHorizontalBox::Slot()
+			.Padding(FMargin(0.0f, 2.0f, 3.0f, 2.0f))
 			[
 				MakeNumericProxyWidget<FVector, float>(StructPropertyHandle, CachedTranslationY, LOCTEXT("TranslationY", "Y"), FLinearColor::White, SNumericEntryBox<float>::GreenLabelBackgroundColor)
 			]
-			+SHorizontalBox::Slot()
-			.Padding( FMargin(0.0f, 2.0f, 0.0f, 2.0f ) )
+			+ SHorizontalBox::Slot()
+			.Padding(FMargin(0.0f, 2.0f, 0.0f, 2.0f))
 			[
 				MakeNumericProxyWidget<FVector, float>(StructPropertyHandle, CachedTranslationZ, LOCTEXT("TranslationZ", "Z"), FLinearColor::White, SNumericEntryBox<float>::BlueLabelBackgroundColor)
 			]
+		];
+
+	StructBuilder.AddChildContent(LOCTEXT("RotationLabel", "Rotation").ToString())
+		.CopyAction(FUIAction(FExecuteAction::CreateSP(this, &FMatrixStructCustomization::OnCopy, FTransformField::Rotation, WeakHandlePtr)))
+		.PasteAction(FUIAction(FExecuteAction::CreateSP(this, &FMatrixStructCustomization::OnPaste, FTransformField::Rotation, WeakHandlePtr)))
+		.NameContent()
+		[
+			StructPropertyHandle->CreatePropertyNameWidget(LOCTEXT("RotationLabel", "Rotation").ToString())
 		]
-		+SVerticalBox::Slot()
-		.VAlign(VAlign_Center)
+		.ValueContent()
+		.MinDesiredWidth(375.0f)
+		.MaxDesiredWidth(375.0f)
 		[
-			SNew( SHorizontalBox )
-			+SHorizontalBox::Slot()
-			.Padding( FMargin(0.0f, 2.0f, 3.0f, 2.0f ) )
+			SNew(SHorizontalBox)
+			+ SHorizontalBox::Slot()
+			.Padding(FMargin(0.0f, 2.0f, 3.0f, 2.0f))
 			[
 				MakeNumericProxyWidget<FRotator, float>(StructPropertyHandle, CachedRotationYaw, LOCTEXT("RotationYaw", "X"), FLinearColor::White, SNumericEntryBox<float>::RedLabelBackgroundColor)
 			]
-			+SHorizontalBox::Slot()
-			.Padding( FMargin(0.0f, 2.0f, 3.0f, 2.0f ) )
+			+ SHorizontalBox::Slot()
+			.Padding(FMargin(0.0f, 2.0f, 3.0f, 2.0f))
 			[
 				MakeNumericProxyWidget<FRotator, float>(StructPropertyHandle, CachedRotationPitch, LOCTEXT("RotationPitch", "Y"), FLinearColor::White, SNumericEntryBox<float>::GreenLabelBackgroundColor)
 			]
-			+SHorizontalBox::Slot()
-			.Padding( FMargin(0.0f, 2.0f, 0.0f, 2.0f ) )
+			+ SHorizontalBox::Slot()
+			.Padding(FMargin(0.0f, 2.0f, 0.0f, 2.0f))
 			[
 				MakeNumericProxyWidget<FRotator, float>(StructPropertyHandle, CachedRotationRoll, LOCTEXT("RotationRoll", "Z"), FLinearColor::White, SNumericEntryBox<float>::BlueLabelBackgroundColor)
 			]
+		];
+
+	StructBuilder.AddChildContent(LOCTEXT("ScaleLabel", "Scale").ToString())
+		.CopyAction(FUIAction(FExecuteAction::CreateSP(this, &FMatrixStructCustomization::OnCopy, FTransformField::Scale, WeakHandlePtr)))
+		.PasteAction(FUIAction(FExecuteAction::CreateSP(this, &FMatrixStructCustomization::OnPaste, FTransformField::Scale, WeakHandlePtr)))
+		.NameContent()
+		[
+			StructPropertyHandle->CreatePropertyNameWidget(LOCTEXT("ScaleLabel", "Scale").ToString())
 		]
-		+SVerticalBox::Slot()
-		.VAlign(VAlign_Center)
+		.ValueContent()
+		.MinDesiredWidth(375.0f)
+		.MaxDesiredWidth(375.0f)
 		[
-			SNew( SHorizontalBox )
-			+SHorizontalBox::Slot()
-			.Padding( FMargin(0.0f, 2.0f, 3.0f, 2.0f ) )
+			SNew(SHorizontalBox)
+			+ SHorizontalBox::Slot()
+			.Padding(FMargin(0.0f, 2.0f, 3.0f, 2.0f))
 			[
 				MakeNumericProxyWidget<FVector, float>(StructPropertyHandle, CachedScaleX, LOCTEXT("ScaleX", "X"), FLinearColor::White, SNumericEntryBox<float>::RedLabelBackgroundColor)
 			]
-			+SHorizontalBox::Slot()
-			.Padding( FMargin(0.0f, 2.0f, 3.0f, 2.0f ) )
+			+ SHorizontalBox::Slot()
+			.Padding(FMargin(0.0f, 2.0f, 3.0f, 2.0f))
 			[
 				MakeNumericProxyWidget<FVector, float>(StructPropertyHandle, CachedScaleY, LOCTEXT("ScaleY", "Y"), FLinearColor::White, SNumericEntryBox<float>::GreenLabelBackgroundColor)
 			]
-			+SHorizontalBox::Slot()
-			.Padding( FMargin(0.0f, 2.0f, 0.0f, 2.0f ) )
+			+ SHorizontalBox::Slot()
+			.Padding(FMargin(0.0f, 2.0f, 0.0f, 2.0f))
 			[
 				MakeNumericProxyWidget<FVector, float>(StructPropertyHandle, CachedScaleZ, LOCTEXT("ScaleZ", "Z"), FLinearColor::White, SNumericEntryBox<float>::BlueLabelBackgroundColor)
 			]
-		]
-	];
-#undef LOCTEXT_NAMESPACE
+		];
 }
 
-bool FMatrixStructCustomization::CacheValues( TWeakPtr<IPropertyHandle> WeakHandlePtr ) const
+void FMatrixStructCustomization::OnCopy(FTransformField::Type Type, TWeakPtr<IPropertyHandle> PropertyHandlePtr)
 {
-	check(WeakHandlePtr.IsValid());
+	auto PropertyHandle = PropertyHandlePtr.Pin();
+
+	if (!PropertyHandle.IsValid())
+	{
+		return;
+	}
+
+	FString CopyStr;
+	CacheValues(PropertyHandle);
+
+	switch (Type)
+	{
+		case FTransformField::Location:
+		{
+			FVector Location = CachedTranslation->Get();
+			CopyStr = FString::Printf(TEXT("(X=%f,Y=%f,Z=%f)"), Location.X, Location.Y, Location.Z);
+			break;
+		}
+
+		case FTransformField::Rotation:
+		{
+			FRotator Rotation = CachedRotation->Get();
+			CopyStr = FString::Printf(TEXT("(Pitch=%f,Yaw=%f,Roll=%f)"), Rotation.Pitch, Rotation.Yaw, Rotation.Roll);
+			break;
+		}
+
+		case FTransformField::Scale:
+		{
+			FVector Scale = CachedScale->Get();
+			CopyStr = FString::Printf(TEXT("(X=%f,Y=%f,Z=%f)"), Scale.X, Scale.Y, Scale.Z);
+			break;
+		}
+	}
+
+	if (!CopyStr.IsEmpty())
+	{
+		FPlatformMisc::ClipboardCopy(*CopyStr);
+	}
+}
+
+void FMatrixStructCustomization::OnPaste(FTransformField::Type Type, TWeakPtr<IPropertyHandle> PropertyHandlePtr)
+{
+	auto PropertyHandle = PropertyHandlePtr.Pin();
+
+	if (!PropertyHandle.IsValid())
+	{
+		return;
+	}
+
+	FString PastedText;
+	FPlatformMisc::ClipboardPaste(PastedText);
+
+	switch (Type)
+	{
+		case FTransformField::Location:
+		{
+			FVector Location;
+			if (Location.InitFromString(PastedText))
+			{
+				FScopedTransaction Transaction(LOCTEXT("PasteLocation", "Paste Location"));
+				CachedTranslationX->Set(Location.X);
+				CachedTranslationY->Set(Location.Y);
+				CachedTranslationZ->Set(Location.Z);
+				FlushValues(PropertyHandle);
+			}
+			break;
+		}
+
+		case FTransformField::Rotation:
+		{
+			FRotator Rotation;
+			PastedText.ReplaceInline(TEXT("Pitch="), TEXT("P="));
+			PastedText.ReplaceInline(TEXT("Yaw="), TEXT("Y="));
+			PastedText.ReplaceInline(TEXT("Roll="), TEXT("R="));
+			if (Rotation.InitFromString(PastedText))
+			{
+				FScopedTransaction Transaction(LOCTEXT("PasteRotation", "Paste Rotation"));
+				CachedRotationPitch->Set(Rotation.Pitch);
+				CachedRotationYaw->Set(Rotation.Yaw);
+				CachedRotationRoll->Set(Rotation.Roll);
+				FlushValues(PropertyHandle);
+			}
+			break;
+		}
+
+		case FTransformField::Scale:
+		{
+			FVector Scale;
+			if (Scale.InitFromString(PastedText))
+			{
+				FScopedTransaction Transaction(LOCTEXT("PasteScale", "Paste Scale"));
+				CachedScaleX->Set(Scale.X);
+				CachedScaleY->Set(Scale.Y);
+				CachedScaleZ->Set(Scale.Z);
+				FlushValues(PropertyHandle);
+			}
+			break;
+		}
+	}
+}
+
+bool FMatrixStructCustomization::CacheValues( TWeakPtr<IPropertyHandle> PropertyHandlePtr ) const
+{
+	auto PropertyHandle = PropertyHandlePtr.Pin();
+
+	if (!PropertyHandle.IsValid())
+	{
+		return false;
+	}
 
 	TArray<void*> RawData;
-	WeakHandlePtr.Pin()->AccessRawData(RawData);
+	PropertyHandle->AccessRawData(RawData);
 
 	if (RawData.Num() == 1)
 	{
 		FMatrix* MatrixValue = reinterpret_cast<FMatrix*>(RawData[0]);
-		if(MatrixValue != NULL)
+		if (MatrixValue != NULL)
 		{
 			CachedTranslation->Set(MatrixValue->GetOrigin());
 			CachedRotation->Set(MatrixValue->Rotator());
@@ -208,18 +330,22 @@ bool FMatrixStructCustomization::CacheValues( TWeakPtr<IPropertyHandle> WeakHand
 	return false;
 }
 
-bool FMatrixStructCustomization::FlushValues( TWeakPtr<IPropertyHandle> WeakHandlePtr ) const
+bool FMatrixStructCustomization::FlushValues( TWeakPtr<IPropertyHandle> PropertyHandlePtr ) const
 {
-	check(WeakHandlePtr.IsValid());
+	auto PropertyHandle = PropertyHandlePtr.Pin();
+	if (!PropertyHandle.IsValid())
+	{
+		return false;
+	}
 
 	TArray<void*> RawData;
-	WeakHandlePtr.Pin()->AccessRawData(RawData);
+	PropertyHandle->AccessRawData(RawData);
 
-	WeakHandlePtr.Pin()->NotifyPreChange();
-	for(int32 ValueIndex = 0; ValueIndex < RawData.Num(); ValueIndex++)
+	PropertyHandle->NotifyPreChange();
+	for (int32 ValueIndex = 0; ValueIndex < RawData.Num(); ValueIndex++)
 	{
 		FMatrix* MatrixValue = reinterpret_cast<FMatrix*>(RawData[ValueIndex]);
-		if(MatrixValue != NULL)
+		if (MatrixValue != NULL)
 		{
 			const FRotator CurrentRotation = MatrixValue->Rotator();
 			const FVector CurrentTranslation = MatrixValue->GetOrigin();
@@ -243,27 +369,34 @@ bool FMatrixStructCustomization::FlushValues( TWeakPtr<IPropertyHandle> WeakHand
 			*MatrixValue = FScaleRotationTranslationMatrix(Scale, Rotation, Translation);
 		}
 	}
-	WeakHandlePtr.Pin()->NotifyPostChange();
+	PropertyHandle->NotifyPostChange();
 
 	return true;
 }
 
+#undef LOCTEXT_NAMESPACE
+
 TSharedRef<IPropertyTypeCustomization> FTransformStructCustomization::MakeInstance() 
 {
 	return MakeShareable( new FTransformStructCustomization );
 }
 
-bool FTransformStructCustomization::CacheValues( TWeakPtr<IPropertyHandle> WeakHandlePtr ) const
+bool FTransformStructCustomization::CacheValues( TWeakPtr<IPropertyHandle> PropertyHandlePtr ) const
 {
-	check(WeakHandlePtr.IsValid());
+	auto PropertyHandle = PropertyHandlePtr.Pin();
+
+	if (!PropertyHandle.IsValid())
+	{
+		return false;
+	}
 
 	TArray<void*> RawData;
-	WeakHandlePtr.Pin()->AccessRawData(RawData);
+	PropertyHandle->AccessRawData(RawData);
 
 	if (RawData.Num() == 1)
 	{
 		FTransform* TransformValue = reinterpret_cast<FTransform*>(RawData[0]);
-		if(TransformValue != NULL)
+		if (TransformValue != NULL)
 		{
 			CachedTranslation->Set(TransformValue->GetTranslation());
 			CachedRotation->Set(TransformValue->GetRotation().Rotator());
@@ -275,18 +408,23 @@ bool FTransformStructCustomization::CacheValues( TWeakPtr<IPropertyHandle> WeakH
 	return false;
 }
 
-bool FTransformStructCustomization::FlushValues( TWeakPtr<IPropertyHandle> WeakHandlePtr ) const
+bool FTransformStructCustomization::FlushValues( TWeakPtr<IPropertyHandle> PropertyHandlePtr ) const
 {
-	check(WeakHandlePtr.IsValid());
+	auto PropertyHandle = PropertyHandlePtr.Pin();
+
+	if (!PropertyHandle.IsValid())
+	{
+		return false;
+	}
 
 	TArray<void*> RawData;
-	WeakHandlePtr.Pin()->AccessRawData(RawData);
+	PropertyHandle->AccessRawData(RawData);
 
-	WeakHandlePtr.Pin()->NotifyPreChange();
-	for(int32 ValueIndex = 0; ValueIndex < RawData.Num(); ValueIndex++)
+	PropertyHandle->NotifyPreChange();
+	for (int32 ValueIndex = 0; ValueIndex < RawData.Num(); ValueIndex++)
 	{
 		FTransform* TransformValue = reinterpret_cast<FTransform*>(RawData[0]);
-		if(TransformValue != NULL)
+		if (TransformValue != NULL)
 		{
 			const FRotator CurrentRotation = TransformValue->GetRotation().Rotator();
 			const FVector CurrentTranslation = TransformValue->GetTranslation();
@@ -310,7 +448,7 @@ bool FTransformStructCustomization::FlushValues( TWeakPtr<IPropertyHandle> WeakH
 			TransformValue->SetComponents(Rotation.Quaternion(), Translation, Scale);
 		}
 	}
-	WeakHandlePtr.Pin()->NotifyPostChange();
+	PropertyHandle->NotifyPostChange();
 
 	return true;
 }
\ No newline at end of file
diff --git a/Engine/Source/Editor/DetailCustomizations/Private/MathStructProxyCustomizations.h b/Engine/Source/Editor/DetailCustomizations/Private/MathStructProxyCustomizations.h
--- a/Engine/Source/Editor/DetailCustomizations/Private/MathStructProxyCustomizations.h
+++ b/Engine/Source/Editor/DetailCustomizations/Private/MathStructProxyCustomizations.h
@@ -232,13 +232,29 @@ public:
 	{
 	}
 
-	/** FMathStructCustomization interface */
-	virtual void MakeHeaderRow( TSharedRef<class IPropertyHandle>& StructPropertyHandle, FDetailWidgetRow& Row ) override;
+	/** IPropertyTypeCustomization interface */
+	virtual void CustomizeChildren(TSharedRef<class IPropertyHandle> StructPropertyHandle, class IDetailChildrenBuilder& StructBuilder, IPropertyTypeCustomizationUtils& StructCustomizationUtils) override;
 
 protected:
+	/** FMathStructCustomization interface */
+	virtual void MakeHeaderRow(TSharedRef<class IPropertyHandle>& InStructPropertyHandle, FDetailWidgetRow& Row) override;
+
 	/** FMathStructProxyCustomization interface */
-	virtual bool CacheValues( TWeakPtr<IPropertyHandle> WeakHandlePtr ) const override;
-	virtual bool FlushValues( TWeakPtr<IPropertyHandle> WeakHandlePtr ) const override;
+	virtual bool CacheValues(TWeakPtr<IPropertyHandle> PropertyHandlePtr) const override;
+	virtual bool FlushValues(TWeakPtr<IPropertyHandle> PropertyHandlePtr) const override;
+
+	struct FTransformField
+	{
+		enum Type
+		{
+			Location,
+			Rotation,
+			Scale
+		};
+	};
+
+	void OnCopy(FTransformField::Type Type, TWeakPtr<IPropertyHandle> PropertyHandlePtr);
+	void OnPaste(FTransformField::Type Type, TWeakPtr<IPropertyHandle> PropertyHandlePtr);
 
 protected:
 	/** Cached rotation values */
@@ -270,6 +286,6 @@ public:
 
 protected:
 	/** FMathStructProxyCustomization interface */
-	virtual bool CacheValues( TWeakPtr<IPropertyHandle> WeakHandlePtr ) const override;
-	virtual bool FlushValues( TWeakPtr<IPropertyHandle> WeakHandlePtr ) const override;
+	virtual bool CacheValues(TWeakPtr<IPropertyHandle> PropertyHandlePtr) const override;
+	virtual bool FlushValues(TWeakPtr<IPropertyHandle> PropertyHandlePtr) const override;
 };
\ No newline at end of file
