Cleanup phat interface so that 123 is xyz

diff --git a/Engine/Source/Editor/DetailCustomizations/Private/PhysicsConstraintComponentDetails.cpp b/Engine/Source/Editor/DetailCustomizations/Private/PhysicsConstraintComponentDetails.cpp
--- a/Engine/Source/Editor/DetailCustomizations/Private/PhysicsConstraintComponentDetails.cpp
+++ b/Engine/Source/Editor/DetailCustomizations/Private/PhysicsConstraintComponentDetails.cpp
@@ -212,7 +212,7 @@ void FPhysicsConstraintComponentDetails::CustomizeDetails( IDetailLayoutBuilder&
 
 
 		uint8 AngularLimitEnum[LCM_MAX] = { ACM_Free, LCM_Limited, LCM_Locked };
-		TSharedPtr<IPropertyHandle> AngularLimitProperties[] = { AngularSwing1MotionProperty, AngularSwing2MotionProperty, AngularTwistMotionProperty };
+		TSharedPtr<IPropertyHandle> AngularLimitProperties[] = { AngularSwing1MotionProperty, AngularTwistMotionProperty, AngularSwing2MotionProperty};
 
 		for (int32 PropertyIdx=0; PropertyIdx < 3; ++PropertyIdx)
 		{
@@ -273,8 +273,12 @@ void FPhysicsConstraintComponentDetails::CustomizeDetails( IDetailLayoutBuilder&
 
 		AngularLimitCat.AddProperty(ConstraintInstance->GetChildHandle("Swing1LimitAngle").ToSharedRef())
 			.Visibility(TAttribute<EVisibility>::Create(TAttribute<EVisibility>::FGetter::CreateSP(this, &FPhysicsConstraintComponentDetails::IsPropertyVisible, EPropertyType::AngularSwing1Limit)));
+		AngularLimitCat.AddProperty(ConstraintInstance->GetChildHandle("TwistLimitAngle").ToSharedRef())
+			.Visibility(TAttribute<EVisibility>::Create(TAttribute<EVisibility>::FGetter::CreateSP(this, &FPhysicsConstraintComponentDetails::IsPropertyVisible, EPropertyType::AngularTwistLimit)));
 		AngularLimitCat.AddProperty(ConstraintInstance->GetChildHandle("Swing2LimitAngle").ToSharedRef())
 			.Visibility(TAttribute<EVisibility>::Create(TAttribute<EVisibility>::FGetter::CreateSP(this, &FPhysicsConstraintComponentDetails::IsPropertyVisible, EPropertyType::AngularSwing2Limit)));
+
+
 		AngularLimitCat.AddProperty(ConstraintInstance->GetChildHandle("bSwingLimitSoft").ToSharedRef())
 			.Visibility(TAttribute<EVisibility>::Create(TAttribute<EVisibility>::FGetter::CreateSP(this, &FPhysicsConstraintComponentDetails::IsPropertyVisible, EPropertyType::AngularSwingLimit)));
 		AngularLimitCat.AddProperty(ConstraintInstance->GetChildHandle("SwingLimitStiffness").ToSharedRef())
@@ -283,8 +287,7 @@ void FPhysicsConstraintComponentDetails::CustomizeDetails( IDetailLayoutBuilder&
 			.Visibility(TAttribute<EVisibility>::Create(TAttribute<EVisibility>::FGetter::CreateSP(this, &FPhysicsConstraintComponentDetails::IsPropertyVisible, EPropertyType::AngularSwingLimit)));
 
 		
-		AngularLimitCat.AddProperty(ConstraintInstance->GetChildHandle("TwistLimitAngle").ToSharedRef())
-			.Visibility(TAttribute<EVisibility>::Create(TAttribute<EVisibility>::FGetter::CreateSP(this, &FPhysicsConstraintComponentDetails::IsPropertyVisible, EPropertyType::AngularTwistLimit)));
+
 		AngularLimitCat.AddProperty(ConstraintInstance->GetChildHandle("bTwistLimitSoft").ToSharedRef())
 			.Visibility(TAttribute<EVisibility>::Create(TAttribute<EVisibility>::FGetter::CreateSP(this, &FPhysicsConstraintComponentDetails::IsPropertyVisible, EPropertyType::AngularTwistLimit)));
 		AngularLimitCat.AddProperty(ConstraintInstance->GetChildHandle("TwistLimitStiffness").ToSharedRef())
@@ -298,7 +301,12 @@ void FPhysicsConstraintComponentDetails::CustomizeDetails( IDetailLayoutBuilder&
 				.Visibility(TAttribute<EVisibility>::Create(TAttribute<EVisibility>::FGetter::CreateSP(this, &FPhysicsConstraintComponentDetails::IsPropertyVisible, EPropertyType::AngularAnyLimit)));
 
 		}
-		
+		else
+		{
+			AngularLimitCat.AddProperty(ConstraintInstance->GetChildHandle("AngularRotationOffset").ToSharedRef())
+				.Visibility(EVisibility::Collapsed);
+		}
+
 		AngularLimitCat.AddProperty(ConstraintInstance->GetChildHandle("bAngularBreakable").ToSharedRef());
 		AngularLimitCat.AddProperty(ConstraintInstance->GetChildHandle("AngularBreakThreshold").ToSharedRef());
 	}
diff --git a/Engine/Source/Editor/PhAT/Private/PhATActions.cpp b/Engine/Source/Editor/PhAT/Private/PhATActions.cpp
--- a/Engine/Source/Editor/PhAT/Private/PhATActions.cpp
+++ b/Engine/Source/Editor/PhAT/Private/PhATActions.cpp
@@ -81,8 +81,8 @@ void FPhATCommands::RegisterCommands()
 	UI_COMMAND(CycleConstraintOrientation, "Cycle selected constraint orientation", "", EUserInterfaceActionType::Button, FInputGesture(EKeys::Q));
 	UI_COMMAND(CycleConstraintActive, "Cycle active constraint", "", EUserInterfaceActionType::Button, FInputGesture(EKeys::Four));
 	UI_COMMAND(ToggleSwing1, "Toggle Swing1 Constraint", "", EUserInterfaceActionType::Button, FInputGesture(EKeys::One));
-	UI_COMMAND(ToggleSwing2, "Toggle Swing2 Constraint", "", EUserInterfaceActionType::Button, FInputGesture(EKeys::Two));
-	UI_COMMAND(ToggleTwist, "Toggle Twist Constraint", "", EUserInterfaceActionType::Button, FInputGesture(EKeys::Three));
+	UI_COMMAND(ToggleSwing2, "Toggle Swing2 Constraint", "", EUserInterfaceActionType::Button, FInputGesture(EKeys::Three));
+	UI_COMMAND(ToggleTwist, "Toggle Twist Constraint", "", EUserInterfaceActionType::Button, FInputGesture(EKeys::Two));
 	UI_COMMAND(FocusOnSelection, "Focus the viewport on the current selection", "", EUserInterfaceActionType::Button, FInputGesture(EKeys::F));
 	UI_COMMAND(CycleTransformMode, "Cycle the transform mode", "", EUserInterfaceActionType::Button, FInputGesture(EKeys::SpaceBar));
 }
diff --git a/Engine/Source/Runtime/Engine/Classes/PhysicsEngine/ConstraintInstance.h b/Engine/Source/Runtime/Engine/Classes/PhysicsEngine/ConstraintInstance.h
--- a/Engine/Source/Runtime/Engine/Classes/PhysicsEngine/ConstraintInstance.h
+++ b/Engine/Source/Runtime/Engine/Classes/PhysicsEngine/ConstraintInstance.h
@@ -185,20 +185,21 @@ struct ENGINE_API FConstraintInstance
 	UPROPERTY()
 	uint32 bTwistLimited_DEPRECATED:1;
 
-	/** Indicates whether angular swing motion along the y axis is allowed, blocked or limited. If limited, the 
-		AngularLimit property will be used to determine if a motion is allowed. See EAngularConstraintMotion. */
+	/** Indicates whether rotation about the Z axis is allowed, blocked, or limited. If limited, the 
+		AngularLimit property will be used to determine the range of motion. See EAngularConstraintMotion. */
 	UPROPERTY(EditAnywhere, Category=Angular)
 	TEnumAsByte<enum EAngularConstraintMotion> AngularSwing1Motion;
-
-	/** Indicates whether angular swing motion along the z axis is allowed, blocked or limited. If limited, the 
-		AngularLimit property will be used to determine if a motion is allowed. See EAngularConstraintMotion. */
+		
+	/** Indicates whether rotation about the the X axis is allowed, blocked, or limited. If limited, the
+		AngularLimit property will be used to determine the range of motion. See EAngularConstraintMotion. */
 	UPROPERTY(EditAnywhere, Category=Angular)
+	TEnumAsByte<enum EAngularConstraintMotion> AngularTwistMotion;
+
+	/** Indicates whether rotation about the Y axis is allowed, blocked, or limited. If limited, the
+		AngularLimit property will be used to determine the range of motion. See EAngularConstraintMotion. */
+	UPROPERTY(EditAnywhere, Category = Angular)
 	TEnumAsByte<enum EAngularConstraintMotion> AngularSwing2Motion;
 
-	/** Indicates whether angular twist motion along the x axis is allowed, blocked or limited. If limited, the 
-		AngularLimit property will be used to determine if a motion is allowed. See EAngularConstraintMotion. */
-	UPROPERTY(EditAnywhere, Category=Angular)
-	TEnumAsByte<enum EAngularConstraintMotion> AngularTwistMotion;
 
 	/** Whether we want to use soft limits for swing motions instead of hard limits. With enabled 
 		soft limit, a constraint is used instead of hard-capping the motion. */
@@ -214,17 +215,17 @@ struct ENGINE_API FConstraintInstance
 		between 0 and 180. */
 	UPROPERTY(EditAnywhere, Category=Angular, meta=(ClampMin = "0.0", ClampMax = "180.0"))
 	float Swing1LimitAngle;
-
-	/** Used if swing motion along the z axis is limited. The limit angle is specified in degrees and should be
-		between 0 and 180. */
-	UPROPERTY(EditAnywhere, Category=Angular, meta=(ClampMin = "0.0", ClampMax = "180.0"))
-	float Swing2LimitAngle;
-
+	
 	/** Used if twist motion along the x axis is limited. The limit angle is specified in degrees and should be
 		between 0 and 180. */
 	UPROPERTY(EditAnywhere, Category=Angular, meta=(ClampMin = "0.0", ClampMax = "180.0"))
 	float TwistLimitAngle;
 
+	/** Used if swing motion along the z axis is limited. The limit angle is specified in degrees and should be
+	between 0 and 180. */
+	UPROPERTY(EditAnywhere, Category = Angular, meta = (ClampMin = "0.0", ClampMax = "180.0"))
+	float Swing2LimitAngle;
+
 	/** Stiffness of the swing limit constraint if soft limit is used for swing motions. */
 	UPROPERTY(EditAnywhere, AdvancedDisplay, Category=Angular, meta=(editcondition = "bSwingLimitSoft", ClampMin = "0.0"))
 	float SwingLimitStiffness;
