Fix non-Unity build.

diff --git a/Engine/Source/Editor/DetailCustomizations/Private/BrushDetails.cpp b/Engine/Source/Editor/DetailCustomizations/Private/BrushDetails.cpp
--- a/Engine/Source/Editor/DetailCustomizations/Private/BrushDetails.cpp
+++ b/Engine/Source/Editor/DetailCustomizations/Private/BrushDetails.cpp
@@ -9,6 +9,7 @@
 #include "ClassViewerFilter.h"
 #include "AssetSelection.h"
 #include "ClassIconFinder.h"
+#include "ScopedTransaction.h"
 
 #define LOCTEXT_NAMESPACE "BrushDetails"
 
