Physics: Improve bSimulatePhysics editing extensibility
- Add UPrimitiveComponent::CanEditSimulatePhysics to allow individual components to determine if simulation is possible given their Physics setup
- Replace if-else cast tree in details panel customization
- Change behavior of base UPrimitiveComponent so that it is editable if there is a valid body setup *on the component* rather than on the BodyInstance, allowing various components to edit the property in a Blueprint CDO (such as UPaperSpriteComponent) and removing the need for custom overrides on others (UShapeComponent, etc...)

diff --git a/Engine/Source/Editor/DetailCustomizations/Private/PrimitiveComponentDetails.cpp b/Engine/Source/Editor/DetailCustomizations/Private/PrimitiveComponentDetails.cpp
--- a/Engine/Source/Editor/DetailCustomizations/Private/PrimitiveComponentDetails.cpp
+++ b/Engine/Source/Editor/DetailCustomizations/Private/PrimitiveComponentDetails.cpp
@@ -27,61 +27,18 @@ bool FPrimitiveComponentDetails::IsSimulatePhysicsEditable() const
 {
 	// Check whether to enable editing of bSimulatePhysics - this will happen if all objects are UPrimitiveComponents & have collision geometry.
 	bool bEnableSimulatePhysics = ObjectsCustomized.Num() > 0;
-	for ( TWeakObjectPtr<UObject> CustomizedObject : ObjectsCustomized )
+	for (TWeakObjectPtr<UObject> CustomizedObject : ObjectsCustomized)
 	{
-		if ( CustomizedObject.IsValid() && CustomizedObject->IsA(UPrimitiveComponent::StaticClass()) )
+		if (UPrimitiveComponent* PrimitiveComponent = Cast<UPrimitiveComponent>(CustomizedObject.Get()))
 		{
-			// Primitive components are the simplest - they have the physics data within them
-			TWeakObjectPtr<UPrimitiveComponent> PrimitiveComponent = CastChecked<UPrimitiveComponent>(CustomizedObject.Get());
-			// Static mesh components can query the underlying mesh too
-			TWeakObjectPtr<UStaticMeshComponent> StaticMeshComponent = Cast<UStaticMeshComponent>(CustomizedObject.Get());
-			// if skeletalmeshcomponent, you'll have to check physics asset to verify if it's eligible to have collision
-			TWeakObjectPtr<USkeletalMeshComponent> SkeletalMeshComponent = Cast<USkeletalMeshComponent>(PrimitiveComponent.Get());
-			// if destructiblemeshcomponent, we will allow it always
-			TWeakObjectPtr<UDestructibleComponent> DestructibleComponent = Cast<UDestructibleComponent>(PrimitiveComponent.Get());
-			// if shape component, we will check another body setup
-			TWeakObjectPtr<UShapeComponent> ShapeComponent = Cast<UShapeComponent>(CustomizedObject.Get());
-			// if instancedstaticmeshcomponent, we will never allow it
-			TWeakObjectPtr<UInstancedStaticMeshComponent> InstancedSMComponent = Cast<UInstancedStaticMeshComponent>(PrimitiveComponent.Get());
-			
-			if(DestructibleComponent.IsValid())
-			{
-				bEnableSimulatePhysics = true;
-			}
-			else if(ShapeComponent.IsValid())
-			{
-				// ShapeComponent : ShapeBodySetup is null when not registered, so we just allow it all the time. 
-				bEnableSimulatePhysics &= true;
-			}
-			else if(InstancedSMComponent.IsValid())
-			{
-				bEnableSimulatePhysics = false;
-			}
-			else if(PrimitiveComponent->BodyInstance.BodySetup.IsValid())
-			{
-												//If there's no collision we still let them simulate physics. The object falls through the world - this behavior is debatable but what we decided on for now
-				bEnableSimulatePhysics &= true; // (PrimitiveComponent->BodyInstance.BodySetup->AggGeom.GetElementCount() > 0 || PrimitiveComponent->BodyInstance.BodySetup->CollisionTraceFlag == CTF_UseComplexAsSimple);
-			}
-			else if(StaticMeshComponent.IsValid() && StaticMeshComponent->StaticMesh && StaticMeshComponent->StaticMesh->BodySetup)
-			{
-				bEnableSimulatePhysics &= (StaticMeshComponent->StaticMesh->BodySetup->AggGeom.GetElementCount() > 0 || StaticMeshComponent->StaticMesh->BodySetup->CollisionTraceFlag == CTF_UseComplexAsSimple);
-			}
-			else if(SkeletalMeshComponent.IsValid())
-			{
-				bEnableSimulatePhysics &= (SkeletalMeshComponent->GetPhysicsAsset() != NULL);
-			}
-			else
+			if (!PrimitiveComponent->CanEditSimulatePhysics())
 			{
 				bEnableSimulatePhysics = false;
 				break;
 			}
 		}
-		else
-		{
-			bEnableSimulatePhysics = false;
-			break;
-		}
 	}
+
 	return bEnableSimulatePhysics;
 }
 
diff --git a/Engine/Source/Runtime/Engine/Classes/Components/DestructibleComponent.h b/Engine/Source/Runtime/Engine/Classes/Components/DestructibleComponent.h
--- a/Engine/Source/Runtime/Engine/Classes/Components/DestructibleComponent.h
+++ b/Engine/Source/Runtime/Engine/Classes/Components/DestructibleComponent.h
@@ -123,6 +123,7 @@ public:
 
 	// Begin UPrimitiveComponent interface.
 	virtual FBodyInstance* GetBodyInstance(FName BoneName = NAME_None) const override;
+	virtual bool CanEditSimulatePhysics() override;
 	virtual bool IsAnySimulatingPhysics() const override;
 
 	virtual void AddImpulse(FVector Impulse, FName BoneName = NAME_None, bool bVelChange = false) override;
diff --git a/Engine/Source/Runtime/Engine/Classes/Components/InstancedStaticMeshComponent.h b/Engine/Source/Runtime/Engine/Classes/Components/InstancedStaticMeshComponent.h
--- a/Engine/Source/Runtime/Engine/Classes/Components/InstancedStaticMeshComponent.h
+++ b/Engine/Source/Runtime/Engine/Classes/Components/InstancedStaticMeshComponent.h
@@ -118,6 +118,7 @@ public:
 	virtual FPrimitiveSceneProxy* CreateSceneProxy() override;
 	virtual void CreatePhysicsState() override;
 	virtual void DestroyPhysicsState() override;
+	virtual bool CanEditSimulatePhysics() override;
 
 	virtual FBoxSphereBounds CalcBounds(const FTransform& BoundTransform) const override;
 #if WITH_EDITOR
diff --git a/Engine/Source/Runtime/Engine/Classes/Components/PrimitiveComponent.h b/Engine/Source/Runtime/Engine/Classes/Components/PrimitiveComponent.h
--- a/Engine/Source/Runtime/Engine/Classes/Components/PrimitiveComponent.h
+++ b/Engine/Source/Runtime/Engine/Classes/Components/PrimitiveComponent.h
@@ -655,6 +655,11 @@ public:
 	UFUNCTION(BlueprintCallable, Category="Physics")
 	virtual void SetSimulatePhysics(bool bSimulate);
 
+	/**
+	 * Determines whether or not the simulate physics setting can be edited interactively on this component
+	 */
+	virtual bool CanEditSimulatePhysics();
+
 	/**
 	 *	Add an impulse to a single rigid body. Good for one time instant burst.
 	 *
diff --git a/Engine/Source/Runtime/Engine/Classes/Components/SkeletalMeshComponent.h b/Engine/Source/Runtime/Engine/Classes/Components/SkeletalMeshComponent.h
--- a/Engine/Source/Runtime/Engine/Classes/Components/SkeletalMeshComponent.h
+++ b/Engine/Source/Runtime/Engine/Classes/Components/SkeletalMeshComponent.h
@@ -747,6 +747,7 @@ public:
 	// Begin UPrimitiveComponent interface.
 	virtual void PostPhysicsTick(FPrimitiveComponentPostPhysicsTickFunction &ThisTickFunction) override;
 	virtual class UBodySetup* GetBodySetup() override;
+	virtual bool CanEditSimulatePhysics() override;
 	virtual FBodyInstance* GetBodyInstance(FName BoneName = NAME_None) const override;
 	virtual void UpdatePhysicsToRBChannels() override;
 	virtual void SetAllPhysicsAngularVelocity(FVector const& NewVel, bool bAddToCurrent = false) override;
diff --git a/Engine/Source/Runtime/Engine/Classes/Components/StaticMeshComponent.h b/Engine/Source/Runtime/Engine/Classes/Components/StaticMeshComponent.h
--- a/Engine/Source/Runtime/Engine/Classes/Components/StaticMeshComponent.h
+++ b/Engine/Source/Runtime/Engine/Classes/Components/StaticMeshComponent.h
@@ -257,6 +257,7 @@ public:
 	virtual ELightMapInteractionType GetStaticLightingType() const override;
 	virtual void GetStreamingTextureInfo(TArray<FStreamingTexturePrimitiveInfo>& OutStreamingTextures) const override;
 	virtual class UBodySetup* GetBodySetup() override;
+	virtual bool CanEditSimulatePhysics() override;
 	virtual FPrimitiveSceneProxy* CreateSceneProxy() override;
 	virtual bool ShouldRecreateProxyOnUpdateTransform() const override;
 	virtual bool UsesOnlyUnlitMaterials() const override;
diff --git a/Engine/Source/Runtime/Engine/Private/DestructibleComponent.cpp b/Engine/Source/Runtime/Engine/Private/DestructibleComponent.cpp
--- a/Engine/Source/Runtime/Engine/Private/DestructibleComponent.cpp
+++ b/Engine/Source/Runtime/Engine/Private/DestructibleComponent.cpp
@@ -387,6 +387,11 @@ UBodySetup* UDestructibleComponent::GetBodySetup()
 	return NULL;
 }
 
+bool UDestructibleComponent::CanEditSimulatePhysics()
+{
+	// if destructiblemeshcomponent, we will allow it always
+	return true;
+}
 
 void UDestructibleComponent::AddImpulse( FVector Impulse, FName BoneName /*= NAME_None*/, bool bVelChange /*= false*/ )
 {
diff --git a/Engine/Source/Runtime/Engine/Private/InstancedStaticMesh.cpp b/Engine/Source/Runtime/Engine/Private/InstancedStaticMesh.cpp
--- a/Engine/Source/Runtime/Engine/Private/InstancedStaticMesh.cpp
+++ b/Engine/Source/Runtime/Engine/Private/InstancedStaticMesh.cpp
@@ -1259,6 +1259,12 @@ void UInstancedStaticMeshComponent::DestroyPhysicsState()
 #endif //WITH_PHYSX
 }
 
+bool UInstancedStaticMeshComponent::CanEditSimulatePhysics()
+{
+	// if instancedstaticmeshcomponent, we will never allow it
+	return false;
+}
+
 FBoxSphereBounds UInstancedStaticMeshComponent::CalcBounds(const FTransform& BoundTransform) const
 {
 	FMatrix BoundTransformMatrix = BoundTransform.ToMatrixWithScale();
diff --git a/Engine/Source/Runtime/Engine/Private/PrimitiveComponent.cpp b/Engine/Source/Runtime/Engine/Private/PrimitiveComponent.cpp
--- a/Engine/Source/Runtime/Engine/Private/PrimitiveComponent.cpp
+++ b/Engine/Source/Runtime/Engine/Private/PrimitiveComponent.cpp
@@ -2462,6 +2462,12 @@ bool UPrimitiveComponent::CanCharacterStepUp(APawn* Pawn) const
 	}
 }
 
+bool UPrimitiveComponent::CanEditSimulatePhysics()
+{
+	//Even if there's no collision but there is a body setup, we still let them simulate physics.
+	// The object falls through the world - this behavior is debatable but what we decided on for now
+	return GetBodySetup() != nullptr;
+}
 
 #undef LOCTEXT_NAMESPACE
 
diff --git a/Engine/Source/Runtime/Engine/Private/SkeletalMeshComponentPhysics.cpp b/Engine/Source/Runtime/Engine/Private/SkeletalMeshComponentPhysics.cpp
--- a/Engine/Source/Runtime/Engine/Private/SkeletalMeshComponentPhysics.cpp
+++ b/Engine/Source/Runtime/Engine/Private/SkeletalMeshComponentPhysics.cpp
@@ -489,6 +489,11 @@ UBodySetup* USkeletalMeshComponent::GetBodySetup()
 	return NULL;
 }
 
+bool USkeletalMeshComponent::CanEditSimulatePhysics()
+{
+	return GetPhysicsAsset() != nullptr;
+}
+
 void USkeletalMeshComponent::SetSimulatePhysics(bool bSimulate)
 {
 	if ( !bEnablePhysicsOnDedicatedServer && IsRunningDedicatedServer() )
diff --git a/Engine/Source/Runtime/Engine/Private/StaticMeshComponent.cpp b/Engine/Source/Runtime/Engine/Private/StaticMeshComponent.cpp
--- a/Engine/Source/Runtime/Engine/Private/StaticMeshComponent.cpp
+++ b/Engine/Source/Runtime/Engine/Private/StaticMeshComponent.cpp
@@ -581,6 +581,18 @@ UBodySetup* UStaticMeshComponent::GetBodySetup()
 	return NULL;
 }
 
+bool UStaticMeshComponent::CanEditSimulatePhysics()
+{
+	if (UBodySetup* BodySetup = GetBodySetup())
+	{
+		return (BodySetup->AggGeom.GetElementCount() > 0) || (BodySetup->CollisionTraceFlag == CTF_UseComplexAsSimple);
+	}
+	else
+	{
+		return false;
+	}
+}
+
 FColor UStaticMeshComponent::GetWireframeColor() const
 {
 	if(bOverrideWireframeColor)
