Fixed FTransform details customization not dirtying packages correctly.

Made sure to call NotifyFinishedChangingProperties so that packages are dirtied when editing values.
Also increased precision in floating-point SetValue. We should not go via strings to do this, but we dont have any other way at the moment, so this is the next best thing.

TTP# 341202 - BLUEPRINT:  A transforms default value does not get carried over from the set values in its structure

diff --git a/Engine/Source/Editor/DetailCustomizations/Private/MathStructProxyCustomizations.cpp b/Engine/Source/Editor/DetailCustomizations/Private/MathStructProxyCustomizations.cpp
--- a/Engine/Source/Editor/DetailCustomizations/Private/MathStructProxyCustomizations.cpp
+++ b/Engine/Source/Editor/DetailCustomizations/Private/MathStructProxyCustomizations.cpp
@@ -6,7 +6,7 @@
 
 void FMathStructProxyCustomization::CustomizeChildren( TSharedRef<class IPropertyHandle> StructPropertyHandle, class IDetailChildrenBuilder& StructBuilder, IPropertyTypeCustomizationUtils& StructCustomizationUtils )
 {
-
+	PropertyUtilities = StructCustomizationUtils.GetPropertyUtilities();
 }
 
 void FMathStructProxyCustomization::MakeHeaderRow( TSharedRef<class IPropertyHandle>& StructPropertyHandle, FDetailWidgetRow& Row )
@@ -110,6 +110,8 @@ void FMatrixStructCustomization::MakeHeaderRow(TSharedRef<class IPropertyHandle>
 
 void FMatrixStructCustomization::CustomizeChildren(TSharedRef<class IPropertyHandle> StructPropertyHandle, class IDetailChildrenBuilder& StructBuilder, IPropertyTypeCustomizationUtils& StructCustomizationUtils)
 {
+	FMathStructProxyCustomization::CustomizeChildren(StructPropertyHandle, StructBuilder, StructCustomizationUtils);
+
 	TWeakPtr<IPropertyHandle> WeakHandlePtr = StructPropertyHandle;
 
 	StructBuilder.AddChildContent(LOCTEXT("LocationLabel", "Location").ToString())
@@ -450,5 +452,12 @@ bool FTransformStructCustomization::FlushValues( TWeakPtr<IPropertyHandle> Prope
 	}
 	PropertyHandle->NotifyPostChange();
 
+	if(PropertyUtilities.IsValid() && !bIsUsingSlider)
+	{
+		const bool bTopologyChange = false;
+		FPropertyChangedEvent ChangeEvent(PropertyHandle->GetProperty(), bTopologyChange, EPropertyChangeType::ValueSet);
+		PropertyUtilities->NotifyFinishedChangingProperties(ChangeEvent);
+	}
+
 	return true;
 }
\ No newline at end of file
diff --git a/Engine/Source/Editor/DetailCustomizations/Private/MathStructProxyCustomizations.h b/Engine/Source/Editor/DetailCustomizations/Private/MathStructProxyCustomizations.h
--- a/Engine/Source/Editor/DetailCustomizations/Private/MathStructProxyCustomizations.h
+++ b/Engine/Source/Editor/DetailCustomizations/Private/MathStructProxyCustomizations.h
@@ -205,6 +205,10 @@ private:
 	/** Called when a value stops being changed by a slider */
 	template<typename NumericType>
 	void OnEndSliderMovement( NumericType NewValue );
+
+protected:
+	/** Cached property utilities */
+	TSharedPtr<IPropertyUtilities> PropertyUtilities;
 };
 
 /** 
diff --git a/Engine/Source/Editor/PropertyEditor/Private/PropertyHandleImpl.cpp b/Engine/Source/Editor/PropertyEditor/Private/PropertyHandleImpl.cpp
--- a/Engine/Source/Editor/PropertyEditor/Private/PropertyHandleImpl.cpp
+++ b/Engine/Source/Editor/PropertyEditor/Private/PropertyHandleImpl.cpp
@@ -2097,7 +2097,7 @@ FPropertyAccess::Result FPropertyHandleFloat::SetValue( const float& NewValue, E
 	// Clamp the value from any meta data ranges stored on the property value
 	float FinalValue = Implementation->ClampFloatValueFromMetaData( NewValue );
 
-	const FString ValueStr = FString::Printf( TEXT("%f"), FinalValue );
+	const FString ValueStr = FString::Printf( TEXT("%.9g"), FinalValue );
 	Res = Implementation->ImportText( ValueStr, Flags );
 
 	return Res;
