Fix non-unity in MathStructProxyCustomizations

diff --git a/Engine/Source/Editor/DetailCustomizations/Private/MathStructProxyCustomizations.cpp b/Engine/Source/Editor/DetailCustomizations/Private/MathStructProxyCustomizations.cpp
--- a/Engine/Source/Editor/DetailCustomizations/Private/MathStructProxyCustomizations.cpp
+++ b/Engine/Source/Editor/DetailCustomizations/Private/MathStructProxyCustomizations.cpp
@@ -2,6 +2,7 @@
 
 #include "DetailCustomizationsPrivatePCH.h"
 #include "MathStructProxyCustomizations.h"
+#include "IPropertyUtilities.h"
 #include "ScopedTransaction.h"
 
 void FMathStructProxyCustomization::CustomizeChildren( TSharedRef<class IPropertyHandle> StructPropertyHandle, class IDetailChildrenBuilder& StructBuilder, IPropertyTypeCustomizationUtils& StructCustomizationUtils )
