Added the ability to flag custom events as CallInEditor and extended the main editor actor details view to provide a drop list in the blueprints to call valid meta tagged functions.
#TTP 337860 - Expose flagged BlueprintEvents as buttons in the Details panel that can be run in-editor
#Branch UE4
#Proj Engine, DetailsCustomization, Kismet, Kismetcompiler, BlueprintGraph

#Change Added FName FBlueprintMetadata::MD_CallInEditor as "CallInEditor"
#Change Added bool UProperty UK2Node_CustomEvent::bCallInEditor to enable call in editor states to be serialised with UK2Node_CustomEvent
#Change Modified FBlueprintGraphActionDetails::CustomizeDetails to add support for the CallInEditor tickbox.

#Change Modified UK2Node_CustomEvent::GetPaletteIcon to return a different node icon when call in editor is enabled.
#Added new icon GraphEditor.CallInEditorEvent_16x for the node corner icons and the blutility combo menu.

#Change Added FKismetUserDeclaredFunctionMetadata::bCallInEditorProperty to help propagate call in editor setting from customevent node's to compiler generated event function entry point
#Change Modified FKismetCompilerContext::CreateFunctionStubForEvent to check customevent node bCallInEditor properties and mark the generated function entry nodes as CallInEditor. This was required as a transitional phase to so the setting could be observed later during meta data generation.
#Change Modified FKismetCompilerContext::FinishCompilingFunction to check function entry points for call in editor flags in the attached FKismetUserDeclaredFunctionMetadata and add CallInEditor meta data if present.

#Change Modified AActor::ProcessEvent to allow events to be processed by the actor instance if the Function has CallInEditor meta data set to true

#Change Added a series of functions to FActorDetails to conditionally add a blutility category and a combo list ( displaying the available blutility functions from the current actor blueprint ) and a button to call the selected function on the actor selection.

diff --git a/Engine/Source/Editor/BlueprintGraph/Classes/EdGraphSchema_K2.h b/Engine/Source/Editor/BlueprintGraph/Classes/EdGraphSchema_K2.h
--- a/Engine/Source/Editor/BlueprintGraph/Classes/EdGraphSchema_K2.h
+++ b/Engine/Source/Editor/BlueprintGraph/Classes/EdGraphSchema_K2.h
@@ -132,6 +132,10 @@ public:
 
 	/** Metadata string that sets the tooltip */
 	static const FName MD_Tooltip;
+
+	/** Metadata string that indicates the specified event can be triggered in editor */
+	static const FName MD_CallInEditor;
+
 private:
 	// This class should never be instantiated
 	FBlueprintMetadata() {}
diff --git a/Engine/Source/Editor/BlueprintGraph/Classes/K2Node_CustomEvent.h b/Engine/Source/Editor/BlueprintGraph/Classes/K2Node_CustomEvent.h
--- a/Engine/Source/Editor/BlueprintGraph/Classes/K2Node_CustomEvent.h
+++ b/Engine/Source/Editor/BlueprintGraph/Classes/K2Node_CustomEvent.h
@@ -10,6 +10,10 @@ class UK2Node_CustomEvent : public UK2Node_Event
 {
 	GENERATED_UCLASS_BODY()
 
+	/** Specifies that the event can be triggered in Editor */
+	UPROPERTY()
+	bool bCallInEditor;
+
 	virtual bool IsEditable() const override;
 
 	// Begin UEdGraphNode interface
@@ -20,7 +24,7 @@ class UK2Node_CustomEvent : public UK2Node_Event
 	virtual FString GetTooltip() const override;
 	virtual FString GetDocumentationLink() const override;
 	virtual FString GetDocumentationExcerptName() const override;
-	virtual FName GetPaletteIcon(FLinearColor& OutColor) const override{ return TEXT("GraphEditor.CustomEvent_16x"); }
+	virtual FName GetPaletteIcon(FLinearColor& OutColor) const override;
 	// End UEdGraphNode interface
 
 	// Begin UK2Node interface
diff --git a/Engine/Source/Editor/BlueprintGraph/Classes/K2Node_EditablePinBase.h b/Engine/Source/Editor/BlueprintGraph/Classes/K2Node_EditablePinBase.h
--- a/Engine/Source/Editor/BlueprintGraph/Classes/K2Node_EditablePinBase.h
+++ b/Engine/Source/Editor/BlueprintGraph/Classes/K2Node_EditablePinBase.h
@@ -66,6 +66,9 @@ struct FKismetUserDeclaredFunctionMetadata
 	UPROPERTY()
 	FLinearColor InstanceTitleColor;
 
+	UPROPERTY()
+	bool bCallInEditor;
+
 public:
 	FKismetUserDeclaredFunctionMetadata()
 		: InstanceTitleColor(FLinearColor::White)
diff --git a/Engine/Source/Editor/BlueprintGraph/Private/EdGraphSchema_K2.cpp b/Engine/Source/Editor/BlueprintGraph/Private/EdGraphSchema_K2.cpp
--- a/Engine/Source/Editor/BlueprintGraph/Private/EdGraphSchema_K2.cpp
+++ b/Engine/Source/Editor/BlueprintGraph/Private/EdGraphSchema_K2.cpp
@@ -81,6 +81,8 @@ const FName FBlueprintMetadata::MD_MaterialParameterCollectionFunction(TEXT("Mat
 
 const FName FBlueprintMetadata::MD_Tooltip(TEXT("Tooltip"));
 
+const FName FBlueprintMetadata::MD_CallInEditor(TEXT("CallInEditor"));
+
 //////////////////////////////////////////////////////////////////////////
 
 #define LOCTEXT_NAMESPACE "KismetSchema"
diff --git a/Engine/Source/Editor/BlueprintGraph/Private/K2Node_CustomEvent.cpp b/Engine/Source/Editor/BlueprintGraph/Private/K2Node_CustomEvent.cpp
--- a/Engine/Source/Editor/BlueprintGraph/Private/K2Node_CustomEvent.cpp
+++ b/Engine/Source/Editor/BlueprintGraph/Private/K2Node_CustomEvent.cpp
@@ -100,6 +100,7 @@ UK2Node_CustomEvent::UK2Node_CustomEvent(const class FPostConstructInitializePro
 	bOverrideFunction = false;
 	bIsEditable = true;
 	bCanRenameNode = true;
+	bCallInEditor = false;
 }
 
 FText UK2Node_CustomEvent::GetNodeTitle(ENodeTitleType::Type TitleType) const
@@ -366,4 +367,9 @@ FString UK2Node_CustomEvent::GetDocumentationExcerptName() const
 	return TEXT("UK2Node_CustomEvent");
 }
 
+FName UK2Node_CustomEvent::GetPaletteIcon(FLinearColor& OutColor) const
+{
+	return bCallInEditor ? TEXT("GraphEditor.CallInEditorEvent_16x") : TEXT("GraphEditor.CustomEvent_16x");
+}
+
 #undef LOCTEXT_NAMESPACE
diff --git a/Engine/Source/Editor/DetailCustomizations/Private/ActorDetails.cpp b/Engine/Source/Editor/DetailCustomizations/Private/ActorDetails.cpp
--- a/Engine/Source/Editor/DetailCustomizations/Private/ActorDetails.cpp
+++ b/Engine/Source/Editor/DetailCustomizations/Private/ActorDetails.cpp
@@ -111,6 +111,7 @@ void FActorDetails::CustomizeDetails( IDetailLayoutBuilder& DetailLayout )
 		if (!HideCategories.Contains(TEXT("Blueprint")))
 		{
 			AddBlueprintCategory(DetailLayout, UniqueBlueprints);
+			AddBlutilityCategory(DetailLayout, UniqueBlueprints);
 		}
 
 		OnExtendActorDetails.Broadcast(DetailLayout, FGetSelectedActors::CreateSP(this, &FActorDetails::GetSelectedActors));
@@ -575,6 +576,78 @@ void FActorDetails::AddBlueprintCategory( IDetailLayoutBuilder& DetailBuilder, c
 	}
 }
 
+void FActorDetails::AddBlutilityCategory( IDetailLayoutBuilder& DetailBuilder, const TMap<UBlueprint*, UObject*>& UniqueBlueprints )
+{
+	// Create the Blutilities Category
+	IDetailCategoryBuilder& BlutilitiesCategory = DetailBuilder.EditCategory("Blutilities", NSLOCTEXT("Blutilities", "BlutilityTitle", "Blutilities").ToString(), ECategoryPriority::Uncommon );
+
+	// Only show the bluetilities section if a single actor is selected
+	if ( SelectedActors.Num() > 0 && DoesActorHaveBlutiltyFunctions() )
+	{
+		// Reset function Selection
+		ActiveBlutilityFunction.Reset();
+
+		// Grab actor label for later use
+		TWeakObjectPtr<AActor> ActorPtr = SelectedActors[0];
+		FText ActorLabel = NSLOCTEXT( "UnrealEd", "None", "None" );
+
+		if ( ActorPtr.IsValid() )
+		{
+			ActorLabel = FText::FromString( ActorPtr.Get()->GetActorLabel() );
+		}
+
+		FFormatNamedArguments Args;
+		Args.Add( TEXT( "ActorLabel" ), ActorLabel );
+		const FText ButtonLabel = FText::Format( NSLOCTEXT( "Blutilities", "CallInEditor_ButtonLabel", "Run" ), Args );
+		const FText ButtonToolTip = FText::Format( NSLOCTEXT( "Blutilities", "CallInEditor_ButtonTooltip", "Run Selected Blutility Function on {ActorLabel}" ), Args );
+		// Build Content
+		BlutilitiesCategory.AddCustomRow( NSLOCTEXT( "Blutilities", "CallInEditorHeader", "Blutility Functions").ToString())
+		.WholeRowContent()
+		[
+			SNew(SHorizontalBox)
+			+SHorizontalBox::Slot()
+			.AutoWidth()
+			.Padding( 0, 0, 2, 0 )
+			[
+				SNew( SBox )
+				.WidthOverride( 200 )
+				[
+					SNew( SComboButton )
+					.ContentPadding( 2 )
+					.HAlign(HAlign_Fill)
+					.OnGetMenuContent( this, &FActorDetails::BuildBlutiltyFunctionContent)
+					.ButtonContent()
+					[
+						SNew( STextBlock )
+						.Text( this, &FActorDetails::GetBlutilityComboButtonLabel ) 
+						.ToolTipText( this, &FActorDetails::GetBlutilityComboButtonLabel )
+						.Font( IDetailLayoutBuilder::GetDetailFont() )
+					]
+				]
+			]
+			+SHorizontalBox::Slot()
+			.AutoWidth()
+			[
+				SNew( SBox )
+				.WidthOverride( 50 )
+				[
+					SNew( SButton )
+					.ToolTipText( ButtonToolTip )
+					.OnClicked( this, &FActorDetails::CallBlutilityFunction )
+					.IsEnabled( this, &FActorDetails::CanCallBlutilityFunction )
+					.HAlign( HAlign_Center )
+					.ContentPadding( 2 )
+					[
+						SNew( STextBlock )
+						.Text( ButtonLabel )
+						.Font( IDetailLayoutBuilder::GetDetailFont() )
+					]
+				]
+			]
+		];
+	}
+}
+
 void FActorDetails::AddComponentsCategory( IDetailLayoutBuilder& DetailBuilder )
 {
 	AActor* SelectedActor = NULL;
@@ -1376,4 +1449,85 @@ FReply FActorDetails::ResetToBlueprintDefaults_OnClicked( TWeakObjectPtr<UBluepr
 	return FReply::Handled();
 }
 
+bool FActorDetails::DoesActorHaveBlutiltyFunctions() const
+{
+	bool bFunctionsFound = false;
+
+	if( SelectedActors.Num() > 0 )
+	{
+		TWeakObjectPtr<AActor> WeakActorPtr = SelectedActors[ 0 ];
+		UClass* ActorClass = WeakActorPtr.IsValid() ? WeakActorPtr->GetClass() : nullptr;
+
+		if( ActorClass )
+		{
+			for (TFieldIterator<UFunction> FunctionIter(ActorClass, EFieldIteratorFlags::IncludeSuper); FunctionIter; ++FunctionIter)
+			{
+				if( FunctionIter->GetBoolMetaData( FBlueprintMetadata::MD_CallInEditor ))
+				{
+					bFunctionsFound = true;
+					break;
+				}
+			}
+		}
+	}
+
+	return bFunctionsFound;
+}
+
+FText FActorDetails::GetBlutilityComboButtonLabel() const
+{
+	return ActiveBlutilityFunction.IsValid() ?  FText::FromString( ActiveBlutilityFunction->GetName() ) : 
+												NSLOCTEXT( "Blutilities", "CallInEditor_ComboLabel", "Select Blutility" );
+}
+
+TSharedRef<SWidget> FActorDetails::BuildBlutiltyFunctionContent() const
+{
+	if( SelectedActors.Num() > 0 )
+	{
+		TWeakObjectPtr<AActor> WeakActorPtr = SelectedActors[ 0 ];
+		UClass* ActorClass = WeakActorPtr.IsValid() ? WeakActorPtr->GetClass() : nullptr;
+
+		if( ActorClass )
+		{
+			FMenuBuilder MenuBuilder( true, NULL );
+			MenuBuilder.BeginSection("BlutilityFunctions", NSLOCTEXT( "Blutilities", "CallInEditorHeader", "Blutility Functions") );
+			const FSlateIcon BlutilityIcon(FEditorStyle::GetStyleSetName(), "GraphEditor.CallInEditorEvent_16x");
+
+			for (TFieldIterator<UFunction> FunctionIter(ActorClass, EFieldIteratorFlags::IncludeSuper); FunctionIter; ++FunctionIter)
+			{
+				if( FunctionIter->GetBoolMetaData( FBlueprintMetadata::MD_CallInEditor ))
+				{
+					MenuBuilder.AddMenuEntry(	FText::FromString( *FunctionIter->GetName() ), 
+												FunctionIter->GetToolTipText(), 
+												BlutilityIcon,
+												FUIAction( FExecuteAction::CreateSP( this, &FActorDetails::SetActiveBlutilityFunction, TWeakObjectPtr<UFunction>(*FunctionIter))));
+				}
+			}
+			MenuBuilder.EndSection();
+
+			return MenuBuilder.MakeWidget();
+		}
+	}
+
+	return SNullWidget::NullWidget;
+}
+
+FReply FActorDetails::CallBlutilityFunction()
+{
+	TWeakObjectPtr<AActor> ActorWeakPtr = SelectedActors.Num() ? SelectedActors[ 0 ] : nullptr;
+	
+	if( ActorWeakPtr.IsValid() && ActiveBlutilityFunction.IsValid() )
+	{
+		AActor* Actor = ActorWeakPtr.Get();
+		UFunction* Function = ActiveBlutilityFunction.Get();
+
+		if( Function->GetOuter() == Actor->GetClass() )
+		{
+			Actor->ProcessEvent( Function, NULL );
+		}
+	}
+
+	return FReply::Handled();
+}
+
 #undef LOCTEXT_NAMESPACE
diff --git a/Engine/Source/Editor/DetailCustomizations/Private/ActorDetails.h b/Engine/Source/Editor/DetailCustomizations/Private/ActorDetails.h
--- a/Engine/Source/Editor/DetailCustomizations/Private/ActorDetails.h
+++ b/Engine/Source/Editor/DetailCustomizations/Private/ActorDetails.h
@@ -85,6 +85,9 @@ private:
 	/** Creates an area meant to be displayed once for any number of selected blueprints that displays helper buttons. */
 	void AddUtilityBlueprintRows( IDetailCategoryBuilder& BlueprintCategory );
 
+	/** Creates a category to display blutility functionality in a separate space from the standard blueprint details */
+	void AddBlutilityCategory( IDetailLayoutBuilder& DetailBuilder, const TMap<UBlueprint*, UObject*>& UniqueBlueprints );
+
 	// Creates a row for a single blueprint in the selection details view   
 	void AddSingleBlueprintRow( IDetailCategoryBuilder& BlueprintCategory, UBlueprint* InBlueprint, UObject* InExampleActor);
 
@@ -107,6 +110,25 @@ private:
 	bool ResetToBlueprintDefaults_IsEnabled( TWeakObjectPtr<UBlueprint> InBlueprint ) const;
 	FReply ResetToBlueprintDefaults_OnClicked( TWeakObjectPtr<UBlueprint> InBlueprint );
 	FText ResetToBlueprintDefaults_ToolTipText( TWeakObjectPtr<UBlueprint> InBlueprint ) const;
+
+	/** Called to determine if an actor has any valid call-in-editor utility functions */
+	bool DoesActorHaveBlutiltyFunctions() const;
+
+	/** Called to retrive the active combo button label */
+	FText GetBlutilityComboButtonLabel() const;
+
+	/** Called to generate the blutility function content */
+	TSharedRef<SWidget> BuildBlutiltyFunctionContent() const;
+
+	/** Called to determine the blutility execute button's status */
+	bool CanCallBlutilityFunction() const { return ActiveBlutilityFunction.IsValid(); }
+
+	/** Called on blutility function selection to cache the active function */
+	void SetActiveBlutilityFunction( TWeakObjectPtr<UFunction> BlutilityFunction ) { ActiveBlutilityFunction = BlutilityFunction; }
+
+	/** Call the currently active Blutility Function on the active actor selection */
+	FReply CallBlutilityFunction();
+
 private:
 
 	/** Bring up the menu for user to select the path to create blueprint at */
@@ -124,4 +146,8 @@ private:
 
 	/** The material category that displays used materials */
 	TSharedPtr< class FActorMaterialCategory > MaterialCategory;
+
+	/** The active blutility function, picked from the combo list and cached here */
+	TWeakObjectPtr<UFunction> ActiveBlutilityFunction;
+
 };
diff --git a/Engine/Source/Editor/EditorStyle/Private/SlateEditorStyle.cpp b/Engine/Source/Editor/EditorStyle/Private/SlateEditorStyle.cpp
--- a/Engine/Source/Editor/EditorStyle/Private/SlateEditorStyle.cpp
+++ b/Engine/Source/Editor/EditorStyle/Private/SlateEditorStyle.cpp
@@ -3218,6 +3218,7 @@ void FSlateEditorStyle::FStyle::SetupGraphEditorStyles()
 		Set( "GraphEditor.StateMachine_16x", new IMAGE_BRUSH( "Icons/icon_Blueprint_StateMachine_16x", Icon16x16 ) );
 		Set( "GraphEditor.Event_16x", new IMAGE_BRUSH( "Icons/icon_Blueprint_Event_16x", Icon16x16 ) );
 		Set( "GraphEditor.CustomEvent_16x", new IMAGE_BRUSH( "Icons/icon_Blueprint_CustomEvent_16x", Icon16x16 ) );
+		Set( "GraphEditor.CallInEditorEvent_16x", new IMAGE_BRUSH( "Icons/icon_Blueprint_CallInEditor_16x", Icon16x16 ) );
 		Set( "GraphEditor.Timeline_16x", new IMAGE_BRUSH("Icons/icon_Blueprint_Timeline_16x", Icon16x16));
 		Set( "GraphEditor.Comment_16x", new IMAGE_BRUSH("Icons/icon_Blueprint_Comment_16x", Icon16x16));
 		Set( "GraphEditor.Switch_16x", new IMAGE_BRUSH("Icons/icon_Blueprint_Switch_16x", Icon16x16));
diff --git a/Engine/Source/Editor/Kismet/Private/BlueprintDetailsCustomization.cpp b/Engine/Source/Editor/Kismet/Private/BlueprintDetailsCustomization.cpp
--- a/Engine/Source/Editor/Kismet/Private/BlueprintDetailsCustomization.cpp
+++ b/Engine/Source/Editor/Kismet/Private/BlueprintDetailsCustomization.cpp
@@ -2284,6 +2284,31 @@ void FBlueprintGraphActionDetails::CustomizeDetails( IDetailLayoutBuilder& Detai
 					]
 				]
 			];
+			Category.AddCustomRow( LOCTEXT( "EditorCallable", "Call In Editor" ).ToString() )
+			.NameContent()
+			[
+				SNew(STextBlock)
+					.Text( LOCTEXT( "EditorCallable", "Call In Editor" ) )
+					.ToolTipText( LOCTEXT("EditorCallable_Tooltip", "Enable this event to be called from within the editor") )
+					.Font( IDetailLayoutBuilder::GetDetailFont() )
+			]
+			.ValueContent()
+			[
+				SNew(SVerticalBox)
+				+SVerticalBox::Slot()
+				.AutoHeight()
+				[
+					SNew(SHorizontalBox)
+					+SHorizontalBox::Slot()
+					.AutoWidth()
+					[
+						SNew( SCheckBox )
+						.IsChecked( this, &FBlueprintGraphActionDetails::GetIsEditorCallableEvent )
+						.ToolTipText( LOCTEXT("EditorCallable_Tooltip", "Enable this event to be called from within the editor" ))
+						.OnCheckStateChanged( this, &FBlueprintGraphActionDetails::OnEditorCallableEventModified )
+					]
+				]
+			];
 		}
 
 		IDetailCategoryBuilder& InputsCategory = DetailLayout.EditCategory("Inputs", LOCTEXT("FunctionDetailsInputs", "Inputs").ToString());
@@ -2461,6 +2486,41 @@ void FBaseBlueprintGraphActionDetails::SetRefreshDelegate(FSimpleDelegate Refres
 	((bForInputs) ? RegenerateInputsChildrenDelegate : RegenerateOutputsChildrenDelegate) = RefreshDelegate;
 }
 
+ESlateCheckBoxState::Type FBlueprintGraphActionDetails::GetIsEditorCallableEvent() const
+{
+	ESlateCheckBoxState::Type Result = ESlateCheckBoxState::Unchecked;
+
+	if( FunctionEntryNodePtr.IsValid() )
+	{
+		UK2Node_CustomEvent* CustomEventNode = Cast<UK2Node_CustomEvent>(FunctionEntryNodePtr.Get());
+
+		if( CustomEventNode && CustomEventNode->bCallInEditor )
+		{
+			Result = ESlateCheckBoxState::Checked;
+		}
+	}
+	return Result;
+}
+
+void FBlueprintGraphActionDetails::OnEditorCallableEventModified( const ESlateCheckBoxState::Type NewCheckedState ) const
+{
+	if( FunctionEntryNodePtr.IsValid() )
+	{
+		if( UK2Node_CustomEvent* CustomEventNode = Cast<UK2Node_CustomEvent>(FunctionEntryNodePtr.Get()) )
+		{
+			if( UBlueprint* Blueprint = FunctionEntryNodePtr->GetBlueprint() )
+			{
+				const bool bCallInEditor = NewCheckedState == ESlateCheckBoxState::Checked;
+				const FText TransactionType = bCallInEditor ?	LOCTEXT( "DisableCallInEditor", "Disable Call In Editor " ) : 
+																LOCTEXT( "EnableCallInEditor", "Enable Call In Editor" );
+				const FScopedTransaction Transaction( TransactionType );
+				CustomEventNode->bCallInEditor = bCallInEditor;
+				FBlueprintEditorUtils::MarkBlueprintAsModified( CustomEventNode->GetBlueprint() );
+			}
+		}
+	}
+}
+
 UMulticastDelegateProperty* FBlueprintDelegateActionDetails::GetDelegatePoperty() const
 {
 	if (MyBlueprint.IsValid())
@@ -4758,6 +4818,16 @@ void FBlueprintGraphNodeDetails::OnNameCommitted(const FText& InNewText, ETextCo
 	}
 }
 
+UBlueprint* FBlueprintGraphNodeDetails::GetBlueprintObj() const
+{
+	if(BlueprintEditorPtr.IsValid())
+	{
+		return BlueprintEditorPtr.Pin()->GetBlueprintObj();
+	}
+
+	return NULL;
+}
+
 TSharedRef<IDetailCustomization> FChildActorComponentDetails::MakeInstance(TWeakPtr<FBlueprintEditor> BlueprintEditorPtrIn)
 {
 	return MakeShareable(new FChildActorComponentDetails(BlueprintEditorPtrIn));
diff --git a/Engine/Source/Editor/Kismet/Private/BlueprintDetailsCustomization.h b/Engine/Source/Editor/Kismet/Private/BlueprintDetailsCustomization.h
--- a/Engine/Source/Editor/Kismet/Private/BlueprintDetailsCustomization.h
+++ b/Engine/Source/Editor/Kismet/Private/BlueprintDetailsCustomization.h
@@ -197,6 +197,9 @@ protected:
 	/** Tries to create the result node (if there are output args) */
 	bool AttemptToCreateResultNode();
 
+	/** Toggles the ability to be called in editor */
+	void OnToggleEditorCallableEvent( const ESlateCheckBoxState::Type NewCheckedState, TWeakObjectPtr<UK2Node_EditablePinBase> SelectedNode ) const;
+
 protected:
 	/** Pointer to the parent */
 	TWeakPtr<SMyBlueprint> MyBlueprint;
@@ -442,6 +445,13 @@ private:
 	bool IsPureFunctionVisible() const;
 	void OnIsPureFunctionModified(const ESlateCheckBoxState::Type NewCheckedState);
 	ESlateCheckBoxState::Type GetIsPureFunction() const;
+
+	/** Determines if the selected event is identified as editor callable */
+	ESlateCheckBoxState::Type GetIsEditorCallableEvent() const;
+
+	/** Enables/Disables selected event as editor callable  */
+	void OnEditorCallableEventModified( const ESlateCheckBoxState::Type NewCheckedState ) const;
+
 	
 	FReply OnAddNewOutputClicked();
 
@@ -673,7 +683,12 @@ public:
 
 	/** IDetailCustomization interface */
 	virtual void CustomizeDetails( IDetailLayoutBuilder& DetailLayout ) override;
-	
+
+protected:
+
+	/** Returns the currently edited blueprint */
+	UBlueprint* GetBlueprintObj() const;
+
 private:
 
 	// Callbacks for uproperty details customization
diff --git a/Engine/Source/Editor/KismetCompiler/Private/KismetCompiler.cpp b/Engine/Source/Editor/KismetCompiler/Private/KismetCompiler.cpp
--- a/Engine/Source/Editor/KismetCompiler/Private/KismetCompiler.cpp
+++ b/Engine/Source/Editor/KismetCompiler/Private/KismetCompiler.cpp
@@ -1151,6 +1151,10 @@ void FKismetCompilerContext::PrecompileFunction(FKismetFunctionContext& Context)
 		{
 			Context.Function->SetMetaData(FBlueprintMetadata::MD_FunctionCategory, *FunctionMetaData.Category);
 		}
+		if( FunctionMetaData.bCallInEditor )
+		{
+			Context.Function->SetMetaData(FBlueprintMetadata::MD_CallInEditor, TEXT( "true" ));
+		}
 		
 		// Link it
 		//@TODO: should this be in regular or reverse order?
@@ -1531,7 +1535,10 @@ void FKismetCompilerContext::FinishCompilingFunction(FKismetFunctionContext& Con
 	{
 		Function->SetMetaData(FBlueprintMetadata::MD_Tooltip, *EntryNode->MetaData.ToolTip);
 	}
-
+	if (EntryNode->MetaData.bCallInEditor)
+	{
+		Function->SetMetaData(FBlueprintMetadata::MD_CallInEditor, TEXT( "true" ));
+	}
 	if (auto WorldContextPin = EntryNode->GetAutoWorldContextPin())
 	{
 		Function->SetMetaData(FBlueprintMetadata::MD_DefaultToSelf, *WorldContextPin->PinName); 
@@ -2225,6 +2232,8 @@ void FKismetCompilerContext::CreateFunctionStubForEvent(UK2Node_Event* SrcEventN
 		EntryNode->UserDefinedPins = SrcCustomEventNode->UserDefinedPins;
 		// CustomEvents may inherit net flags (so let's use their GetNetFlags() incase this is an override)
 		StubContext.MarkAsNetFunction(SrcCustomEventNode->GetNetFlags());
+		// Synchronize the entry node call in editor value with the entry point
+		EntryNode->MetaData.bCallInEditor = SrcCustomEventNode->bCallInEditor;
 	}
 	EntryNode->AllocateDefaultPins();
 
diff --git a/Engine/Source/Runtime/Engine/Private/Actor.cpp b/Engine/Source/Runtime/Engine/Private/Actor.cpp
--- a/Engine/Source/Runtime/Engine/Private/Actor.cpp
+++ b/Engine/Source/Runtime/Engine/Private/Actor.cpp
@@ -609,7 +609,8 @@ void AActor::PostLoadSubobjects(FObjectInstancingGraph* OuterInstanceGraph)
 
 void AActor::ProcessEvent(UFunction* Function, void* Parameters)
 {
-	if( ((GetWorld() && (GetWorld()->AreActorsInitialized() || GAllowActorScriptExecutionInEditor)) || HasAnyFlags(RF_ClassDefaultObject)) && !GIsGarbageCollecting )
+	const bool bAllowScriptExecution = GAllowActorScriptExecutionInEditor || Function->GetBoolMetaData( "CallInEditor" );
+	if( ((GetWorld() && (GetWorld()->AreActorsInitialized() || bAllowScriptExecution)) || HasAnyFlags(RF_ClassDefaultObject)) && !GIsGarbageCollecting )
 	{
 		Super::ProcessEvent(Function, Parameters);
 	}
