[341253] Blueprints - Non-native root component location/rotation info is no longer editable.

The location & rotation will be overwritten at spawn/placement time, so editing those values in the Blueprint editor Details tab would have no effect at construction time and was thus causing some confusion.

Note that native root component transform information was not previously editable. This change exposes the Scale for both native and non-native root components for editing and hides the Location/Rotation fields.

Also added some extra information to the root component's Rich Tooltip in Components mode of the Blueprint editor, to indicate why the location & rotation fields of the Transform section are not available for editing.

diff --git a/Engine/Source/Editor/DetailCustomizations/Private/ComponentTransformDetails.cpp b/Engine/Source/Editor/DetailCustomizations/Private/ComponentTransformDetails.cpp
--- a/Engine/Source/Editor/DetailCustomizations/Private/ComponentTransformDetails.cpp
+++ b/Engine/Source/Editor/DetailCustomizations/Private/ComponentTransformDetails.cpp
@@ -92,6 +92,7 @@ FComponentTransformDetails::FComponentTransformDetails( const TArray< TWeakObjec
 	, bPreserveScaleRatio( false )
 	, NotifyHook( DetailBuilder.GetPropertyUtilities()->GetNotifyHook() )
 	, bEditingRotationInUI( false )
+	, HiddenFieldMask( 0 )
 {
 	GConfig->GetBool(TEXT("SelectionDetails"), TEXT("PreserveScaleRatio"), bPreserveScaleRatio, GEditorUserSettingsIni);
 
@@ -394,8 +395,14 @@ void FComponentTransformDetails::GenerateChildContent( IDetailChildrenBuilder& C
 		
 	FSlateFontInfo FontInfo = IDetailLayoutBuilder::GetDetailFont();
 
+	const bool bHideLocationField = ( HiddenFieldMask & ( 1 << ETransformField::Location ) ) != 0;
+	const bool bHideRotationField = ( HiddenFieldMask & ( 1 << ETransformField::Rotation ) ) != 0;
+	const bool bHideScaleField = ( HiddenFieldMask & ( 1 << ETransformField::Scale ) ) != 0;
+
 	// Location
-	ChildrenBuilder.AddChildContent( LOCTEXT("LocationFilter", "Location").ToString() )
+	if(!bHideLocationField)
+	{
+		ChildrenBuilder.AddChildContent( LOCTEXT("LocationFilter", "Location").ToString() )
 		.CopyAction( CreateCopyAction( ETransformField::Location ) )
 		.PasteAction( CreatePasteAction( ETransformField::Location ) )
 		.NameContent()
@@ -409,7 +416,7 @@ void FComponentTransformDetails::GenerateChildContent( IDetailChildrenBuilder& C
 		.MaxDesiredWidth(125.0f * 3.0f)
 		[
 			SNew( SHorizontalBox )
-			+ SHorizontalBox::Slot()
+			+SHorizontalBox::Slot()
 			.FillWidth(1)
 			.VAlign( VAlign_Center )
 			[
@@ -424,14 +431,13 @@ void FComponentTransformDetails::GenerateChildContent( IDetailChildrenBuilder& C
 				.OnZCommitted( this, &FComponentTransformDetails::OnSetLocation, 2 )
 				.Font( FontInfo )
 			]
-			+ SHorizontalBox::Slot()
+			+SHorizontalBox::Slot()
 			.AutoWidth()
 			[
 				// Just take up space for alignment
 				SNew( SBox )
 				.WidthOverride( 18.0f )
 			]
-
 			+SHorizontalBox::Slot()
 			.VAlign(VAlign_Center)
 			.AutoWidth()
@@ -449,10 +455,12 @@ void FComponentTransformDetails::GenerateChildContent( IDetailChildrenBuilder& C
 				]
 			]
 		];
-
+	}
 	
 	// Rotation
-	ChildrenBuilder.AddChildContent( LOCTEXT("RotationFilter", "Rotation").ToString() )
+	if(!bHideRotationField)
+	{
+		ChildrenBuilder.AddChildContent( LOCTEXT("RotationFilter", "Rotation").ToString() )
 		.CopyAction( CreateCopyAction(ETransformField::Rotation) )
 		.PasteAction( CreatePasteAction(ETransformField::Rotation) )
 		.NameContent()
@@ -466,7 +474,7 @@ void FComponentTransformDetails::GenerateChildContent( IDetailChildrenBuilder& C
 		.MaxDesiredWidth(125.0f * 3.0f)
 		[
 			SNew( SHorizontalBox )
-			+ SHorizontalBox::Slot()
+			+SHorizontalBox::Slot()
 			.FillWidth(1)
 			.VAlign( VAlign_Center )
 			[
@@ -487,14 +495,13 @@ void FComponentTransformDetails::GenerateChildContent( IDetailChildrenBuilder& C
 				.OnYawCommitted( this, &FComponentTransformDetails::OnRotationCommitted, 2 )
 				.Font( FontInfo )
 			]
-			+ SHorizontalBox::Slot()
+			+SHorizontalBox::Slot()
 			.AutoWidth()
 			[
 				// Just take up space for alignment
 				SNew( SBox )
 				.WidthOverride( 18.0f )
 			]
-
 			+SHorizontalBox::Slot()
 			.VAlign(VAlign_Center)
 			.AutoWidth()
@@ -512,9 +519,12 @@ void FComponentTransformDetails::GenerateChildContent( IDetailChildrenBuilder& C
 				]
 			]
 		];
-
-			
-	ChildrenBuilder.AddChildContent( LOCTEXT("ScaleFilter", "Scale").ToString() )
+	}
+	
+	// Scale
+	if(!bHideScaleField)
+	{
+		ChildrenBuilder.AddChildContent( LOCTEXT("ScaleFilter", "Scale").ToString() )
 		.CopyAction( CreateCopyAction(ETransformField::Scale) )
 		.PasteAction( CreatePasteAction(ETransformField::Scale) )
 		.NameContent()
@@ -528,7 +538,7 @@ void FComponentTransformDetails::GenerateChildContent( IDetailChildrenBuilder& C
 		.MaxDesiredWidth(125.0f * 3.0f)
 		[
 			SNew( SHorizontalBox )
-			+ SHorizontalBox::Slot()
+			+SHorizontalBox::Slot()
 			.VAlign( VAlign_Center )
 			.FillWidth(1.0f)
 			[
@@ -546,7 +556,7 @@ void FComponentTransformDetails::GenerateChildContent( IDetailChildrenBuilder& C
 				.ContextMenuExtenderZ( this, &FComponentTransformDetails::ExtendZScaleContextMenu )
 				.Font( FontInfo )
 			]
-			+ SHorizontalBox::Slot()
+			+SHorizontalBox::Slot()
 			.AutoWidth()
 			.MaxWidth( 18.0f )
 			[
@@ -563,7 +573,6 @@ void FComponentTransformDetails::GenerateChildContent( IDetailChildrenBuilder& C
 					.ColorAndOpacity( FSlateColor::UseForeground() )
 				]
 			]
-
 			+SHorizontalBox::Slot()
 			.VAlign(VAlign_Center)
 			.AutoWidth()
@@ -581,6 +590,7 @@ void FComponentTransformDetails::GenerateChildContent( IDetailChildrenBuilder& C
 				]
 			]
 		];
+	}
 }
 END_SLATE_FUNCTION_BUILD_OPTIMIZATION
 
@@ -589,7 +599,7 @@ void FComponentTransformDetails::Tick( float DeltaTime )
 	CacheTransform();
 }
 
-bool FComponentTransformDetails::GetIsEnabled( ) const
+bool FComponentTransformDetails::GetIsEnabled() const
 {
 	return !GEditor->HasLockedActors() || SelectedActorInfo.NumSelected == 0;
 }
diff --git a/Engine/Source/Editor/DetailCustomizations/Private/ComponentTransformDetails.h b/Engine/Source/Editor/DetailCustomizations/Private/ComponentTransformDetails.h
--- a/Engine/Source/Editor/DetailCustomizations/Private/ComponentTransformDetails.h
+++ b/Engine/Source/Editor/DetailCustomizations/Private/ComponentTransformDetails.h
@@ -34,6 +34,11 @@ public:
 	virtual void Tick( float DeltaTime ) override;
 	virtual void SetOnRebuildChildren( FSimpleDelegate OnRebuildChildren ) override{}
 
+	void HideTransformField(const ETransformField::Type InTransformField)
+	{
+		HiddenFieldMask |= (1 << InTransformField);
+	}
+
 private:
 	/** @return Whether the transform details panel should be enabled (editable) or not (read-only / greyed out) */
 	bool GetIsEnabled() const;
@@ -306,4 +311,6 @@ private:
 	TMap< UObject*, FRotator > ObjectToRelativeRotationMap;
 	/** Flag to indicate we are currently editing the rotation in the UI, so we should rely on the cached value in objectToRelativeRotationMap, not the value from the object */
 	bool bEditingRotationInUI;
+	/** Bitmask to indicate which fields should be hidden (if any) */
+	uint8 HiddenFieldMask;
 };
diff --git a/Engine/Source/Editor/DetailCustomizations/Private/SceneComponentDetails.cpp b/Engine/Source/Editor/DetailCustomizations/Private/SceneComponentDetails.cpp
--- a/Engine/Source/Editor/DetailCustomizations/Private/SceneComponentDetails.cpp
+++ b/Engine/Source/Editor/DetailCustomizations/Private/SceneComponentDetails.cpp
@@ -639,14 +639,13 @@ void FSceneComponentDetails::MakeTransformDetails( IDetailLayoutBuilder& DetailB
 		for (int32 ComponentIndex = 0; bShouldShowTransform && ComponentIndex < SceneComponentObjects.Num(); ++ComponentIndex)
 		{
 			USceneComponent* SceneComponent = Cast<USceneComponent>(SceneComponentObjects[ComponentIndex].Get());
-			UBlueprint* Blueprint = (SceneComponent ? Cast<UBlueprint>(SceneComponent->GetOuter()) : NULL);
 			if (SceneComponent->AttachParent == NULL && SceneComponent->GetOuter()->HasAnyFlags(RF_ClassDefaultObject))
 			{
 				bShouldShowTransform = false;
 			}
-			else if (Blueprint && Blueprint->SimpleConstructionScript)
+			else if (const USimpleConstructionScript* SCS = GetSimpleConstructionScript(SceneComponent))
 			{
-				const TArray<USCS_Node*>& RootNodes = Blueprint->SimpleConstructionScript->GetRootNodes();
+				const TArray<USCS_Node*>& RootNodes = SCS->GetRootNodes();
 				for(int32 RootNodeIndex = 0; bShouldShowTransform && RootNodeIndex < RootNodes.Num(); ++RootNodeIndex)
 				{
 					USCS_Node* RootNode = RootNodes[RootNodeIndex];
@@ -660,14 +659,17 @@ void FSceneComponentDetails::MakeTransformDetails( IDetailLayoutBuilder& DetailB
 			}
 		}
 
-		if( bShouldShowTransform )
+		TSharedRef<FComponentTransformDetails> TransformDetails = MakeShareable( new FComponentTransformDetails( bIsEditingBlueprintDefaults ? SceneComponentObjects : DetailBuilder.GetDetailsView().GetSelectedObjects(), SelectedActorInfo, DetailBuilder ) );
+
+		if(!bShouldShowTransform)
 		{
-			TSharedRef<FComponentTransformDetails> TransformDetails = MakeShareable( new FComponentTransformDetails( bIsEditingBlueprintDefaults ? SceneComponentObjects : DetailBuilder.GetDetailsView().GetSelectedObjects(), SelectedActorInfo, DetailBuilder ) );
+			TransformDetails->HideTransformField(ETransformField::Location);
+			TransformDetails->HideTransformField(ETransformField::Rotation);
+		}
 
-			IDetailCategoryBuilder& TransformCategory = DetailBuilder.EditCategory( "TransformCommon", LOCTEXT("TransformCommonCategory", "Transform").ToString(), ECategoryPriority::Transform );
+		IDetailCategoryBuilder& TransformCategory = DetailBuilder.EditCategory( "TransformCommon", LOCTEXT("TransformCommonCategory", "Transform").ToString(), ECategoryPriority::Transform );
 
-			TransformCategory.AddCustomBuilder( TransformDetails );
-		}
+		TransformCategory.AddCustomBuilder( TransformDetails );
 	}
 
 }
