Fixed a couple of confusing part in additive
-  You can't ever set the animation with different skeleton, but it didn't have indication, so it looked just broken, so I added text
-  Make sure you don't change skeleton based on animation because that doesn't work.

diff --git a/Engine/Source/Editor/DetailCustomizations/Private/AnimSequenceDetails.cpp b/Engine/Source/Editor/DetailCustomizations/Private/AnimSequenceDetails.cpp
--- a/Engine/Source/Editor/DetailCustomizations/Private/AnimSequenceDetails.cpp
+++ b/Engine/Source/Editor/DetailCustomizations/Private/AnimSequenceDetails.cpp
@@ -347,6 +347,7 @@ void SAnimationRefPoseViewport::Construct(const FArguments& InArgs)
 		[
 			SAssignNew(Description, STextBlock)
 			.Text(LOCTEXT("DefaultViewportLabel", "Default View"))
+			.AutoWrapText(true)
 		]
 		+SVerticalBox::Slot()
 		.AutoHeight()
@@ -406,29 +407,35 @@ void SAnimationRefPoseViewport::InitSkeleton()
 	{
 		Skeleton = AnimRef->GetSkeleton();
 	}
+	else
+	{
+		Skeleton = TargetSkeleton;
+	}
 
-	if(PreviewComponent != NULL && Skeleton != NULL)
+	// if skeleton doesn't match with target skeleton, this is error, we can't support it
+	if ( Skeleton == TargetSkeleton)
 	{
-		UAnimSingleNodeInstance * Preview = PreviewComponent->PreviewInstance;
-		USkeletalMesh* PreviewSkeletalMesh = Skeleton->GetPreviewMesh();
-		if((Preview == NULL || Preview->CurrentAsset != AnimRef) || PreviewComponent->SkeletalMesh != PreviewSkeletalMesh)
+		if(PreviewComponent != NULL && Skeleton != NULL)
 		{
-			PreviewComponent->SetSkeletalMesh(PreviewSkeletalMesh);
-			PreviewComponent->EnablePreview(true, AnimRef, NULL);
-			PreviewComponent->PreviewInstance->SetLooping(true);
-
-			//Place the camera at a good viewer position
-			FVector NewPosition = LevelViewportClient->GetViewLocation();
-			NewPosition.Normalize();
-			if(PreviewSkeletalMesh)
+			UAnimSingleNodeInstance * Preview = PreviewComponent->PreviewInstance;
+			USkeletalMesh* PreviewSkeletalMesh = Skeleton->GetPreviewMesh();
+			if((Preview == NULL || Preview->CurrentAsset != AnimRef) || PreviewComponent->SkeletalMesh != PreviewSkeletalMesh)
 			{
-				NewPosition *= (PreviewSkeletalMesh->Bounds.SphereRadius*1.5f);
+				PreviewComponent->SetSkeletalMesh(PreviewSkeletalMesh);
+				PreviewComponent->EnablePreview(true, AnimRef, NULL);
+				PreviewComponent->PreviewInstance->SetLooping(true);
+
+				//Place the camera at a good viewer position
+				FVector NewPosition = LevelViewportClient->GetViewLocation();
+				NewPosition.Normalize();
+				if(PreviewSkeletalMesh)
+				{
+					NewPosition *= (PreviewSkeletalMesh->Bounds.SphereRadius*1.5f);
+				}
+				LevelViewportClient->SetViewLocation( NewPosition );
 			}
-			LevelViewportClient->SetViewLocation( NewPosition );
 		}
 	}
-
-	TargetSkeleton = Skeleton;
 }
 
 void SAnimationRefPoseViewport::Tick(const FGeometry& AllottedGeometry, const double InCurrentTime, const float InDeltaTime)
@@ -473,13 +480,17 @@ void SAnimationRefPoseViewport::Tick(const FGeometry& AllottedGeometry, const do
 		{
 			Description->SetText(FText::Format( LOCTEXT("Previewing", "Previewing {0}"), FText::FromString(Component->AnimBlueprintGeneratedClass->GetName()) ).ToString());
 		}
+		else if (AnimRef && AnimRef->GetSkeleton() != TargetSkeleton)
+		{
+			Description->SetText(FText::Format( LOCTEXT("IncorrectSkeleton", "The preview asset doesn't work for the skeleton '{0}'"), FText::FromString(TargetSkeletonName) ).ToString());
+		}
 		else if (Component->SkeletalMesh == NULL)
 		{
 			Description->SetText(FText::Format( LOCTEXT("NoMeshFound", "No skeletal mesh found for skeleton '{0}'"), FText::FromString(TargetSkeletonName) ).ToString());
 		}
 		else
 		{
-			Description->SetText(LOCTEXT("Default", "Default").ToString());
+			Description->SetText(FText::Format( LOCTEXT("SelectAnimation", "Select animation that works for skeleton '{0}'"), FText::FromString(TargetSkeletonName) ).ToString());
 		}
 
 		Component->GetScene()->GetWorld()->Tick(LEVELTICK_All, InDeltaTime);
