Revised FActorDetails::CustomizeDetails to respect some hidden categories.

diff --git a/Engine/Source/Editor/DetailCustomizations/Private/ActorDetails.cpp b/Engine/Source/Editor/DetailCustomizations/Private/ActorDetails.cpp
--- a/Engine/Source/Editor/DetailCustomizations/Private/ActorDetails.cpp
+++ b/Engine/Source/Editor/DetailCustomizations/Private/ActorDetails.cpp
@@ -95,17 +95,26 @@ void FActorDetails::CustomizeDetails( IDetailLayoutBuilder& DetailLayout )
 			DetailLayout.HideProperty( GET_MEMBER_NAME_CHECKED(AActor, SpriteScale) );
 		}
 
-		AddExperimentalWarningCategory( DetailLayout );
-
-		AddTransformCategory( DetailLayout );
-
-		AddMaterialCategory( DetailLayout );
-
-		AddActorCategory( DetailLayout, ActorsPerLevelCount );
+		AddExperimentalWarningCategory(DetailLayout);
 
 		// Get the list of hidden categories
 		TArray<FString> HideCategories;
-		FEditorCategoryUtils::GetClassHideCategories(DetailLayout.GetDetailsView().GetBaseClass(), HideCategories);
+		FEditorCategoryUtils::GetClassHideCategories(DetailLayout.GetDetailsView().GetBaseClass(), HideCategories); 
+
+		if (!HideCategories.Contains(TEXT("Transform")))
+		{
+			AddTransformCategory(DetailLayout);
+		}
+
+		if (!HideCategories.Contains(TEXT("Materials")))
+		{
+			AddMaterialCategory(DetailLayout);
+		}
+		
+		if (!HideCategories.Contains(TEXT("Actor")))
+		{
+			AddActorCategory(DetailLayout, ActorsPerLevelCount);
+		}
 
 		// Add Blueprint category, if not being hidden
 		if (!HideCategories.Contains(TEXT("Blueprint")))
