Added DocumentationActor.

diff --git a/Engine/Source/Editor/DetailCustomizations/Private/ActorDetails.cpp b/Engine/Source/Editor/DetailCustomizations/Private/ActorDetails.cpp
--- a/Engine/Source/Editor/DetailCustomizations/Private/ActorDetails.cpp
+++ b/Engine/Source/Editor/DetailCustomizations/Private/ActorDetails.cpp
@@ -28,6 +28,7 @@
 #include "Runtime/Engine/Classes/Engine/BrushShape.h"
 #include "ActorDetailsDelegates.h"
 #include "EditorCategoryUtils.h"
+#include "Engine/DocumentationActor.h"
 
 #define LOCTEXT_NAMESPACE "ActorDetails"
 
diff --git a/Engine/Source/Editor/DetailCustomizations/Private/DetailCustomizations.cpp b/Engine/Source/Editor/DetailCustomizations/Private/DetailCustomizations.cpp
--- a/Engine/Source/Editor/DetailCustomizations/Private/DetailCustomizations.cpp
+++ b/Engine/Source/Editor/DetailCustomizations/Private/DetailCustomizations.cpp
@@ -75,7 +75,7 @@
 #include "MoviePlayerSettingsDetails.h"
 #include "SourceCodeAccessSettingsDetails.h"
 #include "ParticleSystemComponentDetails.h"
-
+#include "DocumentationActorDetails.h"
 
 IMPLEMENT_MODULE( FDetailCustomizationsModule, DetailCustomizations );
 
@@ -192,6 +192,7 @@ void FDetailCustomizationsModule::RegisterObjectCustomizations()
 	RegisterCustomClassLayout("AmbientSound", FOnGetDetailCustomizationInstance::CreateStatic(&FAmbientSoundDetails::MakeInstance));
 	RegisterCustomClassLayout("WorldSettings", FOnGetDetailCustomizationInstance::CreateStatic(&FWorldSettingsDetails::MakeInstance));
 	RegisterCustomClassLayout("GeneralProjectSettings", FOnGetDetailCustomizationInstance::CreateStatic(&FGeneralProjectSettingsDetails::MakeInstance));
+	RegisterCustomClassLayout("DocumentationActor", FOnGetDetailCustomizationInstance::CreateStatic(&FDocumentationActorDetails::MakeInstance));
 
 	//@TODO: A2REMOVAL: Rename FSkeletalControlNodeDetails to something more generic
 	RegisterCustomClassLayout("K2Node_StructMemberGet", FOnGetDetailCustomizationInstance::CreateStatic(&FSkeletalControlNodeDetails::MakeInstance));
diff --git a/Engine/Source/Editor/DetailCustomizations/Private/DocumentationActorDetails.cpp b/Engine/Source/Editor/DetailCustomizations/Private/DocumentationActorDetails.cpp
new file mode 100644
--- /dev/null
+++ b/Engine/Source/Editor/DetailCustomizations/Private/DocumentationActorDetails.cpp
@@ -0,0 +1,69 @@
+// Copyright 1998-2014 Epic Games, Inc. All Rights Reserved.
+
+#include "DetailCustomizationsPrivatePCH.h"
+#include "DocumentationActorDetails.h"
+#include "Engine/DocumentationActor.h"
+#include "ObjectEditorUtils.h"
+
+
+#include "DetailLayoutBuilder.h"
+#include "DetailCategoryBuilder.h"
+#include "PropertyEditing.h"
+#include "EditorCategoryUtils.h"
+
+
+#define LOCTEXT_NAMESPACE "DocumentationActorDetails"
+
+
+TSharedRef<IDetailCustomization> FDocumentationActorDetails::MakeInstance()
+{
+	return MakeShareable(new FDocumentationActorDetails);
+}
+
+void FDocumentationActorDetails::CustomizeDetails(IDetailLayoutBuilder& DetailBuilder)
+{
+	SelectedDocActor = nullptr; 
+	
+	const TArray< TWeakObjectPtr<UObject> >& SelectedObjects = DetailBuilder.GetDetailsView().GetSelectedObjects();
+	
+	// Find the first documentation actor -we should only open one documentation instance at a time.
+	for (int32 ObjectIndex = 0; ObjectIndex < SelectedObjects.Num(); ++ObjectIndex)
+	{
+		if (ADocumentationActor* CurrentDocActor = Cast<ADocumentationActor>(SelectedObjects[ObjectIndex].Get()))
+		{
+			if (CurrentDocActor != nullptr)
+			{
+				if (CurrentDocActor->HasValidDocumentLink() == true)
+				{
+					SelectedDocActor = CurrentDocActor;
+					break;
+				}
+			}
+		}
+	}
+
+	// Add a button we can click on to open the documentation
+	IDetailCategoryBuilder& HelpCategory = DetailBuilder.EditCategory("Help Data");
+	HelpCategory.AddCustomRow("Help Documentation")
+		[
+			SNew(SButton)
+			.Text(LOCTEXT("HelpDocumentation", "Help Documentation"))
+			.HAlign(HAlign_Center)
+			.OnClicked(this, &FDocumentationActorDetails::OnHelpButtonClicked)
+		];
+	
+}
+
+FReply FDocumentationActorDetails::OnHelpButtonClicked()
+{
+	FReply Handled = FReply::Unhandled();
+	
+	if (SelectedDocActor.IsValid() == true )
+	{
+		Handled = SelectedDocActor->OpenDocumentLink() == true ? FReply::Handled() : FReply::Unhandled();
+	}
+	return Handled;
+}
+
+
+#undef LOCTEXT_NAMESPACE
diff --git a/Engine/Source/Editor/DetailCustomizations/Private/DocumentationActorDetails.h b/Engine/Source/Editor/DetailCustomizations/Private/DocumentationActorDetails.h
new file mode 100644
--- /dev/null
+++ b/Engine/Source/Editor/DetailCustomizations/Private/DocumentationActorDetails.h
@@ -0,0 +1,19 @@
+// Copyright 1998-2014 Epic Games, Inc. All Rights Reserved.
+
+#pragma once
+
+class FDocumentationActorDetails : public IDetailCustomization
+{
+public:
+	/** Makes a new instance of this detail layout class for a specific detail view requesting it */
+	static TSharedRef<IDetailCustomization> MakeInstance();
+
+	/** IDetailCustomization interface */
+	virtual void CustomizeDetails(IDetailLayoutBuilder& DetailBuilder) override;
+	
+	/* Handler for clicking the help button */
+	FReply OnHelpButtonClicked();
+private:
+	/* The documentation actor that we are showing in the details panel */
+	TWeakObjectPtr<class ADocumentationActor>	SelectedDocActor;
+};
diff --git a/Engine/Source/Editor/UnrealEd/Private/EditorServer.cpp b/Engine/Source/Editor/UnrealEd/Private/EditorServer.cpp
--- a/Engine/Source/Editor/UnrealEd/Private/EditorServer.cpp
+++ b/Engine/Source/Editor/UnrealEd/Private/EditorServer.cpp
@@ -52,6 +52,7 @@
 #include "Particles/ParticleSystemComponent.h"
 
 #include "ComponentReregisterContext.h"
+#include "Engine/DocumentationActor.h"
 
 DEFINE_LOG_CATEGORY_STATIC(LogEditorServer, Log, All);
 
@@ -4134,7 +4135,17 @@ void UEditorEngine::MoveViewportCamerasToActor(const TArray<AActor*> &Actors, bo
 	{
 		return;
 	}
-	
+
+	// If the first actor is a documentation actor open his document link
+	if (Actors.Num() == 1)
+	{
+		ADocumentationActor* DocActor = Cast<ADocumentationActor>(Actors[0]);
+		if (DocActor != nullptr)
+		{
+			DocActor->OpenDocumentLink();
+		}
+	}
+
 	TArray<AActor*> InvisLevelActors;
 
 	TArray<UClass*> PrimitiveComponentTypesToIgnore;
diff --git a/Engine/Source/Runtime/Engine/Classes/Engine/DocumentationActor.h b/Engine/Source/Runtime/Engine/Classes/Engine/DocumentationActor.h
new file mode 100644
--- /dev/null
+++ b/Engine/Source/Runtime/Engine/Classes/Engine/DocumentationActor.h
@@ -0,0 +1,32 @@
+// Copyright 1998-2014 Epic Games, Inc. All Rights Reserved.
+
+//=============================================================================
+//=============================================================================
+
+#pragma once
+#include "DocumentationActor.generated.h"
+
+UCLASS(hidecategories = (Actor, Transform, Tags, Materials, StaticMesh, Rendering, Components, Blueprint, StaticMeshActor, Object, Collision, Display, Rendering, Physics, Input, Lighting, Layers))
+class ENGINE_API ADocumentationActor : public AActor
+{
+	GENERATED_UCLASS_BODY()
+
+	/* Open the document link this actor relates to */
+	bool OpenDocumentLink() const;
+	
+	/* Is the document link this actor has valid */
+	bool HasValidDocumentLink() const;
+
+#if WITH_EDITORONLY_DATA
+	/** Scales forces applied from damage events. */
+	UPROPERTY(Category = HelpDocumentation, EditAnywhere, BlueprintReadWrite, AdvancedDisplay)
+	FString DocumentLink; 
+	
+	UPROPERTY(Category = StaticMeshActor, VisibleAnywhere, BlueprintReadOnly)
+	TSubobjectPtr<class UStaticMeshComponent> StaticMeshComponent;
+#endif
+
+};
+
+
+
diff --git a/Engine/Source/Runtime/Engine/Private/DocumentationActor.cpp b/Engine/Source/Runtime/Engine/Private/DocumentationActor.cpp
new file mode 100644
--- /dev/null
+++ b/Engine/Source/Runtime/Engine/Private/DocumentationActor.cpp
@@ -0,0 +1,51 @@
+// Copyright 1998-2014 Epic Games, Inc. All Rights Reserved.
+
+#include "EnginePrivate.h"
+#include "Engine/DocumentationActor.h"
+
+#if WITH_EDITORONLY_DATA
+#include "IDocumentation.h"
+#endif
+
+
+
+ADocumentationActor::ADocumentationActor(const class FPostConstructInitializeProperties& PCIP)
+	: Super(PCIP)
+{
+	TSubobjectPtr<USceneComponent> SceneComponent = PCIP.CreateDefaultSubobject<USceneComponent>(this, TEXT("SceneComp"));
+	RootComponent = SceneComponent;	
+
+#if WITH_EDITORONLY_DATA
+	static ConstructorHelpers::FObjectFinder<UStaticMesh> DocuemtMesh(TEXT("/Engine/EditorMeshes/EditorHelp"));
+	
+	// Create a mesh to represent our actor
+	StaticMeshComponent = PCIP.CreateDefaultSubobject<UStaticMeshComponent>(this, TEXT("MeshComponent"));
+	StaticMeshComponent->SetCollisionProfileName(UCollisionProfile::NoCollision_ProfileName);
+	StaticMeshComponent->bGenerateOverlapEvents = false;
+	StaticMeshComponent->AttachParent = RootComponent;		
+	StaticMeshComponent->SetStaticMesh(DocuemtMesh.Object);
+#endif //WITH_EDITORONLY_DATA
+}
+
+bool ADocumentationActor::OpenDocumentLink() const
+{
+	bool bOpened = false;
+
+#if WITH_EDITORONLY_DATA
+	bOpened = IDocumentation::Get()->Open(DocumentLink);
+#endif //WITH_EDITORONLY_DATA
+
+	return bOpened;
+}
+
+bool ADocumentationActor::HasValidDocumentLink() const
+{
+	bool bDocumentValid = false;
+
+#if WITH_EDITORONLY_DATA
+	bDocumentValid = DocumentLink.IsEmpty() == false; 
+#endif // WITH_EDITORONLY_DATA
+
+	return bDocumentValid;
+	
+}
