Collision: Save project collision settings (custom profiles/channels/object types) in the DefaultEngine ini, fixing issues in cooked games
#TTP 342936, integrated from 4.4 branch CL# 2239287

diff --git a/Engine/Source/Editor/DetailCustomizations/Private/CollisionProfileDetails.cpp b/Engine/Source/Editor/DetailCustomizations/Private/CollisionProfileDetails.cpp
--- a/Engine/Source/Editor/DetailCustomizations/Private/CollisionProfileDetails.cpp
+++ b/Engine/Source/Editor/DetailCustomizations/Private/CollisionProfileDetails.cpp
@@ -1815,7 +1815,7 @@ void FCollisionProfileDetails::UpdateChannel(bool bTraceType)
 	}
 
 	CollisionProfile->LoadProfileConfig(true);
-	CollisionProfile->SaveConfig();
+	CollisionProfile->UpdateDefaultConfigFile();
 
 	// update profile list if the name changed, this has to be refreshed
 	RefreshProfileList();
@@ -1827,7 +1827,7 @@ void FCollisionProfileDetails::UpdateProfile()
 	RefreshProfileList();
 	ProfileListView->RequestListRefresh();
 	CollisionProfile->LoadProfileConfig(true);
-	CollisionProfile->SaveConfig();
+	CollisionProfile->UpdateDefaultConfigFile();
 }
 
 void FCollisionProfileDetails::RefreshChannelList(bool bTraceType)
