Media: Initial revision of details customization for media assets

diff --git a/Engine/Source/Editor/DetailCustomizations/DetailCustomizations.Build.cs b/Engine/Source/Editor/DetailCustomizations/DetailCustomizations.Build.cs
--- a/Engine/Source/Editor/DetailCustomizations/DetailCustomizations.Build.cs
+++ b/Engine/Source/Editor/DetailCustomizations/DetailCustomizations.Build.cs
@@ -35,7 +35,9 @@ public class DetailCustomizations : ModuleRules
 				"SourceControl",
                 "InternationalizationSettings",
 				"SourceCodeAccess",
-				"RHI"
+				"RHI",
+                "Media",
+                "MediaAssets",
 			}
 		);
 
diff --git a/Engine/Source/Editor/DetailCustomizations/Private/DetailCustomizations.cpp b/Engine/Source/Editor/DetailCustomizations/Private/DetailCustomizations.cpp
--- a/Engine/Source/Editor/DetailCustomizations/Private/DetailCustomizations.cpp
+++ b/Engine/Source/Editor/DetailCustomizations/Private/DetailCustomizations.cpp
@@ -76,6 +76,10 @@
 #include "SourceCodeAccessSettingsDetails.h"
 #include "ParticleSystemComponentDetails.h"
 #include "DocumentationActorDetails.h"
+#include "MediaAssetCustomization.h"
+#include "MediaSoundWaveCustomization.h"
+#include "MediaTextureCustomization.h"
+
 
 IMPLEMENT_MODULE( FDetailCustomizationsModule, DetailCustomizations );
 
@@ -94,21 +98,21 @@ void FDetailCustomizationsModule::StartupModule()
 
 void FDetailCustomizationsModule::ShutdownModule()
 {
-	if( FModuleManager::Get().IsModuleLoaded( "PropertyEditor" ) )
+	if (FModuleManager::Get().IsModuleLoaded("PropertyEditor"))
 	{
 		FPropertyEditorModule& PropertyModule = FModuleManager::GetModuleChecked<FPropertyEditorModule>("PropertyEditor");
 
 		// Unregister all classes customized by name
-		for( auto It = RegisteredClassNames.CreateConstIterator(); It; ++It )
+		for (auto It = RegisteredClassNames.CreateConstIterator(); It; ++It)
 		{
-			if( It->IsValid() )
+			if (It->IsValid())
 			{
-				PropertyModule.UnregisterCustomClassLayout( *It );
+				PropertyModule.UnregisterCustomClassLayout(*It);
 			}
 		}
 
-		// Unregister all structs
-		for(auto It = RegisteredPropertyTypes.CreateConstIterator(); It; ++It)
+		// Unregister all structures
+		for (auto It = RegisteredPropertyTypes.CreateConstIterator(); It; ++It)
 		{
 			if(It->IsValid())
 			{
@@ -123,7 +127,6 @@ void FDetailCustomizationsModule::ShutdownModule()
 
 void FDetailCustomizationsModule::RegisterPropertyTypeCustomizations()
 {
-	// Structs
 	RegisterCustomPropertyTypeLayout("StringAssetReference", FOnGetPropertyTypeCustomizationInstance::CreateStatic(&FStringAssetReferenceCustomization::MakeInstance));
 	RegisterCustomPropertyTypeLayout("StringClassReference", FOnGetPropertyTypeCustomizationInstance::CreateStatic(&FStringClassReferenceCustomization::MakeInstance));
 	RegisterCustomPropertyTypeLayout("DataTableRowHandle", FOnGetPropertyTypeCustomizationInstance::CreateStatic(&FDataTableCustomizationLayout::MakeInstance));
@@ -167,6 +170,7 @@ void FDetailCustomizationsModule::RegisterPropertyTypeCustomizations()
 	RegisterCustomPropertyTypeLayout("InputAxisKeyMapping", FOnGetPropertyTypeCustomizationInstance::CreateStatic(&FInputAxisMappingCustomization::MakeInstance));
 }
 
+
 void FDetailCustomizationsModule::RegisterObjectCustomizations()
 {
 	// Note: By default properties are displayed in script defined order (i.e the order in the header).  These layout detail classes are called in the order seen here which will display properties
@@ -238,19 +242,21 @@ void FDetailCustomizationsModule::RegisterObjectCustomizations()
 
 	RegisterCustomClassLayout("SourceCodeAccessSettings", FOnGetDetailCustomizationInstance::CreateStatic(&FSourceCodeAccessSettingsDetails::MakeInstance));
 	RegisterCustomClassLayout("ParticleSystemComponent", FOnGetDetailCustomizationInstance::CreateStatic(&FParticleSystemComponentDetails::MakeInstance));
+
+	RegisterCustomClassLayout("MediaAsset", FOnGetDetailCustomizationInstance::CreateStatic(&FMediaAssetCustomization::MakeInstance));
+	RegisterCustomClassLayout("MediaSoundWave", FOnGetDetailCustomizationInstance::CreateStatic(&FMediaSoundWaveCustomization::MakeInstance));
+	RegisterCustomClassLayout("MediaTexture", FOnGetDetailCustomizationInstance::CreateStatic(&FMediaTextureCustomization::MakeInstance));
 }
 
+
 void FDetailCustomizationsModule::RegisterCustomClassLayout(FName ClassName, FOnGetDetailCustomizationInstance DetailLayoutDelegate )
 {
-	check( ClassName != NAME_None );
+	check(ClassName != NAME_None);
 
-	static FName PropertyEditor("PropertyEditor");
-
-	// Add the class to the list of classes we should unregister later
-	RegisteredClassNames.Add( ClassName );
+	RegisteredClassNames.Add(ClassName);
 
+	static FName PropertyEditor("PropertyEditor");
 	FPropertyEditorModule& PropertyModule = FModuleManager::GetModuleChecked<FPropertyEditorModule>(PropertyEditor);
-
 	PropertyModule.RegisterCustomClassLayout( ClassName, DetailLayoutDelegate );
 }
 
@@ -259,11 +265,9 @@ void FDetailCustomizationsModule::RegisterCustomPropertyTypeLayout(FName Propert
 {
 	check(PropertyTypeName != NAME_None);
 
-	static FName PropertyEditor("PropertyEditor");
-
 	RegisteredPropertyTypes.Add(PropertyTypeName);
 
+	static FName PropertyEditor("PropertyEditor");
 	FPropertyEditorModule& PropertyModule = FModuleManager::GetModuleChecked<FPropertyEditorModule>(PropertyEditor);
-
 	PropertyModule.RegisterCustomPropertyTypeLayout(PropertyTypeName, PropertyTypeLayoutDelegate);
-}
\ No newline at end of file
+}
diff --git a/Engine/Source/Editor/DetailCustomizations/Private/MediaAssetCustomization.cpp b/Engine/Source/Editor/DetailCustomizations/Private/MediaAssetCustomization.cpp
new file mode 100644
--- /dev/null
+++ b/Engine/Source/Editor/DetailCustomizations/Private/MediaAssetCustomization.cpp
@@ -0,0 +1,13 @@
+// Copyright 1998-2014 Epic Games, Inc. All Rights Reserved.
+
+#include "DetailCustomizationsPrivatePCH.h"
+#include "MediaAssetCustomization.h"
+
+
+/* IDetailCustomization interface
+ *****************************************************************************/
+
+void FMediaAssetCustomization::CustomizeDetails( IDetailLayoutBuilder& DetailBuilder )
+{
+
+}
diff --git a/Engine/Source/Editor/DetailCustomizations/Private/MediaAssetCustomization.h b/Engine/Source/Editor/DetailCustomizations/Private/MediaAssetCustomization.h
new file mode 100644
--- /dev/null
+++ b/Engine/Source/Editor/DetailCustomizations/Private/MediaAssetCustomization.h
@@ -0,0 +1,32 @@
+// Copyright 1998-2014 Epic Games, Inc. All Rights Reserved.
+
+#pragma once
+
+#include "Media.h"
+#include "MediaAsset.h"
+
+
+/**
+ * Implements a details view customization for the UMediaAsset class.
+ */
+class FMediaAssetCustomization
+	: public IDetailCustomization
+{
+public:
+
+	// IDetailCustomization interface
+
+	virtual void CustomizeDetails( IDetailLayoutBuilder& DetailBuilder ) override;
+
+public:
+
+	/**
+	 * Creates an instance of this class.
+	 *
+	 * @return The new instance.
+	 */
+	static TSharedRef<IDetailCustomization> MakeInstance( )
+	{
+		return MakeShareable(new FMediaAssetCustomization());
+	}
+};
diff --git a/Engine/Source/Editor/DetailCustomizations/Private/MediaSoundWaveCustomization.cpp b/Engine/Source/Editor/DetailCustomizations/Private/MediaSoundWaveCustomization.cpp
new file mode 100644
--- /dev/null
+++ b/Engine/Source/Editor/DetailCustomizations/Private/MediaSoundWaveCustomization.cpp
@@ -0,0 +1,13 @@
+// Copyright 1998-2014 Epic Games, Inc. All Rights Reserved.
+
+#include "DetailCustomizationsPrivatePCH.h"
+#include "MediaSoundWaveCustomization.h"
+
+
+/* IDetailCustomization interface
+ *****************************************************************************/
+
+void FMediaSoundWaveCustomization::CustomizeDetails( IDetailLayoutBuilder& DetailBuilder )
+{
+
+}
diff --git a/Engine/Source/Editor/DetailCustomizations/Private/MediaSoundWaveCustomization.h b/Engine/Source/Editor/DetailCustomizations/Private/MediaSoundWaveCustomization.h
new file mode 100644
--- /dev/null
+++ b/Engine/Source/Editor/DetailCustomizations/Private/MediaSoundWaveCustomization.h
@@ -0,0 +1,32 @@
+// Copyright 1998-2014 Epic Games, Inc. All Rights Reserved.
+
+#pragma once
+
+#include "Media.h"
+#include "MediaSoundWave.h"
+
+
+/**
+ * Implements a details view customization for the UMediaSoundWave class.
+ */
+class FMediaSoundWaveCustomization
+	: public IDetailCustomization
+{
+public:
+
+	// IDetailCustomization interface
+
+	virtual void CustomizeDetails( IDetailLayoutBuilder& DetailBuilder ) override;
+
+public:
+
+	/**
+	 * Creates an instance of this class.
+	 *
+	 * @return The new instance.
+	 */
+	static TSharedRef<IDetailCustomization> MakeInstance( )
+	{
+		return MakeShareable(new FMediaSoundWaveCustomization());
+	}
+};
diff --git a/Engine/Source/Editor/DetailCustomizations/Private/MediaTextureCustomization.cpp b/Engine/Source/Editor/DetailCustomizations/Private/MediaTextureCustomization.cpp
new file mode 100644
--- /dev/null
+++ b/Engine/Source/Editor/DetailCustomizations/Private/MediaTextureCustomization.cpp
@@ -0,0 +1,153 @@
+// Copyright 1998-2014 Epic Games, Inc. All Rights Reserved.
+
+#include "DetailCustomizationsPrivatePCH.h"
+#include "MediaTextureCustomization.h"
+
+
+#define LOCTEXT_NAMESPACE "FMediaTextureCustomization"
+
+
+/* IDetailCustomization interface
+ *****************************************************************************/
+
+void FMediaTextureCustomization::CustomizeDetails( IDetailLayoutBuilder& DetailBuilder )
+{
+	DetailBuilder.GetObjectsBeingCustomized(CustomizedMediaTextures);
+
+	IDetailCategoryBuilder& MediaAssetCategory = DetailBuilder.EditCategory(TEXT("MediaAsset"));
+	
+	MediaAssetProperty = DetailBuilder.GetProperty("MediaAsset");
+	VideoTrackIndexProperty = DetailBuilder.GetProperty("VideoTrackIndex");
+
+	// customize video track index	
+	IDetailPropertyRow& VideoTrackIndexRow = MediaAssetCategory.AddProperty(VideoTrackIndexProperty);
+
+	VideoTrackIndexRow.DisplayName(TEXT("Video Track"));
+	VideoTrackIndexRow.CustomWidget()
+		.NameContent()
+		[
+			VideoTrackIndexProperty->CreatePropertyNameWidget()
+		]
+		.ValueContent()
+		.MaxDesiredWidth(0)
+		[
+			SNew(SComboButton)
+				.OnGetMenuContent(this, &FMediaTextureCustomization::HandleVideoTrackComboButtonMenuContent)
+ 				.ContentPadding(FMargin( 2.0f, 2.0f ))
+				.ButtonContent()
+				[
+					SNew(STextBlock) 
+						.Text(this, &FMediaTextureCustomization::HandleVideoTrackComboButtonText)
+						.Font(IDetailLayoutBuilder::GetDetailFont())
+				]
+		];
+}
+
+
+/* FMediaTextureCustomization callbacks
+ *****************************************************************************/
+
+TSharedRef<SWidget> FMediaTextureCustomization::HandleVideoTrackComboButtonMenuContent( ) const
+{
+	// get assigned media asset
+	UObject* MediaAssetObj = nullptr;
+	FPropertyAccess::Result Result = MediaAssetProperty->GetValue(MediaAssetObj);
+
+	if (Result == FPropertyAccess::MultipleValues)
+	{
+		return SNullWidget::NullWidget;
+	}
+
+	UMediaAsset* MediaAsset = Cast<UMediaAsset>(MediaAssetObj);
+
+	if (MediaAsset == nullptr)
+	{
+		return SNullWidget::NullWidget;
+	}
+
+	// get media tracks
+	IMediaPlayerPtr MediaPlayer = MediaAsset->GetMediaPlayer();
+
+	if (!MediaPlayer.IsValid())
+	{
+		return SNullWidget::NullWidget;
+	}
+
+	const TArray<IMediaTrackRef>& MediaTracks = MediaPlayer->GetTracks();
+
+	// populate the menu
+	FMenuBuilder MenuBuilder(true, nullptr);
+
+	for (const IMediaTrackRef& Track : MediaTracks)
+	{
+		if (Track->GetType() == EMediaTrackTypes::Video)
+		{
+			FUIAction Action(FExecuteAction::CreateSP(this, &FMediaTextureCustomization::HandleVideoTrackComboButtonMenuEntryExecute, Track->GetIndex()));
+
+			MenuBuilder.AddMenuEntry(
+				Track->GetDisplayName(),
+				TAttribute<FText>(),
+				FSlateIcon(),
+				Action
+			);
+		}
+	}
+
+	return MenuBuilder.MakeWidget();
+}
+
+
+void FMediaTextureCustomization::HandleVideoTrackComboButtonMenuEntryExecute( uint32 TrackIndex )
+{
+	VideoTrackIndexProperty->SetValue((int32)TrackIndex);
+}
+
+
+FText FMediaTextureCustomization::HandleVideoTrackComboButtonText( ) const
+{
+	// get assigned media asset
+	UObject* MediaAssetObj = nullptr;
+	FPropertyAccess::Result Result = MediaAssetProperty->GetValue(MediaAssetObj);
+
+	if (Result != FPropertyAccess::Success)
+	{
+		return FText::GetEmpty();
+	}
+
+	UMediaAsset* MediaAsset = Cast<UMediaAsset>(MediaAssetObj);
+
+	if (MediaAsset == nullptr)
+	{
+		return LOCTEXT("NoMediaAssetSelected", "No Media Asset selected");
+	}
+
+	// get selected value
+	int32 VideoTrackIndex = -1;
+	Result = VideoTrackIndexProperty->GetValue(VideoTrackIndex);
+
+	if (Result != FPropertyAccess::Success)
+	{
+		return FText::GetEmpty();
+	}
+
+	// get selected media track
+	IMediaPlayerPtr MediaPlayer = MediaAsset->GetMediaPlayer();
+
+	if (!MediaPlayer.IsValid())
+	{
+		return LOCTEXT("NoMediaLoaded", "No media loaded");
+	}
+
+	IMediaTrackPtr Track = MediaPlayer->GetTrackSafe(VideoTrackIndex, EMediaTrackTypes::Video);
+
+	// generate track name string
+	if (Track.IsValid())
+	{
+		return Track->GetDisplayName();
+	}
+
+	return LOCTEXT("SelectVideoTrack", "Select a Video Track...");
+}
+
+
+#undef LOCTEXT_NAMESPACE
diff --git a/Engine/Source/Editor/DetailCustomizations/Private/MediaTextureCustomization.h b/Engine/Source/Editor/DetailCustomizations/Private/MediaTextureCustomization.h
new file mode 100644
--- /dev/null
+++ b/Engine/Source/Editor/DetailCustomizations/Private/MediaTextureCustomization.h
@@ -0,0 +1,54 @@
+// Copyright 1998-2014 Epic Games, Inc. All Rights Reserved.
+
+#pragma once
+
+#include "Media.h"
+#include "MediaTexture.h"
+
+
+/**
+ * Implements a details view customization for the UMediaTexture class.
+ */
+class FMediaTextureCustomization
+	: public IDetailCustomization
+{
+public:
+
+	// IDetailCustomization interface
+
+	virtual void CustomizeDetails( IDetailLayoutBuilder& DetailBuilder ) override;
+
+public:
+
+	/**
+	 * Creates an instance of this class.
+	 *
+	 * @return The new instance.
+	 */
+	static TSharedRef<IDetailCustomization> MakeInstance( )
+	{
+		return MakeShareable(new FMediaTextureCustomization());
+	}
+
+private:
+
+	// Callback for generating the menu content of the VideoTrack combo box.
+	TSharedRef<SWidget> HandleVideoTrackComboButtonMenuContent( ) const;
+
+	// Callback for selecting a track in the VideoTrack combo box.
+	void HandleVideoTrackComboButtonMenuEntryExecute( uint32 TrackIndex );
+
+	// Callback for getting the text of the VideoTrack combo box.
+	FText HandleVideoTrackComboButtonText( ) const;
+
+private:
+
+	/** The collection of media textures being customized */
+	TArray<TWeakObjectPtr<UObject>> CustomizedMediaTextures;
+
+	/** Pointer to the MediaAsset property handle. */
+	TSharedPtr<IPropertyHandle> MediaAssetProperty;
+
+	/** Pointer to the VideoTrackIndex property handle. */
+	TSharedPtr<IPropertyHandle> VideoTrackIndexProperty;
+};
