Enabling Audio Streaming options behind experimental setting.

diff --git a/Engine/Source/Editor/DetailCustomizations/Private/DetailCustomizations.cpp b/Engine/Source/Editor/DetailCustomizations/Private/DetailCustomizations.cpp
--- a/Engine/Source/Editor/DetailCustomizations/Private/DetailCustomizations.cpp
+++ b/Engine/Source/Editor/DetailCustomizations/Private/DetailCustomizations.cpp
@@ -79,7 +79,7 @@
 #include "MediaAssetCustomization.h"
 #include "MediaSoundWaveCustomization.h"
 #include "MediaTextureCustomization.h"
-
+#include "SoundWaveDetails.h"
 
 IMPLEMENT_MODULE( FDetailCustomizationsModule, DetailCustomizations );
 
@@ -214,6 +214,7 @@ void FDetailCustomizationsModule::RegisterObjectCustomizations()
 	RegisterCustomClassLayout("LandscapeSplineControlPoint", FOnGetDetailCustomizationInstance::CreateStatic(&FLandscapeSplineDetails::MakeInstance));
 	RegisterCustomClassLayout("LandscapeSplineSegment", FOnGetDetailCustomizationInstance::CreateStatic(&FLandscapeSplineDetails::MakeInstance));
 
+	RegisterCustomClassLayout("SoundWave", FOnGetDetailCustomizationInstance::CreateStatic(&FSoundWaveDetails::MakeInstance));
 	RegisterCustomClassLayout("DialogueWave", FOnGetDetailCustomizationInstance::CreateStatic(&FDialogueWaveDetails::MakeInstance));
 	RegisterCustomClassLayout("BodySetup", FOnGetDetailCustomizationInstance::CreateStatic(&FBodySetupDetails::MakeInstance));
 	RegisterCustomClassLayout("PhysicsConstraintTemplate", FOnGetDetailCustomizationInstance::CreateStatic(&FPhysicsConstraintComponentDetails::MakeInstance));
@@ -270,4 +271,4 @@ void FDetailCustomizationsModule::RegisterCustomPropertyTypeLayout(FName Propert
 	static FName PropertyEditor("PropertyEditor");
 	FPropertyEditorModule& PropertyModule = FModuleManager::GetModuleChecked<FPropertyEditorModule>(PropertyEditor);
 	PropertyModule.RegisterCustomPropertyTypeLayout(PropertyTypeName, PropertyTypeLayoutDelegate);
-}
+}
\ No newline at end of file
diff --git a/Engine/Source/Editor/DetailCustomizations/Private/SoundWaveDetails.cpp b/Engine/Source/Editor/DetailCustomizations/Private/SoundWaveDetails.cpp
new file mode 100644
--- /dev/null
+++ b/Engine/Source/Editor/DetailCustomizations/Private/SoundWaveDetails.cpp
@@ -0,0 +1,21 @@
+// Copyright 1998-2014 Epic Games, Inc. All Rights Reserved.
+
+#include "DetailCustomizationsPrivatePCH.h"
+#include "SoundWaveDetails.h"
+//#include "EditorExperimentalSettings.h"
+
+#define LOCTEXT_NAMESPACE "AmbientSoundDetails"
+
+TSharedRef<IDetailCustomization> FSoundWaveDetails::MakeInstance()
+{
+	return MakeShareable(new FSoundWaveDetails);
+}
+
+void FSoundWaveDetails::CustomizeDetails(IDetailLayoutBuilder& DetailBuilder)
+{
+	if (!GetDefault<UEditorExperimentalSettings>()->bShowAudioStreamingOptions)
+	{
+		DetailBuilder.HideProperty(GET_MEMBER_NAME_CHECKED(USoundWave, bStreaming));
+		DetailBuilder.HideProperty(GET_MEMBER_NAME_CHECKED(USoundWave, StreamingPriority));
+	}
+}
diff --git a/Engine/Source/Editor/DetailCustomizations/Private/SoundWaveDetails.h b/Engine/Source/Editor/DetailCustomizations/Private/SoundWaveDetails.h
new file mode 100644
--- /dev/null
+++ b/Engine/Source/Editor/DetailCustomizations/Private/SoundWaveDetails.h
@@ -0,0 +1,14 @@
+// Copyright 1998-2014 Epic Games, Inc. All Rights Reserved.
+
+#pragma once
+
+class FSoundWaveDetails : public IDetailCustomization
+{
+public:
+	/** Makes a new instance of this detail layout class for a specific detail view requesting it */
+	static TSharedRef<IDetailCustomization> MakeInstance();
+
+private:
+	/** IDetailCustomization interface */
+	virtual void CustomizeDetails(IDetailLayoutBuilder& DetailBuilder) override;
+};
\ No newline at end of file
diff --git a/Engine/Source/Editor/UnrealEd/Classes/Settings/EditorExperimentalSettings.h b/Engine/Source/Editor/UnrealEd/Classes/Settings/EditorExperimentalSettings.h
--- a/Engine/Source/Editor/UnrealEd/Classes/Settings/EditorExperimentalSettings.h
+++ b/Engine/Source/Editor/UnrealEd/Classes/Settings/EditorExperimentalSettings.h
@@ -62,6 +62,10 @@ public:
 	UPROPERTY(EditAnywhere, config, Category=Blueprints, meta=(DisplayName="Allow splitting of struct pins"))
 	bool bAllowSplitStructPins;
 
+	/** Whether to show Audio Streaming options for SoundWaves (disabling will not stop all audio streaming) */
+	UPROPERTY(EditAnywhere, config, Category=Audio)
+	bool bShowAudioStreamingOptions;
+
 	/**
 	 * Returns an event delegate that is executed when a setting has changed.
 	 *
diff --git a/Engine/Source/Runtime/Engine/Classes/Sound/SoundWave.h b/Engine/Source/Runtime/Engine/Classes/Sound/SoundWave.h
--- a/Engine/Source/Runtime/Engine/Classes/Sound/SoundWave.h
+++ b/Engine/Source/Runtime/Engine/Classes/Sound/SoundWave.h
@@ -120,9 +120,14 @@ class ENGINE_API USoundWave : public USoundBase
 	UPROPERTY(EditAnywhere, Category=SoundWave )
 	uint32 bLooping:1;
 
-	UPROPERTY()
+	/** Whether this sound can be streamed to avoid increased memory usage */
+	UPROPERTY(EditAnywhere, Category=Streaming)
 	uint32 bStreaming:1;
 
+	/** Priority of this sound when streaming (lower priority streams may not always play) */
+	UPROPERTY(EditAnywhere, Category=Streaming, meta=(ClampMin=0))
+	int32 StreamingPriority;
+
 	/** Set to true for programmatically-generated, streamed audio. */
 	uint32 bProcedural:1;
 
