Fix build errors.  LOCTEXT_NAMESPACE must be #undef'd at the end of files.

diff --git a/Engine/Source/Editor/DetailCustomizations/Private/SoundWaveDetails.cpp b/Engine/Source/Editor/DetailCustomizations/Private/SoundWaveDetails.cpp
--- a/Engine/Source/Editor/DetailCustomizations/Private/SoundWaveDetails.cpp
+++ b/Engine/Source/Editor/DetailCustomizations/Private/SoundWaveDetails.cpp
@@ -19,3 +19,5 @@ void FSoundWaveDetails::CustomizeDetails(IDetailLayoutBuilder& DetailBuilder)
 		DetailBuilder.HideProperty(GET_MEMBER_NAME_CHECKED(USoundWave, StreamingPriority));
 	}
 }
+
+#undef LOCTEXT_NAMESPACE
