#NUX : Added Padding

diff --git a/Engine/Source/Editor/DetailCustomizations/Private/PhysicsSettingsDetails.cpp b/Engine/Source/Editor/DetailCustomizations/Private/PhysicsSettingsDetails.cpp
--- a/Engine/Source/Editor/DetailCustomizations/Private/PhysicsSettingsDetails.cpp
+++ b/Engine/Source/Editor/DetailCustomizations/Private/PhysicsSettingsDetails.cpp
@@ -35,6 +35,7 @@ TSharedRef<SWidget> SPhysicalSurfaceListItem::GenerateWidgetForColumn(const FNam
 	{
 		return	SNew(SBox)
 			.HeightOverride(20)
+			.Padding(FMargin(5, 2))
 			.VAlign(VAlign_Center)
 			[
 				SNew(STextBlock)
@@ -46,6 +47,7 @@ TSharedRef<SWidget> SPhysicalSurfaceListItem::GenerateWidgetForColumn(const FNam
 	{
 		return	SNew(SBox)
 			.HeightOverride(20)
+			.Padding(FMargin(5, 2))
 			.VAlign(VAlign_Center)
 			[
 				SAssignNew(NameBox, SEditableTextBox)
