NUX: Fixed collision channel /profile name has to be entered before "accept" button gets greyed out.

diff --git a/Engine/Source/Editor/DetailCustomizations/Private/CollisionProfileDetails.cpp b/Engine/Source/Editor/DetailCustomizations/Private/CollisionProfileDetails.cpp
--- a/Engine/Source/Editor/DetailCustomizations/Private/CollisionProfileDetails.cpp
+++ b/Engine/Source/Editor/DetailCustomizations/Private/CollisionProfileDetails.cpp
@@ -269,6 +269,7 @@ void SChannelEditDialog::OnTextChanged(const FText& NewText)
 	else
 	{
 		NameBox->SetError(FText::GetEmpty());
+		NewNameEntered(NewText, ETextCommit::Default);
 	}
 }
 
@@ -702,6 +703,7 @@ void SProfileEditDialog::OnTextChanged(const FText& NewText)
 	else
 	{
 		NameBox->SetError(TEXT(""));
+		NewNameEntered(NewText, ETextCommit::Default);
 	}
 }
 
