NUX: Collision profile/Physics surface type UI cleanup

diff --git a/Engine/Source/Editor/DetailCustomizations/Private/CollisionProfileDetails.cpp b/Engine/Source/Editor/DetailCustomizations/Private/CollisionProfileDetails.cpp
--- a/Engine/Source/Editor/DetailCustomizations/Private/CollisionProfileDetails.cpp
+++ b/Engine/Source/Editor/DetailCustomizations/Private/CollisionProfileDetails.cpp
@@ -18,10 +18,10 @@ DECLARE_DELEGATE_RetVal_TwoParams(bool, FOnValidateProfile, const FCollisionResp
 #define COLLIISION_COLUMN_WIDTH		50
 
 #define PROFILE_WINDOW_WIDTH		300
-#define PROFILE_WINDOW_HEIGHT		500
+#define PROFILE_WINDOW_HEIGHT		540
 
 #define CHANNEL_WINDOW_WIDTH		200
-#define CHANNEL_WINDOW_HEIGHT		100
+#define CHANNEL_WINDOW_HEIGHT		93
 
 #define RowWidth_Customization 50
 //====================================================================================
@@ -105,110 +105,98 @@ void SChannelEditDialog::Construct(const FArguments& InArgs)
 
 	this->ChildSlot
 	[
-		SNew(SBorder)
-//		.BorderImage(FEditorStyle::GetBrush("Menu.Background"))
-		.Padding(1.f)
-		[
-			SNew(SVerticalBox)
+		SNew(SVerticalBox)
 
-			+SVerticalBox::Slot()
-			.Padding(2.f)
-			.AutoHeight()
+		// channel name
+		+ SVerticalBox::Slot()
+		.FillHeight(1)
+		.VAlign(VAlign_Center)
+		.Padding(3.f, 1.f)
+		[
+			SNew(SHorizontalBox)
+			+SHorizontalBox::Slot()
+			.AutoWidth()
+			.HAlign(HAlign_Left)
 			[
-				SNew(STextBlock)
-				.ColorAndOpacity(FSlateColor::UseSubduedForeground())
+				SNew(SBox)
+				.WidthOverride(100)
+				[
+					SNew(STextBlock)
+					.Text(LOCTEXT("SChannelEditDialog_Name", "Name"))
+					.Font(IDetailLayoutBuilder::GetDetailFontBold())
+				]
+			]
+			+SHorizontalBox::Slot()
+			.FillWidth(1)
+			.HAlign(HAlign_Left)
+			[
+				SAssignNew(NameBox, SEditableTextBox)
+				.MinDesiredWidth(64.0f)
+				.Text(this, &SChannelEditDialog::GetName)
 				.Font(IDetailLayoutBuilder::GetDetailFont())
-				.Text(LOCTEXT("Dialog_TextTypeHelp", "Type Enter to apply text change. "))
+				.OnTextCommitted(this, &SChannelEditDialog::NewNameEntered)
+				.OnTextChanged(this, &SChannelEditDialog::OnTextChanged)
 			]
-			// channel name
-			+ SVerticalBox::Slot()
-			.FillHeight(1)
-			.VAlign(VAlign_Center)
-			.Padding(2.f)
+		]
+
+		// default response
+		+ SVerticalBox::Slot()
+		.FillHeight(1)
+		.Padding(3.f, 1.f)
+		.VAlign(VAlign_Center)
+		[
+			SNew(SHorizontalBox)
+			+SHorizontalBox::Slot()
+			.AutoWidth()
+			.HAlign(HAlign_Left)
 			[
-				SNew(SHorizontalBox)
-				+SHorizontalBox::Slot()
-				.AutoWidth()
-				.HAlign(HAlign_Left)
+				SNew(SBox)
+				.WidthOverride(100)
 				[
-					SNew(SBox)
-					.WidthOverride(100)
-					[
-						SNew(STextBlock)
-						.Text(LOCTEXT("SChannelEditDialog_Name", "Name"))
-						.Font(IDetailLayoutBuilder::GetDetailFontBold())
-					]
+					SNew(STextBlock)
+					.Text(LOCTEXT("SChannelEditDialog_DefaultResponse", "Default Response"))
+					.Font(IDetailLayoutBuilder::GetDetailFontBold())
 				]
-				+SHorizontalBox::Slot()
-				.FillWidth(1)
-				.HAlign(HAlign_Left)
+			]
+			+SHorizontalBox::Slot()
+			.FillWidth(1)
+			.HAlign(HAlign_Left)
+			[
+				SAssignNew(ResponseComboBox, SComboBox< TSharedPtr<FString> >)
+				.ContentPadding(FMargin(6.0f, 2.0f))
+				.OptionsSource(&ResponseComboBoxString)
+				.OnGenerateWidget(this, &SChannelEditDialog::HandleResponseComboBoxGenerateWidget)
+				.OnSelectionChanged(this, &SChannelEditDialog::HandleResponseComboBoxSelectionChanged)
 				[
-					SAssignNew(NameBox, SEditableTextBox)
-					.MinDesiredWidth(64.0f)
-					.Text(this, &SChannelEditDialog::GetName)
+					SNew(STextBlock)
+					.Text(this, &SChannelEditDialog::HandleResponseComboBoxContentText)
 					.Font(IDetailLayoutBuilder::GetDetailFont())
-					.OnTextCommitted(this, &SChannelEditDialog::NewNameEntered)
-					.OnTextChanged(this, &SChannelEditDialog::OnTextChanged)
 				]
 			]
+		]
 
-			// default response
-			+ SVerticalBox::Slot()
-			.FillHeight(1)
-			.Padding(1.f)
-			.VAlign(VAlign_Center)
+		// accept or cancel button
+		+ SVerticalBox::Slot()
+		.AutoHeight()
+		.Padding(0.f, 3.f)
+		[
+			SNew(SHorizontalBox)
+			+SHorizontalBox::Slot()
+			.HAlign(HAlign_Center)
 			[
-				SNew(SHorizontalBox)
-				+SHorizontalBox::Slot()
-				.AutoWidth()
-				.HAlign(HAlign_Left)
-				[
-					SNew(SBox)
-					.WidthOverride(100)
-					[
-						SNew(STextBlock)
-						.Text(LOCTEXT("SChannelEditDialog_DefaultResponse", "Default Response"))
-						.Font(IDetailLayoutBuilder::GetDetailFontBold())
-					]
-				]
-				+SHorizontalBox::Slot()
-				.FillWidth(1)
-				.HAlign(HAlign_Left)
-				[
-					SAssignNew(ResponseComboBox, SComboBox< TSharedPtr<FString> >)
-					.ContentPadding(FMargin(6.0f, 2.0f))
-					.OptionsSource(&ResponseComboBoxString)
-					.OnGenerateWidget(this, &SChannelEditDialog::HandleResponseComboBoxGenerateWidget)
-					.OnSelectionChanged(this, &SChannelEditDialog::HandleResponseComboBoxSelectionChanged)
-					[
-						SNew(STextBlock)
-						.Text(this, &SChannelEditDialog::HandleResponseComboBoxContentText)
-						.Font(IDetailLayoutBuilder::GetDetailFont())
-					]
-				]
+				SNew(SButton)
+				.ContentPadding(FEditorStyle::GetMargin("StandardDialog.ContentPadding"))
+				.Text(LOCTEXT("SChannelEditDialog_Accept", "Accept"))
+				.OnClicked(this, &SChannelEditDialog::OnAccept)
+				.IsEnabled(this, &SChannelEditDialog::IsAcceptAvailable)
 			]
-
-			// accept or cancel button
-			+ SVerticalBox::Slot()
-			.FillHeight(1)
-			.Padding(1.f)
+			+SHorizontalBox::Slot()
+			.HAlign(HAlign_Center)
 			[
-				SNew(SHorizontalBox)
-				+SHorizontalBox::Slot()
-				[
-					SNew(SButton)
-					.ContentPadding(FEditorStyle::GetMargin("StandardDialog.ContentPadding"))
-					.Text(LOCTEXT("SChannelEditDialog_Accept", "Accept"))
-					.OnClicked(this, &SChannelEditDialog::OnAccept)
-					.IsEnabled(this, &SChannelEditDialog::IsAcceptAvailable)
-				]
-				+SHorizontalBox::Slot()
-				[
-					SNew(SButton)
-					.ContentPadding(FEditorStyle::GetMargin("StandardDialog.ContentPadding"))
-					.Text(LOCTEXT("SChannelEditDialog_Cancel", "Cancel"))
-					.OnClicked(this, &SChannelEditDialog::OnCancel)
-				]
+				SNew(SButton)
+				.ContentPadding(FEditorStyle::GetMargin("StandardDialog.ContentPadding"))
+				.Text(LOCTEXT("SChannelEditDialog_Cancel", "Cancel"))
+				.OnClicked(this, &SChannelEditDialog::OnCancel)
 			]
 		]
 	];
@@ -442,184 +430,172 @@ void SProfileEditDialog::Construct(const FArguments& InArgs)
 
 	this->ChildSlot
 	[
-		SNew(SBorder)
-//		.BorderImage(FEditorStyle::GetBrush("Menu.Background"))
-		.Padding(2.f)
-		[
-			SNew(SVerticalBox)
+		SNew(SVerticalBox)
 
-			+SVerticalBox::Slot()
-			.Padding(2.f)
-			.AutoHeight()
+		// Profile name
+		+ SVerticalBox::Slot()
+		.AutoHeight()
+		.VAlign(VAlign_Center)
+		.Padding(3.f)
+		[
+			SNew(SHorizontalBox)
+			+SHorizontalBox::Slot()
+			.AutoWidth()
+			.HAlign(HAlign_Left)
 			[
-				SNew(STextBlock)
-				.ColorAndOpacity(FSlateColor::UseSubduedForeground())
+				SNew(SBox)
+				.WidthOverride(100)
+				[
+					SNew(STextBlock)
+					.Text(LOCTEXT("SProfileEditDialog_Name", "Name"))
+					.Font(IDetailLayoutBuilder::GetDetailFontBold())
+				]
+			]
+			+SHorizontalBox::Slot()
+			.FillWidth(1)
+			.HAlign(HAlign_Left)
+			[
+				SAssignNew(NameBox, SEditableTextBox)
+				.MinDesiredWidth(64.0f)
+				.Text(this, &SProfileEditDialog::GetName)
 				.Font(IDetailLayoutBuilder::GetDetailFont())
-				.Text(LOCTEXT("Dialog_TextTypeHelp", "Type Enter to apply text change. ").ToString())
+				.IsEnabled(SProfileEditDialog::CanModify())
+				.OnTextCommitted(this, &SProfileEditDialog::NewNameEntered)
+				.OnTextChanged(this, &SProfileEditDialog::OnTextChanged)
 			]
+		]
 
-			// Profile name
-			+ SVerticalBox::Slot()
-			.AutoHeight()
-			.VAlign(VAlign_Center)
-			.Padding(3.f)
+		// default CollisionEnabled
+		+ SVerticalBox::Slot()
+		.AutoHeight()
+		.VAlign(VAlign_Center)
+		.Padding(3.f)
+		[
+			SNew(SHorizontalBox)
+			+SHorizontalBox::Slot()
+			.AutoWidth()
+			.HAlign(HAlign_Left)
 			[
-				SNew(SHorizontalBox)
-				+SHorizontalBox::Slot()
-				.AutoWidth()
-				.HAlign(HAlign_Left)
+				SNew(SBox)
+				.WidthOverride(100)
 				[
-					SNew(SBox)
-					.WidthOverride(100)
-					[
-						SNew(STextBlock)
-						.Text(LOCTEXT("SProfileEditDialog_Name", "Name"))
-						.Font(IDetailLayoutBuilder::GetDetailFontBold())
-					]
+					SNew(STextBlock)
+					.Text(LOCTEXT("SProfileEditDialog_CollisionEnabled", "CollisionEnabled"))
+					.Font(IDetailLayoutBuilder::GetDetailFontBold())
 				]
-				+SHorizontalBox::Slot()
-				.FillWidth(1)
-				.HAlign(HAlign_Left)
+			]
+			+SHorizontalBox::Slot()
+			.FillWidth(1)
+			.HAlign(HAlign_Left)
+			[
+				SAssignNew(CollisionEnabledComboBox, SComboBox< TSharedPtr<FString> >)
+				.ContentPadding(FMargin(6.0f, 2.0f))
+				.OptionsSource(&CollisionEnabledComboBoxString)
+				.OnGenerateWidget(this, &SProfileEditDialog::HandleCollisionEnabledComboBoxGenerateWidget)
+				.OnSelectionChanged(this, &SProfileEditDialog::HandleCollisionEnabledComboBoxSelectionChanged)
 				[
-					SAssignNew(NameBox, SEditableTextBox)
-					.MinDesiredWidth(64.0f)
-					.Text(this, &SProfileEditDialog::GetName)
+					SNew(STextBlock)
+					.Text(this, &SProfileEditDialog::HandleCollisionEnabledComboBoxContentText)
 					.Font(IDetailLayoutBuilder::GetDetailFont())
-					.IsEnabled(SProfileEditDialog::CanModify())
-					.OnTextCommitted(this, &SProfileEditDialog::NewNameEntered)
-					.OnTextChanged(this, &SProfileEditDialog::OnTextChanged)
 				]
 			]
+		]
 
-			// default CollisionEnabled
-			+ SVerticalBox::Slot()
-			.AutoHeight()
-			.VAlign(VAlign_Center)
-			.Padding(3.f)
+		// default ObjectType
+		+ SVerticalBox::Slot()
+		.AutoHeight()
+		.VAlign(VAlign_Center)
+		.Padding(3.f)
+		[
+			SNew(SHorizontalBox)
+			+SHorizontalBox::Slot()
+			.AutoWidth()
+			.HAlign(HAlign_Left)
 			[
-				SNew(SHorizontalBox)
-				+SHorizontalBox::Slot()
-				.AutoWidth()
-				.HAlign(HAlign_Left)
-				[
-					SNew(SBox)
-					.WidthOverride(100)
-					[
-						SNew(STextBlock)
-						.Text(LOCTEXT("SProfileEditDialog_CollisionEnabled", "CollisionEnabled"))
-						.Font(IDetailLayoutBuilder::GetDetailFontBold())
-					]
-				]
-				+SHorizontalBox::Slot()
-				.FillWidth(1)
-				.HAlign(HAlign_Left)
+				SNew(SBox)
+				.WidthOverride(100)
 				[
-					SAssignNew(CollisionEnabledComboBox, SComboBox< TSharedPtr<FString> >)
-					.ContentPadding(FMargin(6.0f, 2.0f))
-					.OptionsSource(&CollisionEnabledComboBoxString)
-					.OnGenerateWidget(this, &SProfileEditDialog::HandleCollisionEnabledComboBoxGenerateWidget)
-					.OnSelectionChanged(this, &SProfileEditDialog::HandleCollisionEnabledComboBoxSelectionChanged)
-					[
-						SNew(STextBlock)
-						.Text(this, &SProfileEditDialog::HandleCollisionEnabledComboBoxContentText)
-						.Font(IDetailLayoutBuilder::GetDetailFont())
-					]
+					SNew(STextBlock)
+					.Text(LOCTEXT("SProfileEditDialog_ObjectType", "ObjectType"))
+					.Font(IDetailLayoutBuilder::GetDetailFontBold())
 				]
 			]
-
-			// default ObjectType
-			+ SVerticalBox::Slot()
-			.AutoHeight()
-			.VAlign(VAlign_Center)
-			.Padding(3.f)
+			+SHorizontalBox::Slot()
+			.FillWidth(1)
+			.HAlign(HAlign_Left)
 			[
-				SNew(SHorizontalBox)
-				+SHorizontalBox::Slot()
-				.AutoWidth()
-				.HAlign(HAlign_Left)
+				SAssignNew(ObjectTypeComboBox, SComboBox< TSharedPtr<FString> >)
+				.ContentPadding(FMargin(6.0f, 2.0f))
+				.OptionsSource(&ObjectTypeComboBoxString)
+				.OnGenerateWidget(this, &SProfileEditDialog::HandleObjectTypeComboBoxGenerateWidget)
+				.OnSelectionChanged(this, &SProfileEditDialog::HandleObjectTypeComboBoxSelectionChanged)
 				[
-					SNew(SBox)
-					.WidthOverride(100)
-					[
-						SNew(STextBlock)
-						.Text(LOCTEXT("SProfileEditDialog_ObjectType", "ObjectType"))
-						.Font(IDetailLayoutBuilder::GetDetailFontBold())
-					]
-				]
-				+SHorizontalBox::Slot()
-				.FillWidth(1)
-				.HAlign(HAlign_Left)
-				[
-					SAssignNew(ObjectTypeComboBox, SComboBox< TSharedPtr<FString> >)
-					.ContentPadding(FMargin(6.0f, 2.0f))
-					.OptionsSource(&ObjectTypeComboBoxString)
-					.OnGenerateWidget(this, &SProfileEditDialog::HandleObjectTypeComboBoxGenerateWidget)
-					.OnSelectionChanged(this, &SProfileEditDialog::HandleObjectTypeComboBoxSelectionChanged)
-					[
-						SNew(STextBlock)
-						.Text(this, &SProfileEditDialog::HandleObjectTypeComboBoxContentText)
-						.Font(IDetailLayoutBuilder::GetDetailFont())
-					]
+					SNew(STextBlock)
+					.Text(this, &SProfileEditDialog::HandleObjectTypeComboBoxContentText)
+					.Font(IDetailLayoutBuilder::GetDetailFont())
 				]
 			]
+		]
 
-			// Profile Description
-			+ SVerticalBox::Slot()
-			.AutoHeight()
-			.VAlign(VAlign_Center)
-			.Padding(3.f)
+		// Profile Description
+		+ SVerticalBox::Slot()
+		.AutoHeight()
+		.VAlign(VAlign_Center)
+		.Padding(3.f)
+		[
+			SNew(SHorizontalBox)
+			+SHorizontalBox::Slot()
+			.AutoWidth()
+			.HAlign(HAlign_Left)
 			[
-				SNew(SHorizontalBox)
-				+SHorizontalBox::Slot()
-				.AutoWidth()
-				.HAlign(HAlign_Left)
-				[
-					SNew(SBox)
-					.WidthOverride(100)
-					[
-						SNew(STextBlock)
-						.Text(LOCTEXT("SProfileEditDialog_Description", "Description"))
-						.Font(IDetailLayoutBuilder::GetDetailFontBold())
-					]
-				]
-				+SHorizontalBox::Slot()
-				.FillWidth(1)
-				.HAlign(HAlign_Left)
+				SNew(SBox)
+				.WidthOverride(100)
 				[
-					SNew(SEditableTextBox)
-					.MinDesiredWidth(128.0f)
-					.Text(this, &SProfileEditDialog::GetDescription)
-					.Font(IDetailLayoutBuilder::GetDetailFont())
-					.OnTextCommitted(this, &SProfileEditDialog::NewDescriptionEntered)
+					SNew(STextBlock)
+					.Text(LOCTEXT("SProfileEditDialog_Description", "Description"))
+					.Font(IDetailLayoutBuilder::GetDetailFontBold())
 				]
 			]
-
-			+SVerticalBox::Slot()
-			.FillHeight(1)
+			+SHorizontalBox::Slot()
+			.FillWidth(1)
+			.HAlign(HAlign_Left)
 			[
-				SAssignNew(SCollisionPanel, SVerticalBox)
+				SNew(SEditableTextBox)
+				.MinDesiredWidth(128.0f)
+				.Text(this, &SProfileEditDialog::GetDescription)
+				.Font(IDetailLayoutBuilder::GetDetailFont())
+				.OnTextCommitted(this, &SProfileEditDialog::NewDescriptionEntered)
 			]
+		]
 
-			// accept or cancel button
-			+ SVerticalBox::Slot()
-			.AutoHeight()
+		+SVerticalBox::Slot()
+		.FillHeight(1)
+		[
+			SAssignNew(SCollisionPanel, SVerticalBox)
+		]
+
+		// accept or cancel button
+		+ SVerticalBox::Slot()
+		.AutoHeight()
+		.Padding(1,3)
+		[
+			SNew(SHorizontalBox)
+			+SHorizontalBox::Slot()
+			.HAlign(HAlign_Center)
 			[
-				SNew(SHorizontalBox)
-				+SHorizontalBox::Slot()
-				[
-					SNew(SButton)
-					.ContentPadding(FEditorStyle::GetMargin("StandardDialog.ContentPadding"))
-					.Text(LOCTEXT("SProfileEditDialog_Accept", "Accept").ToString())
-					.OnClicked(this, &SProfileEditDialog::OnAccept)
-					.IsEnabled(this, &SProfileEditDialog::IsAcceptAvailable)
-				]
-				+SHorizontalBox::Slot()
-				[
-					SNew(SButton)
-					.ContentPadding(FEditorStyle::GetMargin("StandardDialog.ContentPadding"))
-					.Text(LOCTEXT("SProfileEditDialog_Cancel", "Cancel").ToString())
-					.OnClicked(this, &SProfileEditDialog::OnCancel)
-				]
+				SNew(SButton)
+				.ContentPadding(FEditorStyle::GetMargin("StandardDialog.ContentPadding"))
+				.Text(LOCTEXT("SProfileEditDialog_Accept", "Accept").ToString())
+				.OnClicked(this, &SProfileEditDialog::OnAccept)
+				.IsEnabled(this, &SProfileEditDialog::IsAcceptAvailable)
+			]
+			+SHorizontalBox::Slot()
+			.HAlign(HAlign_Center)
+			[
+				SNew(SButton)
+				.ContentPadding(FEditorStyle::GetMargin("StandardDialog.ContentPadding"))
+				.Text(LOCTEXT("SProfileEditDialog_Cancel", "Cancel").ToString())
+				.OnClicked(this, &SProfileEditDialog::OnCancel)
 			]
 		]
 	];
@@ -1182,6 +1158,7 @@ TSharedRef<SWidget> SChannelListItem::GenerateWidgetForColumn(const FName& Colum
 	{
 		return	SNew(SBox)
 			.HeightOverride(20)
+			.Padding(FMargin(3, 0))
 			.VAlign(VAlign_Center)
 			[
 				SNew(STextBlock)
@@ -1193,6 +1170,7 @@ TSharedRef<SWidget> SChannelListItem::GenerateWidgetForColumn(const FName& Colum
 	{
 		return	SNew(SBox)
 			.HeightOverride(20)
+			.Padding(FMargin(3, 0))
 			.VAlign(VAlign_Center)
 			[
 				SNew(STextBlock)
@@ -1374,54 +1352,52 @@ void FCollisionProfileDetails::CustomizeDetails( IDetailLayoutBuilder& DetailBui
 		.Padding(5)
 		.AutoHeight()
 		[
-			SNew(SBorder)
-			.BorderImage(FEditorStyle::GetBrush("Menu.Background"))
-			.Padding(1.f)
+			SNew(SHorizontalBox)
+			+SHorizontalBox::Slot()
+			.Padding(2, 10)
+			.FillWidth(1)
 			[
-				SNew(SHorizontalBox)
-				+SHorizontalBox::Slot()
-				.Padding(10, 10)
-				.FillWidth(1)
-				[
-					SNew(STextBlock)
-					.ColorAndOpacity(FSlateColor::UseSubduedForeground())
-					.Font(IDetailLayoutBuilder::GetDetailFont())
-					.ToolTip(ObjectChannelTooltip)
-					.Text(LOCTEXT("ObjectChannel_Menu_Description", "You can have up to 18 custom channels including object and trace channels. This is list of object type for your project. \nIf you delete the object type that has been used by game, it will go back to WorldStatic."))
-				]
+				SNew(STextBlock)
+				.AutoWrapText(true)
+				.Font(IDetailLayoutBuilder::GetDetailFont())
+				.ToolTip(ObjectChannelTooltip)
+				.Text(LOCTEXT("ObjectChannel_Menu_Description", "You can have up to 18 custom channels including object and trace channels. This is list of object type for your project. If you delete the object type that has been used by game, it will go back to WorldStatic."))
+			]
 
-				+SHorizontalBox::Slot()
-				.Padding(2, 10)
-				.AutoWidth()
-				[
-					SNew(SButton)
-					.ContentPadding(FEditorStyle::GetMargin("StandardDialog.ContentPadding"))
-					.Text(LOCTEXT("ChannelMenu_NewObject", "New Object Channel...").ToString())
-					.OnClicked(this, &FCollisionProfileDetails::OnNewChannel, false)
-					.IsEnabled(this, &FCollisionProfileDetails::IsNewChannelAvailable)
-				]
+			+SHorizontalBox::Slot()
+			.Padding(2, 10)
+			.AutoWidth()
+			.VAlign(VAlign_Center)
+			[
+				SNew(SButton)
+				.ContentPadding(FEditorStyle::GetMargin("StandardDialog.ContentPadding"))
+				.Text(LOCTEXT("ChannelMenu_NewObject", "New Object Channel...").ToString())
+				.OnClicked(this, &FCollisionProfileDetails::OnNewChannel, false)
+				.IsEnabled(this, &FCollisionProfileDetails::IsNewChannelAvailable)
+			]
 
-				+SHorizontalBox::Slot()
-				.Padding(2, 10)
-				.AutoWidth()
-				[
-					SNew(SButton)
-					.ContentPadding(FEditorStyle::GetMargin("StandardDialog.ContentPadding"))
-					.Text(LOCTEXT("ChannelMenu_Edit", "Edit...").ToString())
-					.OnClicked(this, &FCollisionProfileDetails::OnEditChannel, false)
-					.IsEnabled(this, &FCollisionProfileDetails::IsAnyChannelSelected, false)
-				]
+			+SHorizontalBox::Slot()
+			.Padding(2, 10)
+			.AutoWidth()
+			.VAlign(VAlign_Center)
+			[
+				SNew(SButton)
+				.ContentPadding(FEditorStyle::GetMargin("StandardDialog.ContentPadding"))
+				.Text(LOCTEXT("ChannelMenu_Edit", "Edit...").ToString())
+				.OnClicked(this, &FCollisionProfileDetails::OnEditChannel, false)
+				.IsEnabled(this, &FCollisionProfileDetails::IsAnyChannelSelected, false)
+			]
 
-				+SHorizontalBox::Slot()
-				.Padding(2, 10)
-				.AutoWidth()
-				[
-					SNew(SButton)
-					.ContentPadding(FEditorStyle::GetMargin("StandardDialog.ContentPadding"))
-					.Text(LOCTEXT("ChannelMenu_Delete", "Delete...").ToString())
-					.OnClicked(this, &FCollisionProfileDetails::OnDeleteChannel, false)
-					.IsEnabled(this, &FCollisionProfileDetails::IsAnyChannelSelected, false)
-				]
+			+SHorizontalBox::Slot()
+			.Padding(2, 10)
+			.AutoWidth()
+			.VAlign(VAlign_Center)
+			[
+				SNew(SButton)
+				.ContentPadding(FEditorStyle::GetMargin("StandardDialog.ContentPadding"))
+				.Text(LOCTEXT("ChannelMenu_Delete", "Delete...").ToString())
+				.OnClicked(this, &FCollisionProfileDetails::OnDeleteChannel, false)
+				.IsEnabled(this, &FCollisionProfileDetails::IsAnyChannelSelected, false)
 			]
 		]
 
@@ -1429,44 +1405,36 @@ void FCollisionProfileDetails::CustomizeDetails( IDetailLayoutBuilder& DetailBui
 		.Padding(5)
 		.FillHeight(1)
 		[
-			SNew(SBorder)
-			.Padding(2.0)
-			.Content()
-			[
-				SAssignNew(ObjectChannelListView, SChannelListView)
-				.ItemHeight(15.f)
-				.ListItemsSource(&ObjectChannelList)
-				.OnGenerateRow(this, &FCollisionProfileDetails::HandleGenerateChannelWidget)
-				.OnMouseButtonDoubleClick(this, &FCollisionProfileDetails::OnObjectChannelListItemDoubleClicked)
-				.SelectionMode(ESelectionMode::Single)
-
-				.HeaderRow(
-					SNew(SHeaderRow)
-					// Name
-					+ SHeaderRow::Column("Name")
-					.HAlignCell(HAlign_Left)
-					.FillWidth(1)
-					[
-						SNew(SHorizontalBox)
-						+SHorizontalBox::Slot()
-						.FillWidth(1)
-						[
-							SNew(STextBlock)
-							.Text(LOCTEXT("ChannelListHeader_Name", "Name"))
-							.Font(IDetailLayoutBuilder::GetDetailFont())
-						]
-					]
-					// Default Response
-					+ SHeaderRow::Column("DefaultResponse")
-					.HAlignCell(HAlign_Left)
-					.FillWidth(1)
-					[
-						SNew(STextBlock)
-						.Text(LOCTEXT("ChannelListHeader_DefaultResponse", "Default Response"))
-						.Font(IDetailLayoutBuilder::GetDetailFont())
-					]
-				)
-			]
+			SAssignNew(ObjectChannelListView, SChannelListView)
+			.ItemHeight(15.f)
+			.ListItemsSource(&ObjectChannelList)
+			.OnGenerateRow(this, &FCollisionProfileDetails::HandleGenerateChannelWidget)
+			.OnMouseButtonDoubleClick(this, &FCollisionProfileDetails::OnObjectChannelListItemDoubleClicked)
+			.SelectionMode(ESelectionMode::Single)
+
+			.HeaderRow(
+				SNew(SHeaderRow)
+				// Name
+				+ SHeaderRow::Column("Name")
+				.HAlignCell(HAlign_Left)
+				.FillWidth(1)
+				.HeaderContentPadding(FMargin(0, 3))
+				[
+					SNew(STextBlock)
+					.Text(LOCTEXT("ChannelListHeader_Name", "Name"))
+					.Font(IDetailLayoutBuilder::GetDetailFont())
+				]
+				// Default Response
+				+ SHeaderRow::Column("DefaultResponse")
+				.HAlignCell(HAlign_Left)
+				.FillWidth(1)
+				.HeaderContentPadding(FMargin(0, 3))
+				[
+					SNew(STextBlock)
+					.Text(LOCTEXT("ChannelListHeader_DefaultResponse", "Default Response"))
+					.Font(IDetailLayoutBuilder::GetDetailFont())
+				]
+			)
 		]
 	];
 	
@@ -1478,54 +1446,52 @@ void FCollisionProfileDetails::CustomizeDetails( IDetailLayoutBuilder& DetailBui
 		.Padding(5)
 		.AutoHeight()
 		[
-			SNew(SBorder)
-			.BorderImage(FEditorStyle::GetBrush("Menu.Background"))
-			.Padding(1.f)
+			SNew(SHorizontalBox)
+			+SHorizontalBox::Slot()
+			.Padding(2, 10)
+			.FillWidth(1)
 			[
-				SNew(SHorizontalBox)
-				+SHorizontalBox::Slot()
-				.Padding(10, 10)
-				.FillWidth(1)
-				[
-					SNew(STextBlock)
-					.ColorAndOpacity(FSlateColor::UseSubduedForeground())
-					.Font(IDetailLayoutBuilder::GetDetailFont())
-					.ToolTip(TraceChannelTooltip)
-					.Text(LOCTEXT("TraceChannel_Menu_Description", "You can have up to 18 custom channels including object and trace channels. This is list of trace channel for your project. \nIf you delete the trace channel that has been used by game, the behavior of trace is undefined."))
-				]
+				SNew(STextBlock)
+				.Font(IDetailLayoutBuilder::GetDetailFont())
+				.AutoWrapText(true)
+				.ToolTip(TraceChannelTooltip)
+				.Text(LOCTEXT("TraceChannel_Menu_Description", "You can have up to 18 custom channels including object and trace channels. This is list of trace channel for your project. If you delete the trace channel that has been used by game, the behavior of trace is undefined."))
+			]
 
-				+SHorizontalBox::Slot()
-				.Padding(2, 10)
-				.AutoWidth()
-				[
-					SNew(SButton)
-					.ContentPadding(FEditorStyle::GetMargin("StandardDialog.ContentPadding"))
-					.Text(LOCTEXT("ChannelMenu_NewTrace", "New Trace Channel...").ToString())
-					.OnClicked(this, &FCollisionProfileDetails::OnNewChannel, true)
-					.IsEnabled(this, &FCollisionProfileDetails::IsNewChannelAvailable)
-				]
+			+SHorizontalBox::Slot()
+			.Padding(2, 10)
+			.AutoWidth()
+			.VAlign(VAlign_Center)
+			[
+				SNew(SButton)
+				.ContentPadding(FEditorStyle::GetMargin("StandardDialog.ContentPadding"))
+				.Text(LOCTEXT("ChannelMenu_NewTrace", "New Trace Channel...").ToString())
+				.OnClicked(this, &FCollisionProfileDetails::OnNewChannel, true)
+				.IsEnabled(this, &FCollisionProfileDetails::IsNewChannelAvailable)
+			]
 
-				+SHorizontalBox::Slot()
-				.Padding(2, 10)
-				.AutoWidth()
-				[
-					SNew(SButton)
-					.ContentPadding(FEditorStyle::GetMargin("StandardDialog.ContentPadding"))
-					.Text(LOCTEXT("ChannelMenu_Edit", "Edit...").ToString())
-					.OnClicked(this, &FCollisionProfileDetails::OnEditChannel, true)
-					.IsEnabled(this, &FCollisionProfileDetails::IsAnyChannelSelected, true)
-				]
+			+SHorizontalBox::Slot()
+			.Padding(2, 10)
+			.AutoWidth()
+			.VAlign(VAlign_Center)
+			[
+				SNew(SButton)
+				.ContentPadding(FEditorStyle::GetMargin("StandardDialog.ContentPadding"))
+				.Text(LOCTEXT("ChannelMenu_Edit", "Edit...").ToString())
+				.OnClicked(this, &FCollisionProfileDetails::OnEditChannel, true)
+				.IsEnabled(this, &FCollisionProfileDetails::IsAnyChannelSelected, true)
+			]
 
-				+SHorizontalBox::Slot()
-				.Padding(2, 10)
-				.AutoWidth()
-				[
-					SNew(SButton)
-					.ContentPadding(FEditorStyle::GetMargin("StandardDialog.ContentPadding"))
-					.Text(LOCTEXT("ChannelMenu_Delete", "Delete...").ToString())
-					.OnClicked(this, &FCollisionProfileDetails::OnDeleteChannel, true)
-					.IsEnabled(this, &FCollisionProfileDetails::IsAnyChannelSelected, true)
-				]
+			+SHorizontalBox::Slot()
+			.Padding(2, 10)
+			.AutoWidth()
+			.VAlign(VAlign_Center)
+			[
+				SNew(SButton)
+				.ContentPadding(FEditorStyle::GetMargin("StandardDialog.ContentPadding"))
+				.Text(LOCTEXT("ChannelMenu_Delete", "Delete...").ToString())
+				.OnClicked(this, &FCollisionProfileDetails::OnDeleteChannel, true)
+				.IsEnabled(this, &FCollisionProfileDetails::IsAnyChannelSelected, true)
 			]
 		]
 
@@ -1533,44 +1499,36 @@ void FCollisionProfileDetails::CustomizeDetails( IDetailLayoutBuilder& DetailBui
 		.Padding(5)
 		.FillHeight(1)
 		[
-			SNew(SBorder)
-			.Padding(2.0)
-			.Content()
-			[
-				SAssignNew(TraceChannelListView, SChannelListView)
-				.ItemHeight(15.0f)
-				.ListItemsSource(&TraceChannelList)
-				.OnGenerateRow(this, &FCollisionProfileDetails::HandleGenerateChannelWidget)
-				.OnMouseButtonDoubleClick(this, &FCollisionProfileDetails::OnTraceChannelListItemDoubleClicked)
-				.SelectionMode(ESelectionMode::Single)
-
-				.HeaderRow(
-					SNew(SHeaderRow)
-					// Name
-					+ SHeaderRow::Column("Name")
-					.HAlignCell(HAlign_Left)
-					.FillWidth(1)
-					[
-						SNew(SHorizontalBox)
-						+SHorizontalBox::Slot()
-						.FillWidth(1)
-						[
-							SNew(STextBlock)
-							.Text(LOCTEXT("ChannelListHeader_Name", "Name"))
-							.Font(IDetailLayoutBuilder::GetDetailFont())
-						]
-					]
-					// Default Response
-					+ SHeaderRow::Column("DefaultResponse")
-					.HAlignCell(HAlign_Left)
-					.FillWidth(1)
-					[
-						SNew(STextBlock)
-						.Text(LOCTEXT("ChannelListHeader_DefaultResponse", "Default Response"))
-						.Font(IDetailLayoutBuilder::GetDetailFont())
-					]
-				)
-			]
+			SAssignNew(TraceChannelListView, SChannelListView)
+			.ItemHeight(15.0f)
+			.ListItemsSource(&TraceChannelList)
+			.OnGenerateRow(this, &FCollisionProfileDetails::HandleGenerateChannelWidget)
+			.OnMouseButtonDoubleClick(this, &FCollisionProfileDetails::OnTraceChannelListItemDoubleClicked)
+			.SelectionMode(ESelectionMode::Single)
+
+			.HeaderRow(
+				SNew(SHeaderRow)
+				// Name
+				+ SHeaderRow::Column("Name")
+				.HAlignCell(HAlign_Left)
+				.FillWidth(1)
+				.HeaderContentPadding(FMargin(0, 3))
+				[
+					SNew(STextBlock)
+					.Text(LOCTEXT("ChannelListHeader_Name", "Name"))
+					.Font(IDetailLayoutBuilder::GetDetailFont())
+				]
+				// Default Response
+				+ SHeaderRow::Column("DefaultResponse")
+				.HAlignCell(HAlign_Left)
+				.FillWidth(1)
+				.HeaderContentPadding(FMargin(0, 3))
+				[
+					SNew(STextBlock)
+					.Text(LOCTEXT("ChannelListHeader_DefaultResponse", "Default Response"))
+					.Font(IDetailLayoutBuilder::GetDetailFont())
+				]
+			)
 		]
 	];
 
@@ -1582,53 +1540,51 @@ void FCollisionProfileDetails::CustomizeDetails( IDetailLayoutBuilder& DetailBui
 		.Padding(5)
 		.AutoHeight()
 		[
-			SNew(SBorder)
-			.BorderImage(FEditorStyle::GetBrush("Menu.Background"))
-			.Padding(1.f)
+			SNew(SHorizontalBox)
+			+SHorizontalBox::Slot()
+			.Padding(2, 2)
+			.FillWidth(1)
 			[
-				SNew(SHorizontalBox)
-				+SHorizontalBox::Slot()
-				.Padding(2, 2)
-				.FillWidth(1)
-				[
-					SNew(STextBlock)
-					.ColorAndOpacity(FSlateColor::UseSubduedForeground())
-					.Font(IDetailLayoutBuilder::GetDetailFont())
-					.ToolTip(ProfileTooltip)
-					.Text(LOCTEXT("Profile_Menu_Description", "You can modify any of your project profiles. Please note that if you modify profile, it can change collision behavior.\nPlease be careful when you change currently exisiting (used) collision profiles."))
-				]
+				SNew(STextBlock)
+				.AutoWrapText(true)
+				.Font(IDetailLayoutBuilder::GetDetailFont())
+				.ToolTip(ProfileTooltip)
+				.Text(LOCTEXT("Profile_Menu_Description", "You can modify any of your project profiles. Please note that if you modify profile, it can change collision behavior. Please be careful when you change currently exisiting (used) collision profiles."))
+			]
 
-				+SHorizontalBox::Slot()
-				.Padding(2, 2)
-				.AutoWidth()
-				[
-					SNew(SButton)
-					.ContentPadding(FEditorStyle::GetMargin("StandardDialog.ContentPadding"))
-					.Text(LOCTEXT("ProfileMenu_New", "New...").ToString())
-					.OnClicked(this, &FCollisionProfileDetails::OnNewProfile)
-				]
+			+SHorizontalBox::Slot()
+			.Padding(2, 2)
+			.AutoWidth()
+			.VAlign(VAlign_Center)
+			[
+				SNew(SButton)
+				.ContentPadding(FEditorStyle::GetMargin("StandardDialog.ContentPadding"))
+				.Text(LOCTEXT("ProfileMenu_New", "New...").ToString())
+				.OnClicked(this, &FCollisionProfileDetails::OnNewProfile)
+			]
 
-				+SHorizontalBox::Slot()
-				.Padding(2, 2)
-				.AutoWidth()
-				[
-					SNew(SButton)
-					.ContentPadding(FEditorStyle::GetMargin("StandardDialog.ContentPadding"))
-					.Text(LOCTEXT("ProfileMenu_Edit", "Edit...").ToString())
-					.OnClicked(this, &FCollisionProfileDetails::OnEditProfile)
-					.IsEnabled(this, &FCollisionProfileDetails::IsAnyProfileSelected)
-				]
+			+SHorizontalBox::Slot()
+			.Padding(2, 2)
+			.AutoWidth()
+			.VAlign(VAlign_Center)
+			[
+				SNew(SButton)
+				.ContentPadding(FEditorStyle::GetMargin("StandardDialog.ContentPadding"))
+				.Text(LOCTEXT("ProfileMenu_Edit", "Edit...").ToString())
+				.OnClicked(this, &FCollisionProfileDetails::OnEditProfile)
+				.IsEnabled(this, &FCollisionProfileDetails::IsAnyProfileSelected)
+			]
 
-				+SHorizontalBox::Slot()
-				.Padding(2, 2)
-				.AutoWidth()
-				[
-					SNew(SButton)
-					.ContentPadding(FEditorStyle::GetMargin("StandardDialog.ContentPadding"))
-					.Text(LOCTEXT("ProfileMenu_Delete", "Delete...").ToString())
-					.OnClicked(this, &FCollisionProfileDetails::OnDeleteProfile)
-					.IsEnabled(this, &FCollisionProfileDetails::IsAnyProfileSelected)
-				]
+			+SHorizontalBox::Slot()
+			.Padding(2, 2)
+			.AutoWidth()
+			.VAlign(VAlign_Center)
+			[
+				SNew(SButton)
+				.ContentPadding(FEditorStyle::GetMargin("StandardDialog.ContentPadding"))
+				.Text(LOCTEXT("ProfileMenu_Delete", "Delete...").ToString())
+				.OnClicked(this, &FCollisionProfileDetails::OnDeleteProfile)
+				.IsEnabled(this, &FCollisionProfileDetails::IsAnyProfileSelected)
 			]
 		]
 
@@ -1636,75 +1592,70 @@ void FCollisionProfileDetails::CustomizeDetails( IDetailLayoutBuilder& DetailBui
 		.Padding(5)
 		.FillHeight(1)
 		[
-			SNew(SBorder)
-			.Padding(2.0)
-			[
-				SAssignNew(ProfileListView, SProfileListView)
-				.ItemHeight(20.0f)
-				.ListItemsSource(&ProfileList)
-				.OnGenerateRow(this, &FCollisionProfileDetails::HandleGenerateProfileWidget)
-				.OnMouseButtonDoubleClick(this, &FCollisionProfileDetails::OnProfileListItemDoubleClicked)
-				.SelectionMode(ESelectionMode::Single)
-
-				.HeaderRow(
-					SNew(SHeaderRow)
-					// Name
-					+ SHeaderRow::Column("Engine")
-					.HAlignCell(HAlign_Left)
-					.FixedWidth(30)
-					[
-						SNew(SHorizontalBox)
-						+SHorizontalBox::Slot()
-						.AutoWidth()
-						[
-							SNew(STextBlock)
-							.Text(LOCTEXT("ProfileListHeader_Category", ""))
-							.Font(IDetailLayoutBuilder::GetDetailFont())
-						]
-					]
-					// Name
-					+ SHeaderRow::Column("Name")
-					.HAlignCell(HAlign_Left)
-					.FillWidth(1)
-					[
-						SNew(SHorizontalBox)
-						+SHorizontalBox::Slot()
-						.FillWidth(1)
-						[
-							SNew(STextBlock)
-							.Text(LOCTEXT("ProfileListHeader_Name", "Name"))
-							.Font(IDetailLayoutBuilder::GetDetailFont())
-						]
-					]
-					// Default Response
-					+ SHeaderRow::Column("Collision")
-					.HAlignCell(HAlign_Left)
-					.FillWidth(1)
-					[
-						SNew(STextBlock)
-						.Text(LOCTEXT("ProfileListHeader_Collision", "Collision"))
-						.Font(IDetailLayoutBuilder::GetDetailFont())
-					]
-					// Trace Type
-					+ SHeaderRow::Column("ObjectType")
-					.HAlignCell(HAlign_Left)
-					.FillWidth(1)
-					[
-						SNew(STextBlock)
-						.Text(LOCTEXT("ProfileListHeader_ObjectType", "Object Type"))
-						.Font(IDetailLayoutBuilder::GetDetailFont())
-					]
-					// Static Object
-					+ SHeaderRow::Column("Description")
-					.HAlignCell(HAlign_Left)
-					.FillWidth(2)
+			SAssignNew(ProfileListView, SProfileListView)
+			.ItemHeight(20.0f)
+			.ListItemsSource(&ProfileList)
+			.OnGenerateRow(this, &FCollisionProfileDetails::HandleGenerateProfileWidget)
+			.OnMouseButtonDoubleClick(this, &FCollisionProfileDetails::OnProfileListItemDoubleClicked)
+			.SelectionMode(ESelectionMode::Single)
+
+			.HeaderRow(
+				SNew(SHeaderRow)
+				// Name
+				+ SHeaderRow::Column("Engine")
+				.HAlignCell(HAlign_Left)
+				.FixedWidth(30)
+				[
+					SNew(SHorizontalBox)
+					+SHorizontalBox::Slot()
+					.AutoWidth()
 					[
 						SNew(STextBlock)
-						.Text(LOCTEXT("ProfileListHeader_Description", "Description"))
+						.Text(LOCTEXT("ProfileListHeader_Category", ""))
 						.Font(IDetailLayoutBuilder::GetDetailFont())
 					]
-				)
-			]
+				]
+				// Name
+				+ SHeaderRow::Column("Name")
+				.HAlignCell(HAlign_Left)
+				.FillWidth(1)
+				.HeaderContentPadding(FMargin(0, 3))
+				[
+					SNew(STextBlock)
+					.Text(LOCTEXT("ProfileListHeader_Name", "Name"))
+					.Font(IDetailLayoutBuilder::GetDetailFont())
+				]
+				// Default Response
+				+ SHeaderRow::Column("Collision")
+				.HAlignCell(HAlign_Left)
+				.FillWidth(1)
+				.HeaderContentPadding(FMargin(0, 3))
+				[
+					SNew(STextBlock)
+					.Text(LOCTEXT("ProfileListHeader_Collision", "Collision"))
+					.Font(IDetailLayoutBuilder::GetDetailFont())
+				]
+				// Trace Type
+				+ SHeaderRow::Column("ObjectType")
+				.HAlignCell(HAlign_Left)
+				.FillWidth(1)
+				.HeaderContentPadding(FMargin(0, 3))
+				[
+					SNew(STextBlock)
+					.Text(LOCTEXT("ProfileListHeader_ObjectType", "Object Type"))
+					.Font(IDetailLayoutBuilder::GetDetailFont())
+				]
+				// Static Object
+				+ SHeaderRow::Column("Description")
+				.HAlignCell(HAlign_Left)
+				.FillWidth(2)
+				.HeaderContentPadding(FMargin(0, 3))
+				[
+					SNew(STextBlock)
+					.Text(LOCTEXT("ProfileListHeader_Description", "Description"))
+					.Font(IDetailLayoutBuilder::GetDetailFont())
+				]
+			)
 		]
 	];
 }
diff --git a/Engine/Source/Editor/DetailCustomizations/Private/PhysicsSettingsDetails.cpp b/Engine/Source/Editor/DetailCustomizations/Private/PhysicsSettingsDetails.cpp
--- a/Engine/Source/Editor/DetailCustomizations/Private/PhysicsSettingsDetails.cpp
+++ b/Engine/Source/Editor/DetailCustomizations/Private/PhysicsSettingsDetails.cpp
@@ -145,63 +145,55 @@ void FPhysicsSettingsDetails::CustomizeDetails( IDetailLayoutBuilder& DetailBuil
 		.Padding(5)
 		.AutoHeight()
 		[
-			SNew(SBorder)
-			.BorderImage(FEditorStyle::GetBrush("Menu.Background"))
-			.Padding(10.f)
-			[
-				SNew(STextBlock)
-				.ColorAndOpacity(FSlateColor::UseSubduedForeground())
-				.Font(IDetailLayoutBuilder::GetDetailFont())
-				.ToolTip(PhysicalSurfaceTooltip)
-				.Text(LOCTEXT("PhysicalSurface_Menu_Description", " You can have up to 62 custom surface types for your project. \nOnce you name each type, they will show up as surface type in the physical material."))
-			]
+			SNew(STextBlock)
+			.Font(IDetailLayoutBuilder::GetDetailFont())
+			.ToolTip(PhysicalSurfaceTooltip)
+			.AutoWrapText(true)
+			.Text(LOCTEXT("PhysicalSurface_Menu_Description", " You can have up to 62 custom surface types for your project. \nOnce you name each type, they will show up as surface type in the physical material."))
 		]
 
 		+SVerticalBox::Slot()
 		.Padding(5)
 		.FillHeight(1)
 		[
-			SNew(SBorder)
-			.Padding(2.0)
-			.Content()
-			[
-				SAssignNew(PhysicalSurfaceListView, SPhysicalSurfaceListView)
-				.ItemHeight(15.f)
-				.ListItemsSource(&PhysicalSurfaceList)
-				.OnGenerateRow(this, &FPhysicsSettingsDetails::HandleGenerateListWidget)
-				.SelectionMode(ESelectionMode::None)
-
-				.HeaderRow(
-					SNew(SHeaderRow)
-					+ SHeaderRow::Column("Type")
-					.HAlignCell(HAlign_Left)
+			SAssignNew(PhysicalSurfaceListView, SPhysicalSurfaceListView)
+			.ItemHeight(15.f)
+			.ListItemsSource(&PhysicalSurfaceList)
+			.OnGenerateRow(this, &FPhysicsSettingsDetails::HandleGenerateListWidget)
+			.SelectionMode(ESelectionMode::None)
+
+			.HeaderRow(
+				SNew(SHeaderRow)
+				+ SHeaderRow::Column("Type")
+				.HAlignCell(HAlign_Left)
+				.FillWidth(1)
+				.HeaderContentPadding(FMargin(0, 3))
+				[
+					SNew(SHorizontalBox)
+					+SHorizontalBox::Slot()
 					.FillWidth(1)
 					[
-						SNew(SHorizontalBox)
-						+SHorizontalBox::Slot()
-						.FillWidth(1)
-						[
-							SNew(STextBlock)
-							.Text(LOCTEXT("CPhysicalSurfaceListHeader_Type", "Type"))
-							.Font(IDetailLayoutBuilder::GetDetailFont())
-						]
+						SNew(STextBlock)
+						.Text(LOCTEXT("CPhysicalSurfaceListHeader_Type", "Type"))
+						.Font(IDetailLayoutBuilder::GetDetailFont())
 					]
-					// Name
-					+ SHeaderRow::Column("Name")
-					.HAlignCell(HAlign_Left)
+				]
+				// Name
+				+ SHeaderRow::Column("Name")
+				.HAlignCell(HAlign_Left)
+				.FillWidth(1)
+				.HeaderContentPadding(FMargin(0, 3))
+				[
+					SNew(SHorizontalBox)
+					+SHorizontalBox::Slot()
 					.FillWidth(1)
 					[
-						SNew(SHorizontalBox)
-						+SHorizontalBox::Slot()
-						.FillWidth(1)
-						[
-							SNew(STextBlock)
-							.Text(LOCTEXT("PhysicalSurfaceListHeader_Name", "Name"))
-							.Font(IDetailLayoutBuilder::GetDetailFont())
-						]
+						SNew(STextBlock)
+						.Text(LOCTEXT("PhysicalSurfaceListHeader_Name", "Name"))
+						.Font(IDetailLayoutBuilder::GetDetailFont())
 					]
-				)
-			]
+				]
+			)
 		]
 	];
 }
