Media: Added details view customization for Media Assets

diff --git a/Engine/Source/Editor/DetailCustomizations/Private/MediaAssetCustomization.cpp b/Engine/Source/Editor/DetailCustomizations/Private/MediaAssetCustomization.cpp
--- a/Engine/Source/Editor/DetailCustomizations/Private/MediaAssetCustomization.cpp
+++ b/Engine/Source/Editor/DetailCustomizations/Private/MediaAssetCustomization.cpp
@@ -4,10 +4,268 @@
 #include "MediaAssetCustomization.h"
 
 
+#define LOCTEXT_NAMESPACE "FMediaAssetCustomization"
+
+
 /* IDetailCustomization interface
  *****************************************************************************/
 
 void FMediaAssetCustomization::CustomizeDetails( IDetailLayoutBuilder& DetailBuilder )
 {
+	DetailBuilder.GetObjectsBeingCustomized(CustomizedMediaAssets);
+
+	// customize 'Information' category
+	IDetailCategoryBuilder& InformationCategory = DetailBuilder.EditCategory("Information", TEXT(""), ECategoryPriority::Uncommon);
+	{
+		// duration
+		InformationCategory.AddCustomRow(TEXT("Duration"))
+			.NameContent()
+			[
+				SNew(STextBlock)
+					.Text(LOCTEXT("Duration", "Duration"))
+					.Font(FEditorStyle::GetFontStyle(TEXT("PropertyWindow.NormalFont")))
+			]
+			.ValueContent()
+			.MaxDesiredWidth(0.0f)
+			[
+				SNew(STextBlock)
+					.Text(this, &FMediaAssetCustomization::HandleDurationTextBlockText)
+					.Font(FEditorStyle::GetFontStyle(TEXT("PropertyWindow.NormalFont")))
+			];
+
+		// forward rates (thinned)
+		InformationCategory.AddCustomRow(TEXT("ForwardRatesThinned"))
+			.NameContent()
+			[
+				SNew(STextBlock)
+					.Text(LOCTEXT("ForwardRatesThinned", "Forward Rates (Thinned)"))
+					.Font(FEditorStyle::GetFontStyle(TEXT("PropertyWindow.NormalFont")))
+			]
+			.ValueContent()
+			.MaxDesiredWidth(0.0f)
+			[
+				SNew(STextBlock)
+					.Text(this, &FMediaAssetCustomization::HandleSupportedRatesTextBlockText, EMediaPlaybackDirections::Forward, false)
+					.Font(FEditorStyle::GetFontStyle(TEXT("PropertyWindow.NormalFont")))
+			];
+
+		// forward rates (unthinned)
+		InformationCategory.AddCustomRow(TEXT("ForwardRatesUnthinned"))
+			.NameContent()
+			[
+				SNew(STextBlock)
+					.Text(LOCTEXT("ForwardRatesUnthinned", "Forward Rates (Unhinned)"))
+					.Font(FEditorStyle::GetFontStyle(TEXT("PropertyWindow.NormalFont")))
+			]
+			.ValueContent()
+			.MaxDesiredWidth(0.0f)
+			[
+				SNew(STextBlock)
+					.Text(this, &FMediaAssetCustomization::HandleSupportedRatesTextBlockText, EMediaPlaybackDirections::Forward, true)
+					.Font(FEditorStyle::GetFontStyle(TEXT("PropertyWindow.NormalFont")))
+			];
+
+		// reverse rates (thinned)
+		InformationCategory.AddCustomRow(TEXT("ReverseRatesThinned"))
+			.NameContent()
+			[
+				SNew(STextBlock)
+					.Text(LOCTEXT("ReverseRatesThinned", "Reverse Rates (Thinned)"))
+					.Font(FEditorStyle::GetFontStyle(TEXT("PropertyWindow.NormalFont")))
+			]
+			.ValueContent()
+			.MaxDesiredWidth(0.0f)
+			[
+				SNew(STextBlock)
+					.Text(this, &FMediaAssetCustomization::HandleSupportedRatesTextBlockText, EMediaPlaybackDirections::Reverse, false)
+					.Font(FEditorStyle::GetFontStyle(TEXT("PropertyWindow.NormalFont")))
+			];
+
+		// reverse rates (unthinned)
+		InformationCategory.AddCustomRow(TEXT("ReverseRatesUnthinned"))
+			.NameContent()
+			[
+				SNew(STextBlock)
+					.Text(LOCTEXT("ReverseRatesUnthinned", "Reverse Rates (Unthinned)"))
+					.Font(FEditorStyle::GetFontStyle(TEXT("PropertyWindow.NormalFont")))
+			]
+			.ValueContent()
+			.MaxDesiredWidth(0.0f)
+			[
+				SNew(STextBlock)
+					.Text(this, &FMediaAssetCustomization::HandleSupportedRatesTextBlockText, EMediaPlaybackDirections::Reverse, true)
+					.Font(FEditorStyle::GetFontStyle(TEXT("PropertyWindow.NormalFont")))
+			];
+
+		// supports scrubbing
+		InformationCategory.AddCustomRow(TEXT("SupportsScrubbing"))
+			.NameContent()
+			[
+				SNew(STextBlock)
+					.Text(LOCTEXT("SupportsScrubbing", "Supports Scrubbing"))
+					.Font(FEditorStyle::GetFontStyle(TEXT("PropertyWindow.NormalFont")))
+			]
+			.ValueContent()
+			.MaxDesiredWidth(0.0f)
+			[
+				SNew(STextBlock)
+					.Text(this, &FMediaAssetCustomization::HandleSupportsScrubbingTextBlockText)
+					.Font(FEditorStyle::GetFontStyle(TEXT("PropertyWindow.NormalFont")))
+			];
+
+		// supports seeking
+		InformationCategory.AddCustomRow(TEXT("SupportsSeeking"))
+			.NameContent()
+			[
+				SNew(STextBlock)
+					.Text(LOCTEXT("SupportsSeeking", "Supports Seeking"))
+					.Font(FEditorStyle::GetFontStyle(TEXT("PropertyWindow.NormalFont")))
+			]
+			.ValueContent()
+			.MaxDesiredWidth(0.0f)
+			[
+				SNew(STextBlock)
+					.Text(this, &FMediaAssetCustomization::HandleSupportsSeekingTextBlockText)
+					.Font(FEditorStyle::GetFontStyle(TEXT("PropertyWindow.NormalFont")))
+			];
+	}
+}
+
+
+/* FMediaAssetCustomization callbacks
+ *****************************************************************************/
+
+FText FMediaAssetCustomization::HandleDurationTextBlockText( ) const
+{
+	TOptional<FTimespan> Duration;
+
+	// ensure that all assets have the same value
+	for (auto MediaAssetObject : CustomizedMediaAssets)
+	{
+		if (MediaAssetObject.IsValid())
+		{
+			FTimespan OtherDuration = Cast<UMediaAsset>(MediaAssetObject.Get())->GetDuration();
+
+			if (!Duration.IsSet())
+			{
+				Duration = OtherDuration;
+			}
+			else if (OtherDuration != Duration.GetValue())
+			{
+				return FText::GetEmpty();
+			}
+		}
+	}
+
+	if (!Duration.IsSet())
+	{
+		return FText::GetEmpty();
+	}
+
+	return FText::AsTimespan(Duration.GetValue());
+}
+
+
+FText FMediaAssetCustomization::HandleSupportedRatesTextBlockText( EMediaPlaybackDirections Direction, bool Unthinned ) const
+{
+	TOptional<TRange<float>> Rates;
 
+	// ensure that all assets have the same value
+	for (auto MediaAssetObject : CustomizedMediaAssets)
+	{
+		if (MediaAssetObject.IsValid())
+		{
+			IMediaPlayerPtr MediaPlayer = Cast<UMediaAsset>(MediaAssetObject.Get())->GetMediaPlayer();
+
+			TRange<float> OtherRates = MediaPlayer->GetMediaInfo().GetSupportedRates(Direction, Unthinned);
+
+			if (!Rates.IsSet())
+			{
+				Rates = OtherRates;
+			}
+			else if (!MediaPlayer.IsValid() || (OtherRates != Rates.GetValue()))
+			{
+				return FText::GetEmpty();
+			}
+		}
+	}
+
+	if (!Rates.IsSet())
+	{
+		return FText::GetEmpty();
+	}
+
+	if ((Rates.GetValue().GetLowerBoundValue() == 0.0f) && (Rates.GetValue().GetUpperBoundValue() == 0.0f))
+	{
+		return LOCTEXT("RateNotSupported", "Not supported");
+	}
+
+	return FText::Format(LOCTEXT("RatesFormat", "{0} to {1}"),
+		FText::AsNumber(Rates.GetValue().GetLowerBoundValue()),
+		FText::AsNumber(Rates.GetValue().GetUpperBoundValue()));
 }
+
+
+FText FMediaAssetCustomization::HandleSupportsScrubbingTextBlockText( ) const
+{
+	TOptional<bool> SupportsScrubbing;
+
+	// ensure that all assets have the same value
+	for (auto MediaAssetObject : CustomizedMediaAssets)
+	{
+		if (MediaAssetObject.IsValid())
+		{
+			bool OtherSupportsScrubbing = Cast<UMediaAsset>(MediaAssetObject.Get())->SupportsScrubbing();
+
+			if (!SupportsScrubbing.IsSet())
+			{
+				SupportsScrubbing = OtherSupportsScrubbing;
+			}
+			else if (OtherSupportsScrubbing != SupportsScrubbing.GetValue())
+			{
+				return FText::GetEmpty();
+			}
+		}
+	}
+
+	if (!SupportsScrubbing.IsSet())
+	{
+		return FText::GetEmpty();
+	}
+
+	return SupportsScrubbing.GetValue() ? GTrue : GFalse;
+}
+
+
+FText FMediaAssetCustomization::HandleSupportsSeekingTextBlockText( ) const
+{
+	TOptional<bool> SupportsSeeking;
+
+	// ensure that all assets have the same value
+	for (auto MediaAssetObject : CustomizedMediaAssets)
+	{
+		if (MediaAssetObject.IsValid())
+		{
+			bool OtherSupportsSeeking = Cast<UMediaAsset>(MediaAssetObject.Get())->SupportsSeeking();
+
+			if (!SupportsSeeking.IsSet())
+			{
+				SupportsSeeking = OtherSupportsSeeking;
+			}
+			else if (OtherSupportsSeeking != SupportsSeeking.GetValue())
+			{
+				return FText::GetEmpty();
+			}
+		}
+	}
+
+	if (!SupportsSeeking.IsSet())
+	{
+		return FText::GetEmpty();
+	}
+
+	return SupportsSeeking.GetValue() ? GTrue : GFalse;
+}
+
+
+#undef LOCTEXT_NAMESPACE
diff --git a/Engine/Source/Editor/DetailCustomizations/Private/MediaAssetCustomization.h b/Engine/Source/Editor/DetailCustomizations/Private/MediaAssetCustomization.h
--- a/Engine/Source/Editor/DetailCustomizations/Private/MediaAssetCustomization.h
+++ b/Engine/Source/Editor/DetailCustomizations/Private/MediaAssetCustomization.h
@@ -29,4 +29,36 @@ public:
 	{
 		return MakeShareable(new FMediaAssetCustomization());
 	}
+
+protected:
+
+	/**
+	 * Gets the playback rate range text for the specified playback direction.
+	 *
+	 * @param Direction The playback direction.
+	 * @return The supported playback rate range converted to text.
+	 */
+	FText GetRatesText( EMediaPlaybackDirections Direction ) const;
+
+private:
+
+	// Callback for getting the text of the Duration text block.
+	FText HandleDurationTextBlockText( ) const;
+
+	// Callback for getting the text of the ForwardRates text block.
+	FText HandleForwardRatesTextBlockText( ) const;
+
+	// Callback for getting the text of a supported playback rate text block.
+	FText HandleSupportedRatesTextBlockText( EMediaPlaybackDirections Direction, bool Unthinned ) const;
+
+	// Callback for getting the text of the SupportsScrubbing text block.
+	FText HandleSupportsScrubbingTextBlockText( ) const;
+
+	// Callback for getting the text of the SupportsSeeking text block.
+	FText HandleSupportsSeekingTextBlockText( ) const;
+
+private:
+
+	/** The collection of media assets being customized */
+	TArray<TWeakObjectPtr<UObject>> CustomizedMediaAssets;
 };
diff --git a/Engine/Source/Editor/DetailCustomizations/Private/MediaTextureCustomization.cpp b/Engine/Source/Editor/DetailCustomizations/Private/MediaTextureCustomization.cpp
--- a/Engine/Source/Editor/DetailCustomizations/Private/MediaTextureCustomization.cpp
+++ b/Engine/Source/Editor/DetailCustomizations/Private/MediaTextureCustomization.cpp
@@ -14,34 +14,35 @@
 void FMediaTextureCustomization::CustomizeDetails( IDetailLayoutBuilder& DetailBuilder )
 {
 	DetailBuilder.GetObjectsBeingCustomized(CustomizedMediaTextures);
-
-	IDetailCategoryBuilder& MediaAssetCategory = DetailBuilder.EditCategory(TEXT("MediaAsset"));
-	
 	MediaAssetProperty = DetailBuilder.GetProperty("MediaAsset");
 	VideoTrackIndexProperty = DetailBuilder.GetProperty("VideoTrackIndex");
 
-	// customize video track index	
-	IDetailPropertyRow& VideoTrackIndexRow = MediaAssetCategory.AddProperty(VideoTrackIndexProperty);
-
-	VideoTrackIndexRow.DisplayName(TEXT("Video Track"));
-	VideoTrackIndexRow.CustomWidget()
-		.NameContent()
-		[
-			VideoTrackIndexProperty->CreatePropertyNameWidget()
-		]
-		.ValueContent()
-		.MaxDesiredWidth(0)
-		[
-			SNew(SComboButton)
-				.OnGetMenuContent(this, &FMediaTextureCustomization::HandleVideoTrackComboButtonMenuContent)
- 				.ContentPadding(FMargin( 2.0f, 2.0f ))
-				.ButtonContent()
-				[
-					SNew(STextBlock) 
-						.Text(this, &FMediaTextureCustomization::HandleVideoTrackComboButtonText)
-						.Font(IDetailLayoutBuilder::GetDetailFont())
-				]
-		];
+	// customize MediaAsset category
+	IDetailCategoryBuilder& MediaAssetCategory = DetailBuilder.EditCategory(TEXT("MediaAsset"));
+	{
+		// video track index
+		IDetailPropertyRow& VideoTrackIndexRow = MediaAssetCategory.AddProperty(VideoTrackIndexProperty);
+
+		VideoTrackIndexRow.DisplayName(TEXT("Video Track"));
+		VideoTrackIndexRow.CustomWidget()
+			.NameContent()
+			[
+				VideoTrackIndexProperty->CreatePropertyNameWidget()
+			]
+			.ValueContent()
+			.MaxDesiredWidth(0.0f)
+			[
+				SNew(SComboButton)
+					.OnGetMenuContent(this, &FMediaTextureCustomization::HandleVideoTrackComboButtonMenuContent)
+ 					.ContentPadding(FMargin( 2.0f, 2.0f ))
+					.ButtonContent()
+					[
+						SNew(STextBlock) 
+							.Text(this, &FMediaTextureCustomization::HandleVideoTrackComboButtonText)
+							.Font(IDetailLayoutBuilder::GetDetailFont())
+					]
+			];
+	}
 }
 
 
