342555	LIVE: EDITOR: Surface Types are not saved

diff --git a/Engine/Source/Editor/DetailCustomizations/Private/PhysicsSettingsDetails.cpp b/Engine/Source/Editor/DetailCustomizations/Private/PhysicsSettingsDetails.cpp
--- a/Engine/Source/Editor/DetailCustomizations/Private/PhysicsSettingsDetails.cpp
+++ b/Engine/Source/Editor/DetailCustomizations/Private/PhysicsSettingsDetails.cpp
@@ -203,6 +203,7 @@ void FPhysicsSettingsDetails::UpdatePhysicalSurfaceList()
 	RefreshPhysicalSurfaceList();
 	PhysicalSurfaceListView->RequestListRefresh();
 	PhysicsSettings->LoadSurfaceType();
+	PhysicsSettings->UpdateDefaultConfigFile();
 }
 
 void FPhysicsSettingsDetails::RefreshPhysicalSurfaceList()
