Editor: Refactored file path picker into a reusable widget

diff --git a/Engine/Source/Editor/DetailCustomizations/Private/FilePathStructCustomization.cpp b/Engine/Source/Editor/DetailCustomizations/Private/FilePathStructCustomization.cpp
--- a/Engine/Source/Editor/DetailCustomizations/Private/FilePathStructCustomization.cpp
+++ b/Engine/Source/Editor/DetailCustomizations/Private/FilePathStructCustomization.cpp
@@ -2,157 +2,81 @@
 
 #include "DetailCustomizationsPrivatePCH.h"
 #include "FilePathStructCustomization.h"
-#include "DesktopPlatformModule.h"
-#include "ContentBrowserDelegates.h"
+#include "SFilePathPicker.h"
+//#include "DesktopPlatformModule.h"
+//#include "ContentBrowserDelegates.h"
 
-#define LOCTEXT_NAMESPACE "FilePathStructCustomization"
-
-class SFilePathPicker : public SCompoundWidget
-{
-public:
-	SLATE_BEGIN_ARGS( SFilePathPicker ){}
-
-	SLATE_ARGUMENT(FOnPathPicked, OnPathPicked)
-
-	SLATE_ARGUMENT(FString, FileFilterExtension)
 
-	SLATE_END_ARGS()
-
-	void Construct(const FArguments& InArgs, TSharedRef<IPropertyHandle> InStringProperty)
-	{
-		StringProperty = InStringProperty;
-		FileFilterExtension = InArgs._FileFilterExtension;
-		OnPathPicked = InArgs._OnPathPicked;
-
-		ChildSlot
-		[
-			SNew(SHorizontalBox)
-			+SHorizontalBox::Slot()
-			.FillWidth(1.0f)
-			.VAlign(VAlign_Center)
-			[
-				StringProperty->CreatePropertyValueWidget()
-			]
-			+SHorizontalBox::Slot()
-			.AutoWidth()
-			.Padding(FMargin(4.0f, 0.0f, 0.0f, 0.0f))
-			.VAlign(VAlign_Center)
-			[
-				SNew(SButton)
-				.ButtonStyle( FEditorStyle::Get(), "HoverHintOnly" )
-				.ToolTipText( LOCTEXT( "FileButtonToolTipText", "Choose a file from this computer").ToString() )
-				.OnClicked( FOnClicked::CreateSP(this, &SFilePathPicker::OnPickFile) )
-				.ContentPadding( 2.0f )
-				.ForegroundColor( FSlateColor::UseForeground() )
-				.IsFocusable( false )
-				[
-					SNew( SImage )
-					.Image( FEditorStyle::GetBrush("PropertyWindow.Button_Ellipsis") )
-					.ColorAndOpacity( FSlateColor::UseForeground() )
-				]
-			]
-		];
-	}
-
-private:
-	/** Delegate handling the picker button being clicked */
-	FReply OnPickFile()
-	{
-		IDesktopPlatform* DesktopPlatform = FDesktopPlatformModule::Get();
-		if ( DesktopPlatform )
-		{
-			const FString DefaultPath = FEditorDirectories::Get().GetLastDirectory(ELastDirectory::GENERIC_OPEN);
-
-			TSharedPtr<SWindow> ParentWindow = FSlateApplication::Get().FindWidgetWindow(AsShared());
-			void* ParentWindowHandle = (ParentWindow.IsValid() && ParentWindow->GetNativeWindow().IsValid()) ? ParentWindow->GetNativeWindow()->GetOSWindowHandle() : nullptr;
-			FString Filter;
-
-			if (FileFilterExtension.IsEmpty())
-			{
-				Filter = TEXT("All files (*.*)|*.*");
-			}
-			else
-			{
-				Filter = FString::Printf(TEXT("%s files (*.%s)|*.%s"), *FileFilterExtension, *FileFilterExtension, *FileFilterExtension);
-			}
-
-			TArray<FString> OutFiles;
-		
-			if (DesktopPlatform->OpenFileDialog(ParentWindowHandle, LOCTEXT("PropertyEditorTitle", "File picker...").ToString(), DefaultPath, TEXT(""), Filter, EFileDialogFlags::None, OutFiles))
-			{
-				FEditorDirectories::Get().SetLastDirectory(ELastDirectory::GENERIC_OPEN, FPaths::GetPath(OutFiles[0]));
-
-				bool bSetProperty = true;
-				if(OnPathPicked.IsBound())
-				{
-					bSetProperty = OnPathPicked.Execute(OutFiles[0]);
-				}
-
-				if(bSetProperty)
-				{
-					StringProperty->SetValueFromFormattedString( OutFiles[0] );
-				}
-			}
-		}
-
-		return FReply::Handled();
-	}
-
-private:
-
-	/** Delegate fired when a path is picked */
-	FOnPathPicked OnPathPicked;
+#define LOCTEXT_NAMESPACE "FilePathStructCustomization"
 
-	/** The property we are representing */
-	TSharedPtr<IPropertyHandle> StringProperty;
 
-	/** Extension to filter files with */
-	FString FileFilterExtension;
-};
+/* IDetailCustomization interface
+ *****************************************************************************/
 
-TSharedRef<IPropertyTypeCustomization> FFilePathStructCustomization::MakeInstance()
+void FFilePathStructCustomization::CustomizeChildren( TSharedRef<IPropertyHandle> StructPropertyHandle, class IDetailChildrenBuilder& StructBuilder, IPropertyTypeCustomizationUtils& StructCustomizationUtils )
 {
-	return MakeShareable(new FFilePathStructCustomization());
+	/* do nothing */
 }
 
+
 void FFilePathStructCustomization::CustomizeHeader( TSharedRef<IPropertyHandle> StructPropertyHandle, class FDetailWidgetRow& HeaderRow, IPropertyTypeCustomizationUtils& StructCustomizationUtils )
 {
-	const UProperty* Property = StructPropertyHandle->GetProperty();
-	checkSlow(Property);
+	PathProperty = StructPropertyHandle;
 
-	TSharedPtr<SWidget> PickerWidget = CreatePickerWidget(StructPropertyHandle, Property->GetMetaData( TEXT("FilePathFilter") ));
+	FString FileTypeFilter;
 
-	if(PickerWidget.IsValid())
+	// construct file type filter
+	const FString& MetaData = PathProperty->GetMetaData(TEXT("FilePathFilter"));
+
+	if (MetaData.IsEmpty())
 	{
-		HeaderRow.ValueContent()
-		.MinDesiredWidth(125.0f)
-		.MaxDesiredWidth(600.0f)
-		[
-			PickerWidget.ToSharedRef()
-		]
+		FileTypeFilter = TEXT("All files (*.*)|*.*");
+	}
+	else
+	{
+		FileTypeFilter = FString::Printf(TEXT("%s files (*.%s)|*.%s"), *MetaData, *MetaData, *MetaData);
+	}
+
+	// create path picker widget
+	HeaderRow
 		.NameContent()
 		[
 			StructPropertyHandle->CreatePropertyNameWidget()
+		]
+		.ValueContent()
+		.MaxDesiredWidth(0.0f)
+		.MinDesiredWidth(125.0f)
+		[
+			SNew(SFilePathPicker)
+				.BrowseButtonImage(FEditorStyle::GetBrush("PropertyWindow.Button_Ellipsis"))
+				.BrowseButtonStyle(FEditorStyle::Get(), "HoverHintOnly")
+				.BrowseButtonToolTip(LOCTEXT("FileButtonToolTipText", "Choose a file from this computer"))
+				.BrowseDirectory(FEditorDirectories::Get().GetLastDirectory(ELastDirectory::GENERIC_OPEN))
+				.BrowseTitle(LOCTEXT("PropertyEditorTitle", "File picker..."))
+				.FilePath(this, &FFilePathStructCustomization::HandleFilePathPickerFilePath)
+				.FileTypeFilter(FileTypeFilter)
+				.OnPathPicked(this, &FFilePathStructCustomization::HandleFilePathPickerPathPicked)
 		];
-	}
 }
 
-TSharedPtr<SWidget> FFilePathStructCustomization::CreatePickerWidget(TSharedRef<class IPropertyHandle> StructPropertyHandle, const FString& FileFilterExtension, const FOnPathPicked& OnPathPicked)
+
+/* FMediaTextureCustomization callbacks
+ *****************************************************************************/
+
+FString FFilePathStructCustomization::HandleFilePathPickerFilePath( ) const
 {
-	TSharedPtr<IPropertyHandle> PathProperty = StructPropertyHandle->GetChildHandle("FilePath");
-	if(PathProperty.IsValid())
-	{
-		return SNew(SFilePathPicker, PathProperty.ToSharedRef())
-			.OnPathPicked(OnPathPicked)
-			.FileFilterExtension(FileFilterExtension);
-	}
+	FString FilePath;
+	PathProperty->GetValue(FilePath);
 
-	return TSharedPtr<SWidget>();
+	return FilePath;
 }
 
-void FFilePathStructCustomization::CustomizeChildren( TSharedRef<IPropertyHandle> StructPropertyHandle, class IDetailChildrenBuilder& StructBuilder, IPropertyTypeCustomizationUtils& StructCustomizationUtils )
+
+void FFilePathStructCustomization::HandleFilePathPickerPathPicked( const FString& PickedPath )
 {
+	PathProperty->SetValueFromFormattedString(PickedPath);
+	FEditorDirectories::Get().SetLastDirectory(ELastDirectory::GENERIC_OPEN, FPaths::GetPath(PickedPath));
 }
 
 
-#undef LOCTEXT_NAMESPACE
\ No newline at end of file
+#undef LOCTEXT_NAMESPACE
diff --git a/Engine/Source/Editor/DetailCustomizations/Private/FilePathStructCustomization.h b/Engine/Source/Editor/DetailCustomizations/Private/FilePathStructCustomization.h
--- a/Engine/Source/Editor/DetailCustomizations/Private/FilePathStructCustomization.h
+++ b/Engine/Source/Editor/DetailCustomizations/Private/FilePathStructCustomization.h
@@ -2,26 +2,42 @@
 
 #pragma once
 
-/** Delegate fired when a path is picked */
-DECLARE_DELEGATE_RetVal_OneParam(bool, FOnPathPicked, FString& /* InOutPath */);
 
-class FFilePathStructCustomization : public IPropertyTypeCustomization
+/**
+ * Implements a details view customization for the FFilePath structure.
+ */
+class FFilePathStructCustomization
+	: public IPropertyTypeCustomization
 {
 public:
-	static TSharedRef<IPropertyTypeCustomization> MakeInstance();
 
-	/** IPropertyTypeCustomization interface */
-	virtual void CustomizeHeader( TSharedRef<class IPropertyHandle> StructPropertyHandle, class FDetailWidgetRow& HeaderRow, IPropertyTypeCustomizationUtils& StructCustomizationUtils ) override;
+	/**
+	 * Creates an instance of this class.
+	 *
+	 * @return The new instance.
+	 */
+	static TSharedRef<IPropertyTypeCustomization> MakeInstance( )
+	{
+		return MakeShareable(new FFilePathStructCustomization());
+	}
+
+public:
 
-	virtual void CustomizeChildren( TSharedRef<class IPropertyHandle> StructPropertyHandle, class IDetailChildrenBuilder& StructBuilder, IPropertyTypeCustomizationUtils& StructCustomizationUtils ) override;              
+	// IPropertyTypeCustomization interface
 
-	static TSharedPtr<SWidget> CreatePickerWidget(TSharedRef<class IPropertyHandle> StructPropertyHandle,  const FString& FileFilterExtension = FString(TEXT("*")), const FOnPathPicked& OnPathPicked = FOnPathPicked());
+	virtual void CustomizeChildren( TSharedRef<class IPropertyHandle> StructPropertyHandle, class IDetailChildrenBuilder& StructBuilder, IPropertyTypeCustomizationUtils& StructCustomizationUtils ) override;
+	virtual void CustomizeHeader( TSharedRef<class IPropertyHandle> StructPropertyHandle, class FDetailWidgetRow& HeaderRow, IPropertyTypeCustomizationUtils& StructCustomizationUtils ) override;
 
 private:
 
-	/** Delegate for displaying text value of file */
-	FText GetDisplayedText(TSharedRef<IPropertyHandle> PropertyHandle) const;
+	/** Callback for getting the selected path in the picker widget. */
+	FString HandleFilePathPickerFilePath( ) const;
+
+	/** Callback for picking a file in the file path picker. */
+	void HandleFilePathPickerPathPicked( const FString& PickedPath );
+
+private:
 
-	/** Delegate used to display a file picker */
-	static FReply OnPickFile(TSharedRef<IPropertyHandle> PropertyHandle);
-};
\ No newline at end of file
+	/** Pointer to the Path property handle. */
+	TSharedPtr<IPropertyHandle> PathProperty;
+};
diff --git a/Engine/Source/Editor/DetailCustomizations/Private/MediaAssetCustomization.cpp b/Engine/Source/Editor/DetailCustomizations/Private/MediaAssetCustomization.cpp
--- a/Engine/Source/Editor/DetailCustomizations/Private/MediaAssetCustomization.cpp
+++ b/Engine/Source/Editor/DetailCustomizations/Private/MediaAssetCustomization.cpp
@@ -1,7 +1,10 @@
 // Copyright 1998-2014 Epic Games, Inc. All Rights Reserved.
 
 #include "DetailCustomizationsPrivatePCH.h"
+#include "Media.h"
+#include "MediaAsset.h"
 #include "MediaAssetCustomization.h"
+#include "SFilePathPicker.h"
 
 
 #define LOCTEXT_NAMESPACE "FMediaAssetCustomization"
@@ -14,6 +17,53 @@ void FMediaAssetCustomization::CustomizeDetails( IDetailLayoutBuilder& DetailBui
 {
 	DetailBuilder.GetObjectsBeingCustomized(CustomizedMediaAssets);
 
+	// customize 'Source' category
+	IDetailCategoryBuilder& SourceCategory = DetailBuilder.EditCategory("Source", TEXT(""));
+	{
+		// URL
+		UrlProperty = DetailBuilder.GetProperty("URL");
+		{
+			IDetailPropertyRow& UrlRow = SourceCategory.AddProperty(UrlProperty);
+
+			UrlRow.DisplayName(TEXT("URL"));
+			UrlRow.CustomWidget()
+				.NameContent()
+				[
+					SNew(SHorizontalBox)
+
+					+ SHorizontalBox::Slot()
+						.AutoWidth()
+						[
+							UrlProperty->CreatePropertyNameWidget()
+						]
+
+					+ SHorizontalBox::Slot()
+						.FillWidth(1.0f)
+						.HAlign(HAlign_Left)
+						.Padding(4.0f, 0.0f, 0.0f, 0.0f)
+						[
+							SNew(SImage)
+								.Image(FCoreStyle::Get().GetBrush("Icons.Warning"))
+								.ToolTipText(LOCTEXT("InvalidUrlPathWarning", "The current URL points to a file that does not exist or is not located inside the /Content/Movies/ directory."))
+								.Visibility(this, &FMediaAssetCustomization::HandleUrlWarningIconVisibility)
+						]
+				]
+				.ValueContent()
+				.MaxDesiredWidth(0.0f)
+				.MinDesiredWidth(125.0f)
+				[
+					SNew(SFilePathPicker)
+						.BrowseButtonImage(FEditorStyle::GetBrush("PropertyWindow.Button_Ellipsis"))
+						.BrowseButtonStyle(FEditorStyle::Get(), "HoverHintOnly")
+						.BrowseButtonToolTip(LOCTEXT("UrlBrowseButtonToolTipText", "Choose a file from this computer"))
+						.BrowseDirectory(FPaths::GameContentDir() / TEXT("Movies"))
+						.FilePath(this, &FMediaAssetCustomization::HandleUrlPickerFilePath)
+						.FileTypeFilter(this, &FMediaAssetCustomization::HandleUrlPickerFileTypeFilter)
+						.OnPathPicked(this, &FMediaAssetCustomization::HandleUrlPickerPathPicked)
+				];
+		}
+	}
+
 	// customize 'Information' category
 	IDetailCategoryBuilder& InformationCategory = DetailBuilder.EditCategory("Information", TEXT(""), ECategoryPriority::Uncommon);
 	{
@@ -177,6 +227,11 @@ FText FMediaAssetCustomization::HandleSupportedRatesTextBlockText( EMediaPlaybac
 		{
 			IMediaPlayerPtr MediaPlayer = Cast<UMediaAsset>(MediaAssetObject.Get())->GetMediaPlayer();
 
+			if (!MediaPlayer.IsValid())
+			{
+				return FText::GetEmpty();
+			}
+
 			TRange<float> OtherRates = MediaPlayer->GetMediaInfo().GetSupportedRates(Direction, Unthinned);
 
 			if (!Rates.IsSet())
@@ -268,4 +323,103 @@ FText FMediaAssetCustomization::HandleSupportsSeekingTextBlockText( ) const
 }
 
 
+FString FMediaAssetCustomization::HandleUrlPickerFilePath( ) const
+{
+	FString Url;
+	UrlProperty->GetValue(Url);
+
+	return Url;
+}
+
+
+FString FMediaAssetCustomization::HandleUrlPickerFileTypeFilter( ) const
+{
+	FString Filter = TEXT("All files (*.*)|*.*");
+
+	auto MediaModule = FModuleManager::GetModulePtr<IMediaModule>("Media");
+
+	if (MediaModule == nullptr)
+	{
+		return Filter;
+	}
+
+	FMediaFormats SupportedFormats;
+	MediaModule->GetSupportedFormats(SupportedFormats);
+
+	if (SupportedFormats.Num() == 0)
+	{
+		return Filter;
+	}
+
+	FString AllExtensions;
+	FString AllFilters;
+			
+	for (auto& Format : SupportedFormats)
+	{
+		if (!AllExtensions.IsEmpty())
+		{
+			AllExtensions += TEXT(";");
+		}
+
+		AllExtensions += TEXT("*.") + Format.Key;
+		AllFilters += TEXT("|") + Format.Value.ToString() + TEXT(" (*.") + Format.Key + TEXT(")|*.") + Format.Key;
+	}
+
+	Filter += TEXT("|All movie files (") + AllExtensions + TEXT(")|") + AllExtensions + AllFilters;
+
+	return Filter;
+}
+
+
+void FMediaAssetCustomization::HandleUrlPickerPathPicked( const FString& PickedPath )
+{
+	FString FullPath;
+
+	if (!PickedPath.IsEmpty() && FPaths::IsRelative(PickedPath))
+	{
+		FullPath = FPaths::ConvertRelativePathToFull(PickedPath);
+	}
+	else
+	{
+		FullPath = PickedPath;
+		FPaths::NormalizeFilename(FullPath);
+	}
+
+	if (FullPath.StartsWith(FPaths::RootDir()))
+	{
+		FPaths::MakePathRelativeTo(FullPath, FPlatformProcess::BaseDir());
+	}
+
+	UrlProperty->SetValue(FullPath);
+}
+
+
+EVisibility FMediaAssetCustomization::HandleUrlWarningIconVisibility( ) const
+{
+	FString Url;
+
+	if ((UrlProperty->GetValue(Url) != FPropertyAccess::Success) || Url.IsEmpty() || Url.Contains(TEXT("://")))
+	{
+		return EVisibility::Hidden;
+	}
+
+	const FString FullPath = FPaths::ConvertRelativePathToFull(Url);
+	const FString FullMoviesPath = FPaths::ConvertRelativePathToFull(FPaths::GameContentDir() / TEXT("Movies"));
+
+	if (FullPath.StartsWith(FullMoviesPath))
+	{
+		if (FPaths::FileExists(FullPath))
+		{
+			return EVisibility::Hidden;
+		}
+
+		// file doesn't exist
+		return EVisibility::Visible;
+	}
+
+	// file not inside Movies folder
+	return EVisibility::Visible;
+}
+
+
 #undef LOCTEXT_NAMESPACE
diff --git a/Engine/Source/Editor/DetailCustomizations/Private/MediaAssetCustomization.h b/Engine/Source/Editor/DetailCustomizations/Private/MediaAssetCustomization.h
--- a/Engine/Source/Editor/DetailCustomizations/Private/MediaAssetCustomization.h
+++ b/Engine/Source/Editor/DetailCustomizations/Private/MediaAssetCustomization.h
@@ -2,9 +2,6 @@
 
 #pragma once
 
-#include "Media.h"
-#include "MediaAsset.h"
-
 
 /**
  * Implements a details view customization for the UMediaAsset class.
@@ -30,16 +27,6 @@ public:
 		return MakeShareable(new FMediaAssetCustomization());
 	}
 
-protected:
-
-	/**
-	 * Gets the playback rate range text for the specified playback direction.
-	 *
-	 * @param Direction The playback direction.
-	 * @return The supported playback rate range converted to text.
-	 */
-	FText GetRatesText( EMediaPlaybackDirections Direction ) const;
-
 private:
 
 	// Callback for getting the text of the Duration text block.
@@ -49,7 +36,7 @@ private:
 	FText HandleForwardRatesTextBlockText( ) const;
 
 	// Callback for getting the text of a supported playback rate text block.
-	FText HandleSupportedRatesTextBlockText( EMediaPlaybackDirections Direction, bool Unthinned ) const;
+	FText HandleSupportedRatesTextBlockText( enum class EMediaPlaybackDirections Direction, bool Unthinned ) const;
 
 	// Callback for getting the text of the SupportsScrubbing text block.
 	FText HandleSupportsScrubbingTextBlockText( ) const;
@@ -57,8 +44,23 @@ private:
 	// Callback for getting the text of the SupportsSeeking text block.
 	FText HandleSupportsSeekingTextBlockText( ) const;
 
+	// Callback for getting the selected path in the URL picker widget.
+	FString HandleUrlPickerFilePath( ) const;
+
+	// Callback for getting the file type filter for the URL picker.
+	FString HandleUrlPickerFileTypeFilter( ) const;
+
+	// Callback for picking a path in the URL picker.
+	void HandleUrlPickerPathPicked( const FString& PickedPath );
+
+	// Callback for getting the visibility of warning icon for invalid URLs.
+	EVisibility HandleUrlWarningIconVisibility( ) const;
+
 private:
 
 	/** The collection of media assets being customized */
 	TArray<TWeakObjectPtr<UObject>> CustomizedMediaAssets;
+
+	/** Pointer to the URL property handle. */
+	TSharedPtr<IPropertyHandle> UrlProperty;
 };
diff --git a/Engine/Source/Editor/DetailCustomizations/Private/MoviePlayerSettingsDetails.cpp b/Engine/Source/Editor/DetailCustomizations/Private/MoviePlayerSettingsDetails.cpp
--- a/Engine/Source/Editor/DetailCustomizations/Private/MoviePlayerSettingsDetails.cpp
+++ b/Engine/Source/Editor/DetailCustomizations/Private/MoviePlayerSettingsDetails.cpp
@@ -2,38 +2,93 @@
 
 #include "DetailCustomizationsPrivatePCH.h"
 #include "MoviePlayerSettingsDetails.h"
-#include "FilePathStructCustomization.h"
 #include "SourceControlHelpers.h"
+#include "SFilePathPicker.h"
+
 
 #define LOCTEXT_NAMESPACE "MoviePlayerSettingsDetails"
 
+
+FMoviePlayerSettingsDetails::~FMoviePlayerSettingsDetails()
+{
+}
+
 TSharedRef<IDetailCustomization> FMoviePlayerSettingsDetails::MakeInstance()
 {
 	return MakeShareable(new FMoviePlayerSettingsDetails);
 }
 
-bool FMoviePlayerSettingsDetails::HandlePathPicked(FString& InOutPath)
+
+void FMoviePlayerSettingsDetails::CustomizeDetails( IDetailLayoutBuilder& DetailLayout )
+{
+	IDetailCategoryBuilder& MoviesCategory = DetailLayout.EditCategory("Movies");
+
+	StartupMoviesPropertyHandle = DetailLayout.GetProperty("StartupMovies");
+
+	TSharedRef<FDetailArrayBuilder> StartupMoviesBuilder = MakeShareable( new FDetailArrayBuilder( StartupMoviesPropertyHandle.ToSharedRef() ) );
+	StartupMoviesBuilder->OnGenerateArrayElementWidget( FOnGenerateArrayElementWidget::CreateSP(this, &FMoviePlayerSettingsDetails::GenerateArrayElementWidget) );
+
+	MoviesCategory.AddProperty( "bWaitForMoviesToComplete" );
+	MoviesCategory.AddProperty( "bMoviesAreSkippable" );
+
+	const bool bForAdvanced = false;
+	MoviesCategory.AddCustomBuilder( StartupMoviesBuilder, bForAdvanced );
+}
+
+
+void FMoviePlayerSettingsDetails::GenerateArrayElementWidget(TSharedRef<IPropertyHandle> PropertyHandle, int32 ArrayIndex, IDetailChildrenBuilder& ChildrenBuilder)
+{
+	IDetailPropertyRow& FilePathRow = ChildrenBuilder.AddChildProperty( PropertyHandle );
+	{
+		FilePathRow.CustomWidget(false)
+			.NameContent()
+			[
+				PropertyHandle->CreatePropertyNameWidget()
+			]
+			.ValueContent()
+			.MaxDesiredWidth(0.0f)
+			.MinDesiredWidth(125.0f)
+			[
+				SNew(SFilePathPicker)
+					.BrowseButtonImage(FEditorStyle::GetBrush("PropertyWindow.Button_Ellipsis"))
+					.BrowseButtonStyle(FEditorStyle::Get(), "HoverHintOnly")
+					.BrowseButtonToolTip(LOCTEXT("FileButtonToolTipText", "Choose a file from this computer"))
+					.BrowseDirectory(FEditorDirectories::Get().GetLastDirectory(ELastDirectory::GENERIC_OPEN))
+					.BrowseTitle(LOCTEXT("PropertyEditorTitle", "File picker..."))
+					.FilePath(this, &FMoviePlayerSettingsDetails::HandleFilePathPickerFilePath, PropertyHandle)
+					.FileTypeFilter(TEXT("MPEG-4 Movie (*.mp4)|*.mp4"))
+					.OnPathPicked(this, &FMoviePlayerSettingsDetails::HandleFilePathPickerPathPicked, PropertyHandle)
+			];
+	}
+};
+
+
+void FMoviePlayerSettingsDetails::HandleFilePathPickerPathPicked( const FString& PickedPath, TSharedRef<IPropertyHandle> Property )
 {
-	bool bSucceeded = false;
+	FEditorDirectories::Get().SetLastDirectory(ELastDirectory::GENERIC_OPEN, FPaths::GetPath(PickedPath));
 
 	// sanitize the location of the chosen movies to the content/movies directory
 	const FString MoviesBaseDir = FPaths::ConvertRelativePathToFull(FPaths::GameContentDir() + TEXT("Movies/"));
-	const FString FullPath = FPaths::ConvertRelativePathToFull(InOutPath);
-	if(FullPath.StartsWith(MoviesBaseDir))
+	const FString FullPath = FPaths::ConvertRelativePathToFull(PickedPath);
+
+	if (FullPath.StartsWith(MoviesBaseDir))
 	{
 		// mark for add/checkout
 		FText FailReason;
-		bSucceeded = SourceControlHelpers::CheckoutOrMarkForAdd(InOutPath, LOCTEXT("MovieFileDescription", "movie"), FOnPostCheckOut(), FailReason);
-		if(!bSucceeded)
+		
+		if (SourceControlHelpers::CheckoutOrMarkForAdd(PickedPath, LOCTEXT("MovieFileDescription", "movie"), FOnPostCheckOut(), FailReason))
+		{
+			// already in the movies dir, so just trim the path so we just have a partial path with no extension (the movie player expects this)
+			Property->SetValue(FPaths::GetBaseFilename(FullPath.RightChop(MoviesBaseDir.Len())));
+		}
+		else
 		{
 			FNotificationInfo Info(FailReason);
 			Info.ExpireDuration = 3.0f;
 			FSlateNotificationManager::Get().AddNotification(Info);
 		}
-		// already in the movies dir, so just trim the path so we just have a partial path with no extension (the movie player expects this)
-		InOutPath = FPaths::GetBaseFilename(FullPath.RightChop(MoviesBaseDir.Len()));
 	}
-	else
+	else if (!PickedPath.IsEmpty())
 	{
 		// ask the user if they want to import this movie
 		FSuppressableWarningDialog::FSetupInfo Info( 
@@ -47,72 +102,38 @@ bool FMoviePlayerSettingsDetails::HandlePathPicked(FString& InOutPath)
 
 		if(ImportWarningDialog.ShowModal() != FSuppressableWarningDialog::EResult::Cancel)
 		{
-			const FString FileName = FPaths::GetCleanFilename(InOutPath);
+			const FString FileName = FPaths::GetCleanFilename(PickedPath);
 			const FString DestPath = MoviesBaseDir / FileName;
 
 			FText FailReason;
-			bSucceeded = SourceControlHelpers::CopyFileUnderSourceControl(DestPath, InOutPath, LOCTEXT("MovieFileDescription", "movie"), FailReason);
-			if(!bSucceeded)
+			
+			if (SourceControlHelpers::CopyFileUnderSourceControl(DestPath, PickedPath, LOCTEXT("MovieFileDescription", "movie"), FailReason))
+			{
+				// trim the path so we just have a partial path with no extension (the movie player expects this)
+				Property->SetValue(FPaths::GetBaseFilename(DestPath.RightChop(MoviesBaseDir.Len())));
+			}
+			else
 			{
 				FNotificationInfo FailureInfo(FailReason);
 				FailureInfo.ExpireDuration = 3.0f;
 				FSlateNotificationManager::Get().AddNotification(FailureInfo);
 			}
-			// trim the path so we just have a partial path with no extension (the movie player expects this)
-			InOutPath = FPaths::GetBaseFilename(DestPath.RightChop(MoviesBaseDir.Len()));
 		}		
 	}
-
-	return bSucceeded;
-}
-
-void FMoviePlayerSettingsDetails::GenerateArrayElementWidget(TSharedRef<IPropertyHandle> PropertyHandle, int32 ArrayIndex, IDetailChildrenBuilder& ChildrenBuilder)
-{
-	const UProperty* Property = StartupMoviesPropertyHandle->GetProperty();
-	checkSlow(Property);
-	FString FileFilterExtension = Property->GetMetaData( TEXT("FilePathFilter") );
-	if( FileFilterExtension.IsEmpty() )
+	else
 	{
-		FileFilterExtension = TEXT("*");
+		Property->SetValue(PickedPath);
 	}
+}
 
-	TSharedPtr<SWidget> NewValueWidget = FFilePathStructCustomization::CreatePickerWidget(PropertyHandle, FileFilterExtension, FOnPathPicked::CreateSP(this, &FMoviePlayerSettingsDetails::HandlePathPicked));
-	if(NewValueWidget.IsValid())
-	{
-		IDetailPropertyRow& FilePathRow = ChildrenBuilder.AddChildProperty( PropertyHandle );
-
-		TSharedPtr<SWidget> NameWidget;
-		TSharedPtr<SWidget> ValueWidget;
-		FilePathRow.GetDefaultWidgets(NameWidget, ValueWidget);
-
-		const bool bShowChildren = false;
-		FilePathRow.CustomWidget(bShowChildren)
-		.NameContent()
-		[
-			NameWidget.ToSharedRef()
-		]
-		.ValueContent()
-		.MaxDesiredWidth(325.0f)
-		[
-			NewValueWidget.ToSharedRef()
-		];
-	}
-};
 
-void FMoviePlayerSettingsDetails::CustomizeDetails( IDetailLayoutBuilder& DetailLayout )
+FString FMoviePlayerSettingsDetails::HandleFilePathPickerFilePath( TSharedRef<IPropertyHandle> Property ) const
 {
-	IDetailCategoryBuilder& MoviesCategory = DetailLayout.EditCategory("Movies");
-
-	StartupMoviesPropertyHandle = DetailLayout.GetProperty("StartupMovies");
-
-	TSharedRef<FDetailArrayBuilder> StartupMoviesBuilder = MakeShareable( new FDetailArrayBuilder( StartupMoviesPropertyHandle.ToSharedRef() ) );
-	StartupMoviesBuilder->OnGenerateArrayElementWidget( FOnGenerateArrayElementWidget::CreateSP(this, &FMoviePlayerSettingsDetails::GenerateArrayElementWidget) );
-
-	MoviesCategory.AddProperty( "bWaitForMoviesToComplete" );
-	MoviesCategory.AddProperty( "bMoviesAreSkippable" );
+	FString FilePath;
+	Property->GetValue(FilePath);
 
-	const bool bForAdvanced = false;
-	MoviesCategory.AddCustomBuilder( StartupMoviesBuilder, bForAdvanced );
+	return FilePath;
 }
 
+
 #undef LOCTEXT_NAMESPACE
\ No newline at end of file
diff --git a/Engine/Source/Editor/DetailCustomizations/Private/MoviePlayerSettingsDetails.h b/Engine/Source/Editor/DetailCustomizations/Private/MoviePlayerSettingsDetails.h
--- a/Engine/Source/Editor/DetailCustomizations/Private/MoviePlayerSettingsDetails.h
+++ b/Engine/Source/Editor/DetailCustomizations/Private/MoviePlayerSettingsDetails.h
@@ -5,6 +5,8 @@
 class FMoviePlayerSettingsDetails : public IDetailCustomization
 {
 public:
+	~FMoviePlayerSettingsDetails();
+
 	/** Makes a new instance of this detail layout class for a specific detail view requesting it */
 	static TSharedRef<IDetailCustomization> MakeInstance();
 
@@ -12,11 +14,15 @@ public:
 	virtual void CustomizeDetails( IDetailLayoutBuilder& DetailLayout ) override;
 
 private:
+
 	/** Generate a widget for a movie array element */
 	void GenerateArrayElementWidget(TSharedRef<IPropertyHandle> PropertyHandle, int32 ArrayIndex, IDetailChildrenBuilder& ChildrenBuilder);
 
+	/** Callback for getting the selected path in the picker widget. */
+	FString HandleFilePathPickerFilePath( TSharedRef<IPropertyHandle> Property ) const;
+
 	/** Delegate handler for when a new movie path is picked */
-	bool HandlePathPicked(FString& InOutPath);
+	void HandleFilePathPickerPathPicked( const FString& InOutPath, TSharedRef<IPropertyHandle> Property );
 
 private:
 	/** Handle to the movies array property */
diff --git a/Engine/Source/Runtime/AppFramework/AppFramework.Build.cs b/Engine/Source/Runtime/AppFramework/AppFramework.Build.cs
--- a/Engine/Source/Runtime/AppFramework/AppFramework.Build.cs
+++ b/Engine/Source/Runtime/AppFramework/AppFramework.Build.cs
@@ -16,6 +16,7 @@ public class AppFramework : ModuleRules
 
 		PrivateDependencyModuleNames.AddRange(
 			new string[] {
+                "DesktopPlatform",
                 "InputCore",
 			}
 		);
diff --git a/Engine/Source/Runtime/AppFramework/Private/AppFrameworkPrivatePCH.h b/Engine/Source/Runtime/AppFramework/Private/AppFrameworkPrivatePCH.h
--- a/Engine/Source/Runtime/AppFramework/Private/AppFrameworkPrivatePCH.h
+++ b/Engine/Source/Runtime/AppFramework/Private/AppFrameworkPrivatePCH.h
@@ -8,6 +8,7 @@
 /* Private dependencies
  *****************************************************************************/
 
+#include "DesktopPlatformModule.h"
 
 
 /* Private includes
diff --git a/Engine/Source/Runtime/AppFramework/Private/Widgets/Input/SFilePathPicker.cpp b/Engine/Source/Runtime/AppFramework/Private/Widgets/Input/SFilePathPicker.cpp
new file mode 100644
--- /dev/null
+++ b/Engine/Source/Runtime/AppFramework/Private/Widgets/Input/SFilePathPicker.cpp
@@ -0,0 +1,108 @@
+// Copyright 1998-2014 Epic Games, Inc. All Rights Reserved.
+
+#include "AppFrameworkPrivatePCH.h"
+#include "SFilePathPicker.h"
+
+
+#define LOCTEXT_NAMESPACE "SFilePathPicker"
+
+
+/* SFilePathPicker interface
+ *****************************************************************************/
+
+void SFilePathPicker::Construct( const FArguments& InArgs )
+{
+	BrowseDirectory = InArgs._BrowseDirectory;
+	BrowseTitle = InArgs._BrowseTitle;
+	FilePath = InArgs._FilePath;
+	FileTypeFilter = InArgs._FileTypeFilter;
+	OnPathPicked = InArgs._OnPathPicked;
+
+	ChildSlot
+	[
+		SNew(SHorizontalBox)
+
+		+ SHorizontalBox::Slot()
+			.FillWidth(1.0f)
+			.VAlign(VAlign_Center)
+			[
+				// file path text box
+				SAssignNew(TextBox, SEditableTextBox)
+					.Text(this, &SFilePathPicker::HandleTextBoxText)
+					.Font(InArgs._Font)
+					.SelectAllTextWhenFocused(true)
+					.ClearKeyboardFocusOnCommit(false)
+					.OnTextCommitted(this, &SFilePathPicker::HandleTextBoxTextCommitted)
+					.SelectAllTextOnCommit(false)
+					.IsReadOnly(InArgs._IsReadOnly)
+			]
+
+		+ SHorizontalBox::Slot()
+			.AutoWidth()
+			.Padding(4.0f, 0.0f, 0.0f, 0.0f)
+			.VAlign(VAlign_Center)
+			[
+				// browse button
+				SNew(SButton)
+					.ButtonStyle(InArgs._BrowseButtonStyle)
+					.ToolTipText(InArgs._BrowseButtonToolTip)
+					.OnClicked(this, &SFilePathPicker::HandleBrowseButtonClicked)
+					.ContentPadding(2.0f)
+					.ForegroundColor(FSlateColor::UseForeground())
+					.IsFocusable(false)
+					[
+						SNew(SImage)
+							.Image(InArgs._BrowseButtonImage)
+							.ColorAndOpacity(FSlateColor::UseForeground())
+					]
+			]
+	];
+}
+
+
+/* SFilePathPicker callbacks
+ *****************************************************************************/
+
+FReply SFilePathPicker::HandleBrowseButtonClicked( )
+{
+	IDesktopPlatform* DesktopPlatform = FDesktopPlatformModule::Get();
+
+	if (DesktopPlatform == nullptr)
+	{
+		return FReply::Handled();
+	}
+
+	const FString DefaultPath = BrowseDirectory.IsSet()
+		? BrowseDirectory.Get()
+		: FPaths::GetPath(FilePath.Get());
+
+	// show the file browse dialog
+	TSharedPtr<SWindow> ParentWindow = FSlateApplication::Get().FindWidgetWindow(AsShared());
+	void* ParentWindowHandle = (ParentWindow.IsValid() && ParentWindow->GetNativeWindow().IsValid())
+		? ParentWindow->GetNativeWindow()->GetOSWindowHandle()
+		: nullptr;
+
+	TArray<FString> OutFiles;
+
+	if (DesktopPlatform->OpenFileDialog(ParentWindowHandle, BrowseTitle.Get().ToString(), DefaultPath, TEXT(""), FileTypeFilter.Get(), EFileDialogFlags::None, OutFiles))
+	{
+		OnPathPicked.ExecuteIfBound(OutFiles[0]);
+	}
+
+	return FReply::Handled();
+}
+
+
+FText SFilePathPicker::HandleTextBoxText( ) const
+{
+	return FText::FromString(FilePath.Get());
+}
+
+
+void SFilePathPicker::HandleTextBoxTextCommitted( const FText& NewText, ETextCommit::Type CommitInfo )
+{
+	OnPathPicked.ExecuteIfBound(NewText.ToString());
+}
+
+
+#undef LOCTEXT_NAMESPACE
diff --git a/Engine/Source/Runtime/AppFramework/Public/Widgets/Input/SFilePathPicker.h b/Engine/Source/Runtime/AppFramework/Public/Widgets/Input/SFilePathPicker.h
new file mode 100644
--- /dev/null
+++ b/Engine/Source/Runtime/AppFramework/Public/Widgets/Input/SFilePathPicker.h
@@ -0,0 +1,100 @@
+// Copyright 1998-2014 Epic Games, Inc. All Rights Reserved.
+
+#pragma once
+
+
+/**
+ * Declares a delegate that is executed when a file was picked in the SFilePathPicker widget.
+ *
+ * The first parameter will contain the path to the picked file.
+ */
+DECLARE_DELEGATE_OneParam(FOnPathPicked, const FString& /*PickedPath*/);
+
+
+/**
+ * Implements an editable text box with a browse button.
+ */
+class APPFRAMEWORK_API SFilePathPicker
+	: public SCompoundWidget
+{
+public:
+
+	SLATE_BEGIN_ARGS(SFilePathPicker)
+		: _BrowseButtonToolTip(NSLOCTEXT("SFilePathPicker", "BrowseButtonToolTip", "Choose a file from this computer"))
+		, _Font()
+		, _FileTypeFilter(TEXT("All files (*.*)|*.*"))
+		, _IsReadOnly(false)
+	{ }
+
+		/** Browse button image resource. */
+		SLATE_ATTRIBUTE(const FSlateBrush*, BrowseButtonImage)
+
+		/** Browse button visual style. */
+		SLATE_STYLE_ARGUMENT(FButtonStyle, BrowseButtonStyle)
+
+		/** Browse button tool tip text. */
+		SLATE_ATTRIBUTE(FText, BrowseButtonToolTip)
+
+		/** The directory to browse by default */
+		SLATE_ATTRIBUTE(FString, BrowseDirectory)
+
+		/** Title for the browse dialog window. */
+		SLATE_ATTRIBUTE(FText, BrowseTitle)
+
+		/** The currently selected file path. */
+		SLATE_ATTRIBUTE(FString, FilePath)
+
+		/** File type filter string. */
+		SLATE_ATTRIBUTE(FString, FileTypeFilter)
+
+		/** Font color and opacity of the path text box. */
+		SLATE_ATTRIBUTE(FSlateFontInfo, Font)
+
+		/** Whether the path text box can be modified by the user. */
+		SLATE_ATTRIBUTE(bool, IsReadOnly)
+
+		/** Called when a file path has been picked. */
+		SLATE_EVENT(FOnPathPicked, OnPathPicked)
+
+	SLATE_END_ARGS()
+
+	/**
+	 * Constructs a new widget.
+	 *
+	 * @param InArgs The construction arguments.
+	 */
+	void Construct( const FArguments& InArgs );
+
+private:
+
+	// Callback for clicking the browse button.
+	FReply HandleBrowseButtonClicked( );
+
+	// Callback for getting the text in the path text box.
+	FText HandleTextBoxText( ) const;
+
+	// Callback for committing the text in the path text box.
+	void HandleTextBoxTextCommitted( const FText& NewText, ETextCommit::Type /*CommitInfo*/ );
+
+private:
+
+	/** Holds the directory path to browse by default. */
+	TAttribute<FString> BrowseDirectory;
+
+	/** Holds the title for the browse dialog window. */
+	TAttribute<FText> BrowseTitle;
+
+	/** Holds the currently selected file path. */
+	TAttribute<FString> FilePath;
+
+	/** Holds the file type filter string. */
+	TAttribute<FString> FileTypeFilter;
+
+	/** Holds the editable text box. */
+	TSharedPtr<SEditableTextBox> TextBox;
+
+private:
+
+	/** Holds a delegate that is executed when a file was picked. */
+	FOnPathPicked OnPathPicked;
+};
diff --git a/Engine/Source/Runtime/MediaAssets/Private/Assets/MediaAsset.cpp b/Engine/Source/Runtime/MediaAssets/Private/Assets/MediaAsset.cpp
--- a/Engine/Source/Runtime/MediaAssets/Private/Assets/MediaAsset.cpp
+++ b/Engine/Source/Runtime/MediaAssets/Private/Assets/MediaAsset.cpp
@@ -96,7 +96,7 @@ bool UMediaAsset::IsStopped( ) const
 
 bool UMediaAsset::OpenUrl( const FString& NewUrl )
 {
-	URL.FilePath = NewUrl;
+	URL = NewUrl;
 	InitializePlayer();
 
 	return (CurrentUrl == NewUrl);
@@ -211,9 +211,7 @@ void UMediaAsset::PostEditChangeProperty( FPropertyChangedEvent& PropertyChanged
 
 void UMediaAsset::InitializePlayer( )
 {
-	const FString& NewUrl = URL.FilePath;
-
-	if (NewUrl != CurrentUrl)
+	if (URL != CurrentUrl)
 	{
 		// close previous player
 		CurrentUrl = FString();
@@ -226,13 +224,13 @@ void UMediaAsset::InitializePlayer( )
 			MediaPlayer.Reset();
 		}
 
-		if (NewUrl.IsEmpty())
+		if (URL.IsEmpty())
 		{
 			return;
 		}
 
 		// create new player
-		MediaPlayer = IMediaModule::Get().CreatePlayer(NewUrl);
+		MediaPlayer = IMediaModule::Get().CreatePlayer(URL);
 
 		if (!MediaPlayer.IsValid())
 		{
@@ -247,11 +245,11 @@ void UMediaAsset::InitializePlayer( )
 
 		if (StreamMode == EMediaAssetStreamModes::MASM_FromUrl)
 		{
-			OpenedSuccessfully = MediaPlayer->Open(NewUrl);
+			OpenedSuccessfully = MediaPlayer->Open(FPaths::ConvertRelativePathToFull(URL));
 		}
-		else
+		else if (FPaths::FileExists(URL))
 		{
-			FArchive* FileReader = IFileManager::Get().CreateFileReader(*NewUrl);
+			FArchive* FileReader = IFileManager::Get().CreateFileReader(*URL);
 		
 			if (FileReader == nullptr)
 			{
@@ -265,7 +263,7 @@ void UMediaAsset::InitializePlayer( )
 				FileData->AddUninitialized(FileReader->TotalSize());
 				FileReader->Serialize(FileData->GetData(), FileReader->TotalSize());
 
-				OpenedSuccessfully = MediaPlayer->Open(MakeShareable(FileData), NewUrl);
+				OpenedSuccessfully = MediaPlayer->Open(MakeShareable(FileData), URL);
 			}
 
 			delete FileReader;
@@ -274,7 +272,7 @@ void UMediaAsset::InitializePlayer( )
 		// finish initialization
 		if (OpenedSuccessfully)
 		{
-			CurrentUrl = NewUrl;
+			CurrentUrl = URL;
 		}
 	}
 
diff --git a/Engine/Source/Runtime/MediaAssets/Public/MediaAsset.h b/Engine/Source/Runtime/MediaAssets/Public/MediaAsset.h
--- a/Engine/Source/Runtime/MediaAssets/Public/MediaAsset.h
+++ b/Engine/Source/Runtime/MediaAssets/Public/MediaAsset.h
@@ -294,7 +294,7 @@ protected:
 
 	/** The URL to the media file to be played. */
 	UPROPERTY(EditAnywhere, Category=Source)
-	FFilePath URL;
+	FString URL;
 
 private:
 
diff --git a/Engine/Source/Runtime/MoviePlayer/Private/DefaultGameMoviePlayer.cpp b/Engine/Source/Runtime/MoviePlayer/Private/DefaultGameMoviePlayer.cpp
--- a/Engine/Source/Runtime/MoviePlayer/Private/DefaultGameMoviePlayer.cpp
+++ b/Engine/Source/Runtime/MoviePlayer/Private/DefaultGameMoviePlayer.cpp
@@ -268,34 +268,32 @@ void FDefaultGameMoviePlayer::SetupLoadingScreenFromIni()
 	// We may have already setup a movie from a startup module
 	if( !LoadingScreenAttributes.IsValid() )
 	{
+		// fill out the attributes
+		FLoadingScreenAttributes LoadingScreen;
+		LoadingScreen.bAutoCompleteWhenLoadingCompletes = !GetDefault<UMoviePlayerSettings>()->bWaitForMoviesToComplete;
+		LoadingScreen.bMoviesAreSkippable = GetDefault<UMoviePlayerSettings>()->bMoviesAreSkippable;
+
 		// look in the settings for a list of loading movies
 		GetMutableDefault<UMoviePlayerSettings>()->LoadConfig();
-		TArray<FFilePath> StartupMovies = GetDefault<UMoviePlayerSettings>()->StartupMovies;
+		const TArray<FString>& StartupMovies = GetDefault<UMoviePlayerSettings>()->StartupMovies;
 
-		// if we didn't have any, add default_startup, movie streamers should gracefully handle not finding this movie
 		if (StartupMovies.Num() == 0)
         {
-            FFilePath defaultMovie;
-            defaultMovie.FilePath = TEXT("Default_Startup");
-            StartupMovies.Add(defaultMovie);
+			LoadingScreen.MoviePaths.Add(TEXT("Default_Startup"));
         }
+		else
 		{
-			// fill out the attributes
-			FLoadingScreenAttributes LoadingScreen;
-			LoadingScreen.bAutoCompleteWhenLoadingCompletes = !GetDefault<UMoviePlayerSettings>()->bWaitForMoviesToComplete;
-			LoadingScreen.bMoviesAreSkippable = GetDefault<UMoviePlayerSettings>()->bMoviesAreSkippable;
-
-			for(const auto& Movie : StartupMovies)
+			for (const FString& Movie : StartupMovies)
 			{
-				if(Movie.FilePath.Len() > 0)
+				if (!Movie.IsEmpty())
 				{
-					LoadingScreen.MoviePaths.Add(Movie.FilePath);
+					LoadingScreen.MoviePaths.Add(Movie);
 				}
 			}
-
-			// now setup the actual loading screen
-			SetupLoadingScreen(LoadingScreen);
 		}
+
+		// now setup the actual loading screen
+		SetupLoadingScreen(LoadingScreen);
 	}
 }
 
diff --git a/Engine/Source/Runtime/MoviePlayer/Public/MoviePlayerSettings.h b/Engine/Source/Runtime/MoviePlayer/Public/MoviePlayerSettings.h
--- a/Engine/Source/Runtime/MoviePlayer/Public/MoviePlayerSettings.h
+++ b/Engine/Source/Runtime/MoviePlayer/Public/MoviePlayerSettings.h
@@ -25,6 +25,6 @@ public:
 	bool bMoviesAreSkippable;
 
 	/** Movies to play on startup. Note that these must be in your game's Game/Content/Movies directory. */
-	UPROPERTY(globalconfig, EditAnywhere, Category="Movies", meta=(FilePathFilter="mp4"))
-	TArray<FFilePath> StartupMovies;
+	UPROPERTY(globalconfig, EditAnywhere, Category="Movies")
+	TArray<FString> StartupMovies;
 };
