Editor: Added details view customization for FDateTime and FTimespan; also some code cleanup

diff --git a/Engine/Source/Editor/DetailCustomizations/Private/DateTimeStructCustomization.cpp b/Engine/Source/Editor/DetailCustomizations/Private/DateTimeStructCustomization.cpp
new file mode 100644
--- /dev/null
+++ b/Engine/Source/Editor/DetailCustomizations/Private/DateTimeStructCustomization.cpp
@@ -0,0 +1,96 @@
+// Copyright 1998-2014 Epic Games, Inc. All Rights Reserved.
+
+#include "DetailCustomizationsPrivatePCH.h"
+#include "DateTimeStructCustomization.h"
+
+
+#define LOCTEXT_NAMESPACE "DateTimeStructCustomization"
+
+
+/* IDetailCustomization interface
+ *****************************************************************************/
+
+void FDateTimeStructCustomization::CustomizeChildren( TSharedRef<IPropertyHandle> StructPropertyHandle, class IDetailChildrenBuilder& StructBuilder, IPropertyTypeCustomizationUtils& StructCustomizationUtils )
+{
+	/* do nothing */
+}
+
+
+void FDateTimeStructCustomization::CustomizeHeader( TSharedRef<IPropertyHandle> StructPropertyHandle, class FDetailWidgetRow& HeaderRow, IPropertyTypeCustomizationUtils& StructCustomizationUtils )
+{
+	PropertyHandle = StructPropertyHandle;
+
+	HeaderRow
+		.NameContent()
+		[
+			StructPropertyHandle->CreatePropertyNameWidget()
+		]
+		.ValueContent()
+		.MaxDesiredWidth(0.0f)
+		.MinDesiredWidth(125.0f)
+		[
+			SNew(SEditableTextBox)
+				.ClearKeyboardFocusOnCommit(false)
+				.IsEnabled(!PropertyHandle->IsEditConst())
+				.ForegroundColor(this, &FDateTimeStructCustomization::HandleTextBoxForegroundColor)
+				.OnTextChanged(this, &FDateTimeStructCustomization::HandleTextBoxTextChanged)
+				.OnTextCommitted(this, &FDateTimeStructCustomization::HandleTextBoxTextCommited)
+				.SelectAllTextOnCommit(true)
+				.Text(this, &FDateTimeStructCustomization::HandleTextBoxText)
+		];
+}
+
+
+/* FDateTimeStructCustomization callbacks
+ *****************************************************************************/
+
+FSlateColor FDateTimeStructCustomization::HandleTextBoxForegroundColor( ) const
+{
+	if (InputValid)
+	{
+		return FEditorStyle::GetSlateColor("InvertedForeground");
+	}
+
+	return FLinearColor::Red;
+}
+
+
+FText FDateTimeStructCustomization::HandleTextBoxText( ) const
+{
+	TArray<void*> RawData;
+	PropertyHandle->AccessRawData(RawData);
+	
+	if (RawData.Num() != 1)
+	{
+		return LOCTEXT("MultipleValues", "Multiple Values");
+	}
+
+	return FText::FromString(((FDateTime*)RawData[0])->ToString());
+}
+
+
+void FDateTimeStructCustomization::HandleTextBoxTextChanged( const FText& NewText )
+{
+	FDateTime DateTime;
+	InputValid = FDateTime::Parse(NewText.ToString(), DateTime);
+}
+
+
+void FDateTimeStructCustomization::HandleTextBoxTextCommited( const FText& NewText, ETextCommit::Type CommitInfo )
+{
+	FDateTime ParsedDateTime;
+								
+	if (FDateTime::Parse(NewText.ToString(), ParsedDateTime))
+	{
+		TArray<void*> RawData;
+		PropertyHandle->AccessRawData(RawData);
+
+		for (auto RawDataInstance : RawData)
+		{
+			*(FDateTime*)RawDataInstance = ParsedDateTime;
+		}
+	}
+}
+
+
+#undef LOCTEXT_NAMESPACE
diff --git a/Engine/Source/Editor/DetailCustomizations/Private/DateTimeStructCustomization.h b/Engine/Source/Editor/DetailCustomizations/Private/DateTimeStructCustomization.h
new file mode 100644
--- /dev/null
+++ b/Engine/Source/Editor/DetailCustomizations/Private/DateTimeStructCustomization.h
@@ -0,0 +1,55 @@
+// Copyright 1998-2014 Epic Games, Inc. All Rights Reserved.
+
+#pragma once
+
+
+/**
+ * Implements a details view customization for the FDateTime structure.
+ */
+class FDateTimeStructCustomization
+	: public IPropertyTypeCustomization
+{
+public:
+
+	/**
+	 * Creates an instance of this class.
+	 *
+	 * @return The new instance.
+	 */
+	static TSharedRef<IPropertyTypeCustomization> MakeInstance( )
+	{
+		return MakeShareable(new FDateTimeStructCustomization());
+	}
+
+public:
+
+	// IPropertyTypeCustomization interface
+
+	virtual void CustomizeChildren( TSharedRef<class IPropertyHandle> StructPropertyHandle, class IDetailChildrenBuilder& StructBuilder, IPropertyTypeCustomizationUtils& StructCustomizationUtils ) override;
+	virtual void CustomizeHeader( TSharedRef<class IPropertyHandle> StructPropertyHandle, class FDetailWidgetRow& HeaderRow, IPropertyTypeCustomizationUtils& StructCustomizationUtils ) override;
+
+private:
+
+	/** Handles getting the text color of the editable text box. */
+	FSlateColor HandleTextBoxForegroundColor( ) const;
+
+	/** Handles getting the text to be displayed in the editable text box. */
+	FText HandleTextBoxText( ) const;
+
+	/** Handles changing the value in the editable text box. */
+	void HandleTextBoxTextChanged( const FText& NewText );
+
+	/** Handles committing the text in the editable text box. */
+	void HandleTextBoxTextCommited( const FText& NewText, ETextCommit::Type CommitInfo );
+
+private:
+
+	/** Holds a flag indicating whether the current input is a valid GUID. */
+	bool InputValid;
+
+	/** Holds a handle to the property being edited. */
+	TSharedPtr<IPropertyHandle> PropertyHandle;
+
+	/** Holds the text box for editing the Guid. */
+	TSharedPtr<SEditableTextBox> TextBox;
+};
diff --git a/Engine/Source/Editor/DetailCustomizations/Private/DetailCustomizations.cpp b/Engine/Source/Editor/DetailCustomizations/Private/DetailCustomizations.cpp
--- a/Engine/Source/Editor/DetailCustomizations/Private/DetailCustomizations.cpp
+++ b/Engine/Source/Editor/DetailCustomizations/Private/DetailCustomizations.cpp
@@ -81,6 +81,8 @@
 #include "MediaTextureCustomization.h"
 #include "SoundWaveDetails.h"
 #include "AudioSettingsDetails.h"
+#include "DateTimeStructCustomization.h"
+#include "TimespanStructCustomization.h"
 
 IMPLEMENT_MODULE( FDetailCustomizationsModule, DetailCustomizations );
 
@@ -154,6 +156,8 @@ void FDetailCustomizationsModule::RegisterPropertyTypeCustomizations()
 	RegisterCustomPropertyTypeLayout("Key", FOnGetPropertyTypeCustomizationInstance::CreateStatic(&FKeyStructCustomization::MakeInstance));
 	RegisterCustomPropertyTypeLayout("FloatRange", FOnGetPropertyTypeCustomizationInstance::CreateStatic(&FRangeStructCustomization<float>::MakeInstance));
 	RegisterCustomPropertyTypeLayout("Int32Range", FOnGetPropertyTypeCustomizationInstance::CreateStatic(&FRangeStructCustomization<int32>::MakeInstance));
+	RegisterCustomPropertyTypeLayout("DateTime", FOnGetPropertyTypeCustomizationInstance::CreateStatic(&FDateTimeStructCustomization::MakeInstance));
+	RegisterCustomPropertyTypeLayout("Timespan", FOnGetPropertyTypeCustomizationInstance::CreateStatic(&FTimespanStructCustomization::MakeInstance));
 
 	RegisterCustomPropertyTypeLayout("BlackboardEntry", FOnGetPropertyTypeCustomizationInstance::CreateStatic(&FBlackboardEntryDetails::MakeInstance));
 	RegisterCustomPropertyTypeLayout("RuntimeFloatCurve", FOnGetPropertyTypeCustomizationInstance::CreateStatic(&FCurveStructCustomization::MakeInstance));
diff --git a/Engine/Source/Editor/DetailCustomizations/Private/FilePathStructCustomization.cpp b/Engine/Source/Editor/DetailCustomizations/Private/FilePathStructCustomization.cpp
--- a/Engine/Source/Editor/DetailCustomizations/Private/FilePathStructCustomization.cpp
+++ b/Engine/Source/Editor/DetailCustomizations/Private/FilePathStructCustomization.cpp
@@ -3,8 +3,6 @@
 #include "DetailCustomizationsPrivatePCH.h"
 #include "FilePathStructCustomization.h"
 #include "SFilePathPicker.h"
-//#include "DesktopPlatformModule.h"
-//#include "ContentBrowserDelegates.h"
 
 
 #define LOCTEXT_NAMESPACE "FilePathStructCustomization"
@@ -60,7 +58,7 @@ void FFilePathStructCustomization::CustomizeHeader( TSharedRef<IPropertyHandle>
 }
 
 
-/* FMediaTextureCustomization callbacks
+/* FFilePathStructCustomization callbacks
  *****************************************************************************/
 
 FString FFilePathStructCustomization::HandleFilePathPickerFilePath( ) const
diff --git a/Engine/Source/Editor/DetailCustomizations/Private/GuidStructCustomization.cpp b/Engine/Source/Editor/DetailCustomizations/Private/GuidStructCustomization.cpp
--- a/Engine/Source/Editor/DetailCustomizations/Private/GuidStructCustomization.cpp
+++ b/Engine/Source/Editor/DetailCustomizations/Private/GuidStructCustomization.cpp
@@ -7,15 +7,6 @@
 #define LOCTEXT_NAMESPACE "FGuidStructCustomization"
 
 
-/* FGuidStructCustomization static interface
- *****************************************************************************/
-
-TSharedRef<IPropertyTypeCustomization> FGuidStructCustomization::MakeInstance( )
-{
-	return MakeShareable(new FGuidStructCustomization);
-}
-
-
 /* IPropertyTypeCustomization interface
  *****************************************************************************/
 
diff --git a/Engine/Source/Editor/DetailCustomizations/Private/GuidStructCustomization.h b/Engine/Source/Editor/DetailCustomizations/Private/GuidStructCustomization.h
--- a/Engine/Source/Editor/DetailCustomizations/Private/GuidStructCustomization.h
+++ b/Engine/Source/Editor/DetailCustomizations/Private/GuidStructCustomization.h
@@ -25,13 +25,6 @@ namespace EPropertyEditorGuidActions
 class FGuidStructCustomization
 	: public IPropertyTypeCustomization
 {
-public:
-
-	// IPropertyTypeCustomization interface
-
-	virtual void CustomizeHeader( TSharedRef<class IPropertyHandle> StructPropertyHandle, class FDetailWidgetRow& HeaderRow, IPropertyTypeCustomizationUtils& StructCustomizationUtils ) override;
-	virtual void CustomizeChildren( TSharedRef<class IPropertyHandle> StructPropertyHandle, class IDetailChildrenBuilder& StructBuilder, IPropertyTypeCustomizationUtils& StructCustomizationUtils ) override;
-
 public:
 
 	/**
@@ -39,7 +32,17 @@ public:
 	 *
 	 * @return A new struct customization for Guids.
 	 */
-	static TSharedRef<IPropertyTypeCustomization> MakeInstance( );
+	static TSharedRef<IPropertyTypeCustomization> MakeInstance( )
+	{
+		return MakeShareable(new FGuidStructCustomization);
+	}
+
+public:
+
+	// IPropertyTypeCustomization interface
+
+	virtual void CustomizeHeader( TSharedRef<class IPropertyHandle> StructPropertyHandle, class FDetailWidgetRow& HeaderRow, IPropertyTypeCustomizationUtils& StructCustomizationUtils ) override;
+	virtual void CustomizeChildren( TSharedRef<class IPropertyHandle> StructPropertyHandle, class IDetailChildrenBuilder& StructBuilder, IPropertyTypeCustomizationUtils& StructCustomizationUtils ) override;
 
 protected:
 
@@ -52,29 +55,29 @@ protected:
 
 private:
 
-	// Callback for clicking an item in the quick-set menu.
+	/** Callback for clicking an item in the quick-set menu. */
 	void HandleGuidActionClicked( EPropertyEditorGuidActions::Type Action );
 
-	// Handles getting the text color of the editable text box.
+	/** Handles getting the text color of the editable text box. */
 	FSlateColor HandleTextBoxForegroundColor( ) const;
 
-	// Handles getting the text to be displayed in the editable text box.
+	/** Handles getting the text to be displayed in the editable text box. */
 	FText HandleTextBoxText( ) const;
 
-	// Handles changing the value in the editable text box.
+	/** Handles changing the value in the editable text box. */
 	void HandleTextBoxTextChanged( const FText& NewText );
 
-	// Handles committing the text in the editable text box.
+	/** Handles committing the text in the editable text box. */
 	void HandleTextBoxTextCommited( const FText& NewText, ETextCommit::Type CommitInfo );
 
 private:
 
-	// Holds a flag indicating whether the current input is a valid GUID.
+	/** Holds a flag indicating whether the current input is a valid GUID. */
 	bool InputValid;
 
-	// Holds a handle to the property being edited.
+	/** Holds a handle to the property being edited. */
 	TSharedPtr<IPropertyHandle> PropertyHandle;
 
-	// Holds the text box for editing the Guid.
+	/** Holds the text box for editing the Guid. */
 	TSharedPtr<SEditableTextBox> TextBox;
 };
diff --git a/Engine/Source/Editor/DetailCustomizations/Private/TimespanStructCustomization.cpp b/Engine/Source/Editor/DetailCustomizations/Private/TimespanStructCustomization.cpp
new file mode 100644
--- /dev/null
+++ b/Engine/Source/Editor/DetailCustomizations/Private/TimespanStructCustomization.cpp
@@ -0,0 +1,96 @@
+// Copyright 1998-2014 Epic Games, Inc. All Rights Reserved.
+
+#include "DetailCustomizationsPrivatePCH.h"
+#include "TimespanStructCustomization.h"
+
+
+#define LOCTEXT_NAMESPACE "TimespanStructCustomization"
+
+
+/* IDetailCustomization interface
+ *****************************************************************************/
+
+void FTimespanStructCustomization::CustomizeChildren( TSharedRef<IPropertyHandle> StructPropertyHandle, class IDetailChildrenBuilder& StructBuilder, IPropertyTypeCustomizationUtils& StructCustomizationUtils )
+{
+	/* do nothing */
+}
+
+
+void FTimespanStructCustomization::CustomizeHeader( TSharedRef<IPropertyHandle> StructPropertyHandle, class FDetailWidgetRow& HeaderRow, IPropertyTypeCustomizationUtils& StructCustomizationUtils )
+{
+	PropertyHandle = StructPropertyHandle;
+
+	HeaderRow
+		.NameContent()
+		[
+			StructPropertyHandle->CreatePropertyNameWidget()
+		]
+		.ValueContent()
+		.MaxDesiredWidth(0.0f)
+		.MinDesiredWidth(125.0f)
+		[
+			SNew(SEditableTextBox)
+				.ClearKeyboardFocusOnCommit(false)
+				.IsEnabled(!PropertyHandle->IsEditConst())
+				.ForegroundColor(this, &FTimespanStructCustomization::HandleTextBoxForegroundColor)
+				.OnTextChanged(this, &FTimespanStructCustomization::HandleTextBoxTextChanged)
+				.OnTextCommitted(this, &FTimespanStructCustomization::HandleTextBoxTextCommited)
+				.SelectAllTextOnCommit(true)
+				.Text(this, &FTimespanStructCustomization::HandleTextBoxText)
+		];
+}
+
+
+/* FTimespanStructCustomization callbacks
+ *****************************************************************************/
+
+FSlateColor FTimespanStructCustomization::HandleTextBoxForegroundColor( ) const
+{
+	if (InputValid)
+	{
+		return FEditorStyle::GetSlateColor("InvertedForeground");
+	}
+
+	return FLinearColor::Red;
+}
+
+
+FText FTimespanStructCustomization::HandleTextBoxText( ) const
+{
+	TArray<void*> RawData;
+	PropertyHandle->AccessRawData(RawData);
+	
+	if (RawData.Num() != 1)
+	{
+		return LOCTEXT("MultipleValues", "Multiple Values");
+	}
+
+	return FText::FromString(((FTimespan*)RawData[0])->ToString());
+}
+
+
+void FTimespanStructCustomization::HandleTextBoxTextChanged( const FText& NewText )
+{
+	FTimespan Timespan;
+	InputValid = FTimespan::Parse(NewText.ToString(), Timespan);
+}
+
+
+void FTimespanStructCustomization::HandleTextBoxTextCommited( const FText& NewText, ETextCommit::Type CommitInfo )
+{
+	FTimespan ParsedTimespan;
+								
+	if (FTimespan::Parse(NewText.ToString(), ParsedTimespan))
+	{
+		TArray<void*> RawData;
+		PropertyHandle->AccessRawData(RawData);
+
+		for (auto RawDataInstance : RawData)
+		{
+			*(FTimespan*)RawDataInstance = ParsedTimespan;
+		}
+	}
+}
+
+
+#undef LOCTEXT_NAMESPACE
diff --git a/Engine/Source/Editor/DetailCustomizations/Private/TimespanStructCustomization.h b/Engine/Source/Editor/DetailCustomizations/Private/TimespanStructCustomization.h
new file mode 100644
--- /dev/null
+++ b/Engine/Source/Editor/DetailCustomizations/Private/TimespanStructCustomization.h
@@ -0,0 +1,55 @@
+// Copyright 1998-2014 Epic Games, Inc. All Rights Reserved.
+
+#pragma once
+
+
+/**
+ * Implements a details view customization for the FTimespan structure.
+ */
+class FTimespanStructCustomization
+	: public IPropertyTypeCustomization
+{
+public:
+
+	/**
+	 * Creates an instance of this class.
+	 *
+	 * @return The new instance.
+	 */
+	static TSharedRef<IPropertyTypeCustomization> MakeInstance( )
+	{
+		return MakeShareable(new FTimespanStructCustomization());
+	}
+
+public:
+
+	// IPropertyTypeCustomization interface
+
+	virtual void CustomizeChildren( TSharedRef<class IPropertyHandle> StructPropertyHandle, class IDetailChildrenBuilder& StructBuilder, IPropertyTypeCustomizationUtils& StructCustomizationUtils ) override;
+	virtual void CustomizeHeader( TSharedRef<class IPropertyHandle> StructPropertyHandle, class FDetailWidgetRow& HeaderRow, IPropertyTypeCustomizationUtils& StructCustomizationUtils ) override;
+
+private:
+
+	/** Handles getting the text color of the editable text box. */
+	FSlateColor HandleTextBoxForegroundColor( ) const;
+
+	/** Handles getting the text to be displayed in the editable text box. */
+	FText HandleTextBoxText( ) const;
+
+	/** Handles changing the value in the editable text box. */
+	void HandleTextBoxTextChanged( const FText& NewText );
+
+	/** Handles committing the text in the editable text box. */
+	void HandleTextBoxTextCommited( const FText& NewText, ETextCommit::Type CommitInfo );
+
+private:
+
+	/** Holds a flag indicating whether the current input is a valid GUID. */
+	bool InputValid;
+
+	/** Holds a handle to the property being edited. */
+	TSharedPtr<IPropertyHandle> PropertyHandle;
+
+	/** Holds the text box for editing the Guid. */
+	TSharedPtr<SEditableTextBox> TextBox;
+};
diff --git a/Engine/Source/Editor/SourceControlWindows/Private/SSourceControlHistory.cpp b/Engine/Source/Editor/SourceControlWindows/Private/SSourceControlHistory.cpp
--- a/Engine/Source/Editor/SourceControlWindows/Private/SSourceControlHistory.cpp
+++ b/Engine/Source/Editor/SourceControlWindows/Private/SSourceControlHistory.cpp
@@ -540,7 +540,7 @@ public:
 		{
 			return
 				SNew(STextBlock)
-				.Text( RevisionListItem->Date.Ticks == 0 ? FText() : FText::AsDateTime( RevisionListItem->Date ) );
+				.Text( RevisionListItem->Date > FDateTime::MinValue() == 0 ? FText() : FText::AsDateTime( RevisionListItem->Date ) );
 		}
 		else if (ColumnName == TEXT("UserName"))
 		{
diff --git a/Engine/Source/Programs/CrashReporter/CrashReportClient/Private/CrashDescription.cpp b/Engine/Source/Programs/CrashReporter/CrashReportClient/Private/CrashDescription.cpp
--- a/Engine/Source/Programs/CrashReporter/CrashReportClient/Private/CrashDescription.cpp
+++ b/Engine/Source/Programs/CrashReporter/CrashReportClient/Private/CrashDescription.cpp
@@ -231,7 +231,7 @@ void FCrashDescription::SendAnalytics()
 	CrashAttributes.Add( FAnalyticsEventAttribute( TEXT( "EpicAccountId" ), EpicAccountId ) );
 
 	CrashAttributes.Add( FAnalyticsEventAttribute( TEXT( "Platform" ), Platform ) );
-	CrashAttributes.Add( FAnalyticsEventAttribute( TEXT( "TimeOfCrash" ), TimeOfCrash.Ticks ) );
+	CrashAttributes.Add( FAnalyticsEventAttribute( TEXT( "TimeOfCrash" ), TimeOfCrash.GetTicks() ) );
 	CrashAttributes.Add( FAnalyticsEventAttribute( TEXT( "EngineMode" ), EngineMode ) );
 	CrashAttributes.Add( FAnalyticsEventAttribute( TEXT( "LanguageLCID" ), LanguageLCID ) );
 
diff --git a/Engine/Source/Runtime/Core/Private/Misc/DateTime.cpp b/Engine/Source/Runtime/Core/Private/Misc/DateTime.cpp
--- a/Engine/Source/Runtime/Core/Private/Misc/DateTime.cpp
+++ b/Engine/Source/Runtime/Core/Private/Misc/DateTime.cpp
@@ -64,6 +64,14 @@ FArchive& operator<<( FArchive& Ar, FDateTime& DateTime )
 /* FDateTime interface
  *****************************************************************************/
 
+bool FDateTime::ExportTextItem( FString& ValueStr, FDateTime const& DefaultValue, UObject* Parent, int32 PortFlags, UObject* ExportRootScope ) const
+{
+	ValueStr += ToString();
+
+	return true;
+}
+
+
 void FDateTime::GetDate( int32& OutYear, int32& OutMonth, int32& OutDay ) const
 {
 	// Based on FORTRAN code in:
@@ -155,6 +163,24 @@ int32 FDateTime::GetYear( ) const
 }
 
 
+bool FDateTime::ImportTextItem( const TCHAR*& Buffer, int32 PortFlags, UObject* Parent, FOutputDevice* ErrorText )
+{
+	if (FPlatformString::Strlen(Buffer) < 19)
+	{
+		return false;
+	}
+
+	if (!Parse(FString(Buffer).Left(19), *this))
+	{
+		return false;
+	}
+
+	Buffer += 19;
+
+	return true;
+}
+
+
 FString FDateTime::ToIso8601( ) const
 {
 	return ToString(TEXT("%Y-%m-%dT%H:%M:%S.%sZ"));
@@ -427,28 +453,3 @@ FDateTime FDateTime::UtcNow( )
 
 	return FDateTime(Year, Month, Day, Hour, Minute, Second, Millisecond);
 }
-
-
-bool FDateTime::ExportTextItem( FString& ValueStr, FDateTime const& DefaultValue, UObject* Parent, int32 PortFlags, class UObject* ExportRootScope ) const
-{
-	ValueStr += ToString();
-
-	return true;
-}
-
-bool FDateTime::ImportTextItem( const TCHAR*& Buffer, int32 PortFlags, class UObject* Parent, FOutputDevice* ErrorText )
-{
-	if (FPlatformString::Strlen(Buffer) < 19)
-	{
-		return false;
-	}
-
-	if (!Parse(FString(Buffer).Left(19), *this))
-	{
-		return false;
-	}
-
-	Buffer += 19;
-
-	return true;
-}
diff --git a/Engine/Source/Runtime/Core/Public/Misc/DateTime.h b/Engine/Source/Runtime/Core/Public/Misc/DateTime.h
--- a/Engine/Source/Runtime/Core/Public/Misc/DateTime.h
+++ b/Engine/Source/Runtime/Core/Public/Misc/DateTime.h
@@ -74,7 +74,7 @@ namespace EMonthOfYear
  * @see FDateRange
  * @see FTimespan
  */
-class FDateTime
+struct FDateTime
 {
 public:
 
@@ -235,6 +235,19 @@ public:
 
 public:
 
+	/**
+	 * Exports the DateTime value to a string.
+	 *
+	 * @param ValueStr Will hold the string value.
+	 * @param DefaultValue The default value.
+	 * @param Parent Not used.
+	 * @param PortFlags Not used.
+	 * @param ExportRootScope Not used.
+	 * @return true on success, false otherwise.
+	 * @see ImportTextItem
+	 */
+	CORE_API bool ExportTextItem( FString& ValueStr, FDateTime const& DefaultValue, UObject* Parent, int32 PortFlags, UObject* ExportRootScope ) const;
+
 	/**
 	 * Gets the date part of this date.
 	 *
@@ -409,6 +422,18 @@ public:
 	 */
 	CORE_API int32 GetYear( ) const;
 
+	/**
+	 * Imports the DateTime value from a text buffer.
+	 *
+	 * @param Buffer The text buffer to import from.
+	 * @param PortFlags Not used.
+	 * @param Parent Not used.
+	 * @param ErrorText The output device for error logging.
+	 * @return true on success, false otherwise.
+	 * @see ExportTextItem
+	 */
+	CORE_API bool ImportTextItem( const TCHAR*& Buffer, int32 PortFlags, UObject* Parent, FOutputDevice* ErrorText );
+
 	/**
 	 * Gets whether this date's time is in the afternoon.
 	 *
@@ -612,6 +637,8 @@ public:
 
 public:
 
+	friend class UObject;
+
 	/**
 	 * Serializes the given date and time from or into the specified archive.
 	 *
@@ -634,31 +661,6 @@ public:
 		return GetTypeHash(DateTime.Ticks);
 	}
 
-	/**
-	 * Exports the DateTime value to a string.
-	 *
-	 * @param ValueStr Will hold the string value.
-	 * @param DefaultValue The default value.
-	 * @param Parent Not used.
-	 * @param PortFlags Not used.
-	 * @param ExportRootScope Not used.
-	 * @return true on success, false otherwise.
-	 * @see ImportTextItem
-	 */
-	CORE_API bool ExportTextItem( FString& ValueStr, FDateTime const& DefaultValue, UObject* Parent, int32 PortFlags, class UObject* ExportRootScope ) const;
-
-	/**
-	 * Imports the DateTime value from a text buffer.
-	 *
-	 * @param Buffer The text buffer to import from.
-	 * @param PortFlags Not used.
-	 * @param Parent Not used.
-	 * @param ErrorText The output device for error logging.
-	 * @return true on success, false otherwise.
-	 * @see ExportTextItem
-	 */
-	CORE_API bool ImportTextItem( const TCHAR*& Buffer, int32 PortFlags, class UObject* Parent, FOutputDevice* ErrorText );
-
 protected:
 
 	/**
@@ -671,8 +673,8 @@ protected:
 	 */
 	static const int32 DaysToMonth[];
 
-public:
+private:
 
-	// Holds the ticks in 100 nanoseconds resolution since January 1, 0001 A.D.
+	/** Holds the ticks in 100 nanoseconds resolution since January 1, 0001 A.D. */
 	int64 Ticks;
 };
diff --git a/Engine/Source/Runtime/Core/Public/Misc/Guid.h b/Engine/Source/Runtime/Core/Public/Misc/Guid.h
--- a/Engine/Source/Runtime/Core/Public/Misc/Guid.h
+++ b/Engine/Source/Runtime/Core/Public/Misc/Guid.h
@@ -7,6 +7,10 @@
 #pragma once
 
 
+class UObject;
+class FOutputDevice;
+
+
 namespace EGuidFormats
 {
 	/**
diff --git a/Engine/Source/Runtime/Core/Public/Misc/Timespan.h b/Engine/Source/Runtime/Core/Public/Misc/Timespan.h
--- a/Engine/Source/Runtime/Core/Public/Misc/Timespan.h
+++ b/Engine/Source/Runtime/Core/Public/Misc/Timespan.h
@@ -7,6 +7,10 @@
 #pragma once
 
 
+class UObject;
+class FOutputDevice;
+
+
 namespace ETimespan
 {
 	/** The number of timespan ticks per day. */
@@ -518,6 +522,8 @@ public:
 
 public:
 
+	friend class UObject;
+
 	/**
 	 * Serializes the given time span from or into the specified archive.
 	 *
@@ -555,7 +561,7 @@ protected:
 
 private:
 
-	// Holds the time span in 100 nanoseconds resolution.
+	/** Holds the time span in 100 nanoseconds resolution. */
 	int64 Ticks;
 };
 
diff --git a/Engine/Source/Runtime/CoreUObject/Classes/Object.h b/Engine/Source/Runtime/CoreUObject/Classes/Object.h
--- a/Engine/Source/Runtime/CoreUObject/Classes/Object.h
+++ b/Engine/Source/Runtime/CoreUObject/Classes/Object.h
@@ -677,7 +677,6 @@ public:
 USTRUCT(immutable, noexport, BlueprintType)
 struct FDateTime
 {
-	UPROPERTY()
 	int64 Ticks;
 };
 
@@ -687,7 +686,6 @@ struct FDateTime
 USTRUCT(immutable, noexport, BlueprintType)
 struct FTimespan
 {
-	UPROPERTY()
 	int64 Ticks;
 };
 
