Media: fixed enum forward declaration compile error on Mac

diff --git a/Engine/Source/Editor/DetailCustomizations/Private/MediaAssetCustomization.h b/Engine/Source/Editor/DetailCustomizations/Private/MediaAssetCustomization.h
--- a/Engine/Source/Editor/DetailCustomizations/Private/MediaAssetCustomization.h
+++ b/Engine/Source/Editor/DetailCustomizations/Private/MediaAssetCustomization.h
@@ -3,6 +3,9 @@
 #pragma once
 
 
+enum class EMediaPlaybackDirections;
+
+
 /**
  * Implements a details view customization for the UMediaAsset class.
  */
@@ -29,36 +32,36 @@ public:
 
 private:
 
-	// Callback for getting the text of the Duration text block.
+	/** Callback for getting the text of the Duration text block. */
 	FText HandleDurationTextBlockText( ) const;
 
-	// Callback for getting the text of the ForwardRates text block.
+	/** Callback for getting the text of the ForwardRates text block. */
 	FText HandleForwardRatesTextBlockText( ) const;
 
-	// Callback for getting the text of a supported playback rate text block.
-	FText HandleSupportedRatesTextBlockText( enum class EMediaPlaybackDirections Direction, bool Unthinned ) const;
+	/** Callback for getting the text of a supported playback rate text block. */
+	FText HandleSupportedRatesTextBlockText( EMediaPlaybackDirections Direction, bool Unthinned ) const;
 
-	// Callback for getting the text of the SupportsScrubbing text block.
+	/** Callback for getting the text of the SupportsScrubbing text block. */
 	FText HandleSupportsScrubbingTextBlockText( ) const;
 
-	// Callback for getting the text of the SupportsSeeking text block.
+	/** Callback for getting the text of the SupportsSeeking text block. */
 	FText HandleSupportsSeekingTextBlockText( ) const;
 
-	// Callback for getting the selected path in the URL picker widget.
+	/** Callback for getting the selected path in the URL picker widget. */
 	FString HandleUrlPickerFilePath( ) const;
 
-	// Callback for getting the file type filter for the URL picker.
+	/** Callback for getting the file type filter for the URL picker. */
 	FString HandleUrlPickerFileTypeFilter( ) const;
 
-	// Callback for picking a path in the URL picker.
+	/** Callback for picking a path in the URL picker. */
 	void HandleUrlPickerPathPicked( const FString& PickedPath );
 
-	// Callback for getting the visibility of warning icon for invalid URLs.
+	/** Callback for getting the visibility of warning icon for invalid URLs. */
 	EVisibility HandleUrlWarningIconVisibility( ) const;
 
 private:
 
-	/** The collection of media assets being customized */
+	/** The collection of media assets being customized. */
 	TArray<TWeakObjectPtr<UObject>> CustomizedMediaAssets;
 
 	/** Pointer to the URL property handle. */
