Fixed compliation error on mac.

diff --git a/Engine/Source/Editor/DetailCustomizations/Private/SlateBrushCustomization.cpp b/Engine/Source/Editor/DetailCustomizations/Private/SlateBrushCustomization.cpp
--- a/Engine/Source/Editor/DetailCustomizations/Private/SlateBrushCustomization.cpp
+++ b/Engine/Source/Editor/DetailCustomizations/Private/SlateBrushCustomization.cpp
@@ -974,7 +974,7 @@ class SSlateBrushStaticPreview : public SCompoundWidget
 		TemporaryBrush = *static_cast<FSlateBrush*>( RawData[0] );
 	}
 
-	const FSlateBrush* SSlateBrushStaticPreview::GetPropertyBrush() const
+	const FSlateBrush* GetPropertyBrush() const
 	{
 		return &TemporaryBrush;
 	}
