Fix "reset to default" on transform details customization

diff --git a/Engine/Source/Editor/DetailCustomizations/Private/ComponentTransformDetails.h b/Engine/Source/Editor/DetailCustomizations/Private/ComponentTransformDetails.h
--- a/Engine/Source/Editor/DetailCustomizations/Private/ComponentTransformDetails.h
+++ b/Engine/Source/Editor/DetailCustomizations/Private/ComponentTransformDetails.h
@@ -232,7 +232,11 @@ private:
 	/** @return The Z component of location */
 	TOptional<float> GetLocationZ() const { return CachedLocation.Z; }
 	/** @return The visibility of the "Reset to Default" button for the location component */
-	EVisibility GetLocationResetVisibility() const { return CachedLocation.X.GetValue() == 0.f && CachedLocation.Y.GetValue() == 0.f && CachedLocation.Z.GetValue() == 0.f ? EVisibility::Hidden : EVisibility::Visible; }
+	EVisibility GetLocationResetVisibility() const
+	{
+		// unset means multiple differing values, so show "Reset to Default" in that case
+		return CachedLocation.IsSet() && CachedLocation.X.GetValue() == 0.0f && CachedLocation.Y.GetValue() == 0.0f && CachedLocation.Z.GetValue() == 0.0f ? EVisibility::Hidden : EVisibility::Visible;
+	}
 
 	/** @return The X component of rotation */
 	TOptional<float> GetRotationX() const { return CachedRotation.X; }
@@ -241,7 +245,11 @@ private:
 	/** @return The Z component of rotation */
 	TOptional<float> GetRotationZ() const { return CachedRotation.Z; }
 	/** @return The visibility of the "Reset to Default" button for the rotation component */
-	EVisibility GetRotationResetVisibility() const { return CachedRotation.X.GetValue() == 0.f && CachedRotation.Y.GetValue() == 0.f && CachedRotation.Z.GetValue() == 0.f ? EVisibility::Hidden : EVisibility::Visible; }
+	EVisibility GetRotationResetVisibility() const
+	{
+		// unset means multiple differing values, so show "Reset to Default" in that case
+		return CachedRotation.IsSet() && CachedRotation.X.GetValue() == 0.0f && CachedRotation.Y.GetValue() == 0.0f && CachedRotation.Z.GetValue() == 0.0f ? EVisibility::Hidden : EVisibility::Visible;
+	}
 
 	/** @return The X component of scale */
 	TOptional<float> GetScaleX() const { return CachedScale.X; }
@@ -250,7 +258,11 @@ private:
 	/** @return The Z component of scale */
 	TOptional<float> GetScaleZ() const { return CachedScale.Z; }
 	/** @return The visibility of the "Reset to Default" button for the scale component */
-	EVisibility GetScaleResetVisibility() const { return CachedScale.X.GetValue() == 1.f && CachedScale.Y.GetValue() == 1.f && CachedScale.Z.GetValue() == 1.f ? EVisibility::Hidden : EVisibility::Visible; }
+	EVisibility GetScaleResetVisibility() const
+	{
+		// unset means multiple differing values, so show "Reset to Default" in that case
+		return CachedScale.IsSet() && CachedScale.X.GetValue() == 1.0f && CachedScale.Y.GetValue() == 1.0f && CachedScale.Z.GetValue() == 1.0f ? EVisibility::Hidden : EVisibility::Visible;
+	}
 
 private:
 	/** A vector where it may optionally be unset */
