Fix crash removing last keybinding

diff --git a/Engine/Source/Editor/DetailCustomizations/Private/KeyStructCustomization.cpp b/Engine/Source/Editor/DetailCustomizations/Private/KeyStructCustomization.cpp
--- a/Engine/Source/Editor/DetailCustomizations/Private/KeyStructCustomization.cpp
+++ b/Engine/Source/Editor/DetailCustomizations/Private/KeyStructCustomization.cpp
@@ -27,7 +27,7 @@ void FKeyStructCustomization::CustomizeHeader( TSharedRef<class IPropertyHandle>
 	TArray<void*> StructPtrs;
 	StructPropertyHandle->AccessRawData(StructPtrs);
 	check(StructPtrs.Num() != 0);
-	SelectedKey = (FKey*)StructPtrs[0];
+	FKey* SelectedKey = (FKey*)StructPtrs[0];
 
 	bool bMultipleValues = false;
 	for (int32 StructPtrIndex = 1; StructPtrIndex < StructPtrs.Num(); ++StructPtrIndex)
@@ -58,7 +58,29 @@ void FKeyStructCustomization::CustomizeHeader( TSharedRef<class IPropertyHandle>
 
 FKey FKeyStructCustomization::GetCurrentKey() const
 {
-	return *SelectedKey;
+	TArray<void*> StructPtrs;
+	PropertyHandle->AccessRawData(StructPtrs);
+	check(StructPtrs.Num() != 0);
+	FKey* SelectedKey = (FKey*)StructPtrs[0];
+
+	bool bMultipleValues = false;
+	for(int32 StructPtrIndex = 1; StructPtrIndex < StructPtrs.Num(); ++StructPtrIndex)
+	{
+		if(*(FKey*)StructPtrs[StructPtrIndex] != *SelectedKey)
+		{
+			bMultipleValues = true;
+			break;
+		}
+	}
+
+	if( !bMultipleValues && SelectedKey )
+	{
+		return *SelectedKey;
+	}
+	else
+	{
+		return FKey();
+	}
 }
 
 void FKeyStructCustomization::OnKeyChanged(TSharedPtr<FKey> SelectedKey)
diff --git a/Engine/Source/Editor/DetailCustomizations/Private/KeyStructCustomization.h b/Engine/Source/Editor/DetailCustomizations/Private/KeyStructCustomization.h
--- a/Engine/Source/Editor/DetailCustomizations/Private/KeyStructCustomization.h
+++ b/Engine/Source/Editor/DetailCustomizations/Private/KeyStructCustomization.h
@@ -34,7 +34,4 @@ private:
 
 	/** Holds a handle to the property being edited. */
 	TSharedPtr<IPropertyHandle> PropertyHandle;
-
-	/** Pointer to the first key being edited. */
-	FKey* SelectedKey;
 };
