Fix file pickers not showing anything in the details panel

diff --git a/Engine/Source/Editor/DetailCustomizations/Private/FilePathStructCustomization.cpp b/Engine/Source/Editor/DetailCustomizations/Private/FilePathStructCustomization.cpp
--- a/Engine/Source/Editor/DetailCustomizations/Private/FilePathStructCustomization.cpp
+++ b/Engine/Source/Editor/DetailCustomizations/Private/FilePathStructCustomization.cpp
@@ -19,12 +19,12 @@ void FFilePathStructCustomization::CustomizeChildren( TSharedRef<IPropertyHandle
 
 void FFilePathStructCustomization::CustomizeHeader( TSharedRef<IPropertyHandle> StructPropertyHandle, class FDetailWidgetRow& HeaderRow, IPropertyTypeCustomizationUtils& StructCustomizationUtils )
 {
-	PathProperty = StructPropertyHandle;
+	PathStringProperty = StructPropertyHandle->GetChildHandle("FilePath");
 
 	FString FileTypeFilter;
 
 	// construct file type filter
-	const FString& MetaData = PathProperty->GetMetaData(TEXT("FilePathFilter"));
+	const FString& MetaData = StructPropertyHandle->GetMetaData(TEXT("FilePathFilter"));
 
 	if (MetaData.IsEmpty())
 	{
@@ -64,7 +64,7 @@ void FFilePathStructCustomization::CustomizeHeader( TSharedRef<IPropertyHandle>
 FString FFilePathStructCustomization::HandleFilePathPickerFilePath( ) const
 {
 	FString FilePath;
-	PathProperty->GetValue(FilePath);
+	PathStringProperty->GetValue(FilePath);
 
 	return FilePath;
 }
@@ -72,7 +72,7 @@ FString FFilePathStructCustomization::HandleFilePathPickerFilePath( ) const
 
 void FFilePathStructCustomization::HandleFilePathPickerPathPicked( const FString& PickedPath )
 {
-	PathProperty->SetValueFromFormattedString(PickedPath);
+	PathStringProperty->SetValue(PickedPath);
 	FEditorDirectories::Get().SetLastDirectory(ELastDirectory::GENERIC_OPEN, FPaths::GetPath(PickedPath));
 }
 
diff --git a/Engine/Source/Editor/DetailCustomizations/Private/FilePathStructCustomization.h b/Engine/Source/Editor/DetailCustomizations/Private/FilePathStructCustomization.h
--- a/Engine/Source/Editor/DetailCustomizations/Private/FilePathStructCustomization.h
+++ b/Engine/Source/Editor/DetailCustomizations/Private/FilePathStructCustomization.h
@@ -37,7 +37,6 @@ private:
 	void HandleFilePathPickerPathPicked( const FString& PickedPath );
 
 private:
-
-	/** Pointer to the Path property handle. */
-	TSharedPtr<IPropertyHandle> PathProperty;
+	/** Pointer to the string that will be seet when changing the path */
+	TSharedPtr<IPropertyHandle> PathStringProperty;
 };
