Added url support to doc actors.

diff --git a/Engine/Source/Editor/DetailCustomizations/Private/DocumentationActorDetails.cpp b/Engine/Source/Editor/DetailCustomizations/Private/DocumentationActorDetails.cpp
--- a/Engine/Source/Editor/DetailCustomizations/Private/DocumentationActorDetails.cpp
+++ b/Engine/Source/Editor/DetailCustomizations/Private/DocumentationActorDetails.cpp
@@ -22,34 +22,25 @@ TSharedRef<IDetailCustomization> FDocumentationActorDetails::MakeInstance()
 
 void FDocumentationActorDetails::CustomizeDetails(IDetailLayoutBuilder& DetailBuilder)
 {
-	SelectedDocActor = nullptr; 
-	
-	const TArray< TWeakObjectPtr<UObject> >& SelectedObjects = DetailBuilder.GetDetailsView().GetSelectedObjects();
-	
-	// Find the first documentation actor -we should only open one documentation instance at a time.
-	for (int32 ObjectIndex = 0; ObjectIndex < SelectedObjects.Num(); ++ObjectIndex)
+	EDocumentationActorType::Type SelectedType = EDocumentationActorType::None;
+	TArray<TWeakObjectPtr<UObject>> Objects;
+	DetailBuilder.GetObjectsBeingCustomized(Objects);
+	if (Objects.Num() > 0)
 	{
-		if (ADocumentationActor* CurrentDocActor = Cast<ADocumentationActor>(SelectedObjects[ObjectIndex].Get()))
-		{
-			if (CurrentDocActor != nullptr)
-			{
-				if (CurrentDocActor->HasValidDocumentLink() == true)
-				{
-					SelectedDocActor = CurrentDocActor;
-					break;
-				}
-			}
-		}
+		SelectedDocumentationActor = Cast<ADocumentationActor>(Objects[0].Get());
 	}
+	PropertyHandle = DetailBuilder.GetProperty("DocumentLink");
 
 	// Add a button we can click on to open the documentation
 	IDetailCategoryBuilder& HelpCategory = DetailBuilder.EditCategory("Help Data");
 	HelpCategory.AddCustomRow("Help Documentation")
 		[
 			SNew(SButton)
-			.Text(LOCTEXT("HelpDocumentation", "Help Documentation"))
+			.Text(this, &FDocumentationActorDetails::OnGetButtonText)
+			.ToolTipText(this, &FDocumentationActorDetails::OnGetButtonTooltipText)
 			.HAlign(HAlign_Center)
 			.OnClicked(this, &FDocumentationActorDetails::OnHelpButtonClicked)
+			.IsEnabled(this, &FDocumentationActorDetails::IsButtonEnabled)
 		];
 	
 }
@@ -58,12 +49,68 @@ FReply FDocumentationActorDetails::OnHelpButtonClicked()
 {
 	FReply Handled = FReply::Unhandled();
 	
-	if (SelectedDocActor.IsValid() == true )
+	FString DocumentLink;
+	FPropertyAccess::Result Result = PropertyHandle->GetValue(DocumentLink);
+	if (Result == FPropertyAccess::Success)
 	{
-		Handled = SelectedDocActor->OpenDocumentLink() == true ? FReply::Handled() : FReply::Unhandled();
+		if (SelectedDocumentationActor.IsValid() == true)
+		{
+			Handled = SelectedDocumentationActor->OpenDocumentLink() == true ? FReply::Handled() : FReply::Unhandled();
+		}
 	}
 	return Handled;
 }
 
+FText FDocumentationActorDetails::OnGetButtonText() const
+{
+	FText ButtonText;
+	if (IsButtonEnabled() == true)
+	{
+		EDocumentationActorType::Type LinkType = SelectedDocumentationActor->GetLinkType();
+		if (LinkType == EDocumentationActorType::UDNLink)
+		{
+			ButtonText = LOCTEXT("HelpDocumentation", "Open Help Documentation");
+		}
+		else if (LinkType == EDocumentationActorType::URLLink)
+		{
+			ButtonText = LOCTEXT("HelpDocumentationURL", "Open Help URL");
+		}
+	}
+	return ButtonText;
+}
+
+FText FDocumentationActorDetails::OnGetButtonTooltipText() const
+{
+	FText Tooltip;
+	EDocumentationActorType::Type LinkType = SelectedDocumentationActor->GetLinkType();
+	switch (LinkType)
+	{
+		case EDocumentationActorType::UDNLink:
+			Tooltip = FText::Format(LOCTEXT("OpenUDNLink", "Opens the UDN link:{0}"), FText::FromString(SelectedDocumentationActor->DocumentLink));
+			break;
+		case EDocumentationActorType::URLLink:
+			Tooltip = FText::Format(LOCTEXT("OpenURLLink", "Opens the Web page:{0}"), FText::FromString(SelectedDocumentationActor->DocumentLink));
+			break;			
+		case EDocumentationActorType::None:
+			Tooltip = LOCTEXT("OpenNoString", "Enter a string in the link field to have this button open it");
+			break;
+		case EDocumentationActorType::InvalidLink:
+			Tooltip = LOCTEXT("OpenInvalid", "The link field contains invalid data");
+			break;			
+	}
+	return Tooltip;
+}
+
+bool FDocumentationActorDetails::IsButtonEnabled() const
+{
+	bool bResult = false;
+	if ((SelectedDocumentationActor.IsValid() == true) && (SelectedDocumentationActor->HasValidDocumentLink() == true))
+	{
+		EDocumentationActorType::Type LinkType = SelectedDocumentationActor->GetLinkType();
+		bResult = ((LinkType == EDocumentationActorType::UDNLink) || (LinkType == EDocumentationActorType::URLLink)) ? true : false;
+	}
+
+	return bResult;
+}
 
 #undef LOCTEXT_NAMESPACE
diff --git a/Engine/Source/Editor/DetailCustomizations/Private/DocumentationActorDetails.h b/Engine/Source/Editor/DetailCustomizations/Private/DocumentationActorDetails.h
--- a/Engine/Source/Editor/DetailCustomizations/Private/DocumentationActorDetails.h
+++ b/Engine/Source/Editor/DetailCustomizations/Private/DocumentationActorDetails.h
@@ -2,6 +2,8 @@
 
 #pragma once
 
+class ADocumentationActor;
+
 class FDocumentationActorDetails : public IDetailCustomization
 {
 public:
@@ -13,7 +15,20 @@ public:
 	
 	/* Handler for clicking the help button */
 	FReply OnHelpButtonClicked();
+
+	/* Handler to get the text for the button */
+	FText OnGetButtonText() const;
+	
+	/* Handler to get the text for the button tooltip  */
+	FText FDocumentationActorDetails::OnGetButtonTooltipText() const;
+
+	/* Handler to determine if the button is enabled (link is valid)  */
+	bool IsButtonEnabled() const;
+
 private:
-	/* The documentation actor that we are showing in the details panel */
-	TWeakObjectPtr<class ADocumentationActor>	SelectedDocActor;
+	/* The first documentation actor that we are showing in the details panel */
+	TWeakObjectPtr<ADocumentationActor>	SelectedDocumentationActor;
+
+	/* Handle to the string property */
+	TSharedPtr<IPropertyHandle> PropertyHandle;
 };
diff --git a/Engine/Source/Runtime/Engine/Classes/Engine/DocumentationActor.h b/Engine/Source/Runtime/Engine/Classes/Engine/DocumentationActor.h
--- a/Engine/Source/Runtime/Engine/Classes/Engine/DocumentationActor.h
+++ b/Engine/Source/Runtime/Engine/Classes/Engine/DocumentationActor.h
@@ -6,6 +6,21 @@
 #pragma once
 #include "DocumentationActor.generated.h"
 
+namespace EDocumentationActorType
+{
+	enum Type
+	{
+		// The link is unset
+		None,
+		// A UDN document link
+		UDNLink,
+		// A URL
+		URLLink,
+		// the link is invalid
+		InvalidLink
+	};
+}
+
 UCLASS(hidecategories = (Sprite, MaterialSprite, Actor, Transform, Tags, Materials, Rendering, Components, Blueprint, bject, Collision, Display, Rendering, Physics, Input, Lighting, Layers))
 class ENGINE_API ADocumentationActor : public AActor
 {
@@ -17,16 +32,33 @@ class ENGINE_API ADocumentationActor : public AActor
 	/* Is the document link this actor has valid */
 	bool HasValidDocumentLink() const;
 
+	// Actor interface
+#if WITH_EDITOR
+	virtual void PostEditChangeProperty(FPropertyChangedEvent& PropertyChangedEvent) override;
+
+	virtual void PostLoad() override;
+#endif
+	// End of Actor interface
+
 #if WITH_EDITORONLY_DATA
-	/** Scales forces applied from damage events. */
+	/** Link to a help document. */
 	UPROPERTY(Category = HelpDocumentation, EditAnywhere, BlueprintReadWrite, AdvancedDisplay)
 	FString DocumentLink; 
 	
  	UPROPERTY(Category = Sprite, VisibleAnywhere, BlueprintReadOnly)
 	TSubobjectPtr<class UMaterialBillboardComponent> Billboard;
-
 #endif
 
+	/** Determine the type of link contained in the DocumentLink member */
+	EDocumentationActorType::Type GetLinkType() const;
+	
+private:
+
+	/** Update the member variable that days what type of link we have */
+	void UpdateLinkType();
+
+	EDocumentationActorType::Type	LinkType;
+
 };
 
 
diff --git a/Engine/Source/Runtime/Engine/Private/DocumentationActor.cpp b/Engine/Source/Runtime/Engine/Private/DocumentationActor.cpp
--- a/Engine/Source/Runtime/Engine/Private/DocumentationActor.cpp
+++ b/Engine/Source/Runtime/Engine/Private/DocumentationActor.cpp
@@ -6,8 +6,9 @@
 #include "Components/MaterialBillboardComponent.h"
 #include "IDocumentation.h"
 #endif
+#include "Dialogs.h"
 
-
+#define LOCTEXT_NAMESPACE "DocumentationActor"
 
 ADocumentationActor::ADocumentationActor(const class FPostConstructInitializeProperties& PCIP)
 	: Super(PCIP)
@@ -29,7 +30,29 @@ bool ADocumentationActor::OpenDocumentLink() const
 	bool bOpened = false;
 
 #if WITH_EDITOR
-	bOpened = IDocumentation::Get()->Open(DocumentLink);
+	// Warn the user if they are opening a URL
+	if (GetLinkType() == EDocumentationActorType::URLLink)
+	{
+		FText Message = LOCTEXT("OpeningURLMessage", "You are about to open an external URL. This will open your web browser. Do you want to proceed?");
+		FText URLDialog = LOCTEXT("OpeningURLTitle", "Open external link");
+
+		FSuppressableWarningDialog::FSetupInfo Info(Message, URLDialog, "SupressOpenURLWarning");
+		Info.ConfirmText = LOCTEXT("OpenURL_yes", "Yes");
+		Info.CancelText = LOCTEXT("OpenURL_no", "No");
+		FSuppressableWarningDialog OpenURLWarning(Info);
+		if (OpenURLWarning.ShowModal() == FSuppressableWarningDialog::Cancel)
+		{
+			bOpened = false;
+		}
+		else
+		{
+			FPlatformProcess::LaunchURL(*DocumentLink, nullptr, nullptr);
+		}
+	}
+	else
+	{
+		bOpened = IDocumentation::Get()->Open(DocumentLink);
+	}
 #endif //WITH_EDITOR
 
 	return bOpened;
@@ -46,3 +69,44 @@ bool ADocumentationActor::HasValidDocumentLink() const
 	return bDocumentValid;
 }
 
+EDocumentationActorType::Type ADocumentationActor::GetLinkType() const
+{
+	return LinkType;
+}
+
+void ADocumentationActor::PostLoad()
+{
+	Super::PostLoad();
+
+	UpdateLinkType();
+}
+
+
+void ADocumentationActor::PostEditChangeProperty(FPropertyChangedEvent& PropertyChangedEvent)
+{
+	Super::PostEditChangeProperty(PropertyChangedEvent);
+	
+	UpdateLinkType();
+}
+
+void ADocumentationActor::UpdateLinkType()
+{
+	if (DocumentLink.IsEmpty() == true)
+	{
+		LinkType = EDocumentationActorType::None;
+	}
+	else if ((DocumentLink.StartsWith(TEXT("HTTP://")) == true) || (DocumentLink.StartsWith(TEXT("HTTPS://")) == true))
+	{
+		LinkType = EDocumentationActorType::URLLink;
+	}
+	else if (IDocumentation::Get()->PageExists(DocumentLink) == true )
+	{
+		LinkType = EDocumentationActorType::UDNLink;
+	}
+	else
+	{
+		LinkType = EDocumentationActorType::InvalidLink;
+	}
+}
+
+#undef LOCTEXT_NAMESPACE 
