Fix compile error on Mac.

diff --git a/Engine/Source/Editor/DetailCustomizations/Private/DocumentationActorDetails.h b/Engine/Source/Editor/DetailCustomizations/Private/DocumentationActorDetails.h
--- a/Engine/Source/Editor/DetailCustomizations/Private/DocumentationActorDetails.h
+++ b/Engine/Source/Editor/DetailCustomizations/Private/DocumentationActorDetails.h
@@ -20,7 +20,7 @@ public:
 	FText OnGetButtonText() const;
 	
 	/* Handler to get the text for the button tooltip  */
-	FText FDocumentationActorDetails::OnGetButtonTooltipText() const;
+	FText OnGetButtonTooltipText() const;
 
 	/* Handler to determine if the button is enabled (link is valid)  */
 	bool IsButtonEnabled() const;
