UMG/Editor - The slate brush customization is now a publically reusable customization.

diff --git a/Engine/Source/Editor/DetailCustomizations/Private/SlateBrushCustomization.h b/Engine/Source/Editor/DetailCustomizations/Private/SlateBrushCustomization.h
deleted file mode 100644
diff --git a/Engine/Source/Editor/DetailCustomizations/Public/Customizations/SlateBrushCustomization.h b/Engine/Source/Editor/DetailCustomizations/Public/Customizations/SlateBrushCustomization.h
new file mode 100644
--- /dev/null
+++ b/Engine/Source/Editor/DetailCustomizations/Public/Customizations/SlateBrushCustomization.h
@@ -0,0 +1,38 @@
+// Copyright 1998-2014 Epic Games, Inc. All Rights Reserved.
+
+#pragma once
+
+
+class FSlateBrushStructCustomization : public IPropertyTypeCustomization
+{
+public:
+	static TSharedRef<IPropertyTypeCustomization> MakeInstance();
+
+	FSlateBrushStructCustomization();
+
+	/** IPropertyTypeCustomization interface */
+	virtual void CustomizeHeader( TSharedRef<class IPropertyHandle> StructPropertyHandle, class FDetailWidgetRow& HeaderRow, IPropertyTypeCustomizationUtils& StructCustomizationUtils ) override;
+
+	virtual void CustomizeChildren( TSharedRef<class IPropertyHandle> StructPropertyHandle, class IDetailChildrenBuilder& StructBuilder, IPropertyTypeCustomizationUtils& StructCustomizationUtils ) override;              
+
+private:
+	/**
+	 * Get the Slate Brush tiling property row visibility
+	 */
+	EVisibility GetTilingPropertyVisibility() const;
+
+	/**
+	 *  Get the Slate Brush margin property row visibility
+	 */
+	EVisibility GetMarginPropertyVisibility() const;
+
+	/** Slate Brush DrawAs property */
+	TSharedPtr<IPropertyHandle> DrawAsProperty;
+
+	/** Error text to display if the resource object is not valid*/
+	TSharedPtr<SErrorText> ResourceErrorText;
+
+	/**  */
+	bool bSlimView;
+};
+
