- Fixed tooltip
- Fixed undo with rig set up

diff --git a/Engine/Source/Editor/DetailCustomizations/Private/RigDetails.cpp b/Engine/Source/Editor/DetailCustomizations/Private/RigDetails.cpp
--- a/Engine/Source/Editor/DetailCustomizations/Private/RigDetails.cpp
+++ b/Engine/Source/Editor/DetailCustomizations/Private/RigDetails.cpp
@@ -33,7 +33,6 @@ void FRigDetails::CustomizeDetails(IDetailLayoutBuilder& DetailBuilder)
 
 	ItemBeingEdited = Objects[0];
 
-
 	IDetailCategoryBuilder& NodeCategory = DetailBuilder.EditCategory("Node");
 	IDetailCategoryBuilder& RigControlCategory = DetailBuilder.EditCategory("Constraint Setup");
 
@@ -54,6 +53,8 @@ void FRigDetails::CustomizeDetails(IDetailLayoutBuilder& DetailBuilder)
 	if ( NumElement > 0 )
 	{
 		ParentSpaceOptionList.AddZeroed(NumElement);
+		// I need two per each - translation/orientation
+		ParentSpaceComboBoxes.AddZeroed(NumElement*2);
 	}
 
 	TSharedRef<FDetailArrayBuilder> NodeArrayBuilder = MakeShareable(new FDetailArrayBuilder(NodesPropertyHandle.ToSharedRef()));
@@ -255,6 +256,7 @@ void FRigDetails::GenerateRigControlArrayElementWidget(TSharedRef<IPropertyHandl
 	// create string list for picking parent node
 	// make sure you don't include itself and find what is curretn selected item
 	TArray<TSharedPtr<FString>> & ParentNodeOptions = ParentSpaceOptionList[ArrayIndex];
+
 	ParentNodeOptions.Empty();
 	ParentNodeOptions.Add(MakeShareable(new FString(URig::WorldNodeName.ToString())));
 	URig * Rig = Cast<URig>(ItemBeingEdited.Get());
@@ -333,10 +335,18 @@ void FRigDetails::GenerateRigControlArrayElementWidget(TSharedRef<IPropertyHandl
 					.WidthOverride(250)
 					.Content()
 					[
-						SNew(STextComboBox)
+
+						SAssignNew(ParentSpaceComboBoxes[ArrayIndex*2], SComboBox< TSharedPtr<FString> >)
 						.OptionsSource(&ParentNodeOptions)
 						.InitiallySelectedItem(ParentNodeOptions[ParentIndex_T])
 						.OnSelectionChanged(this, &FRigDetails::OnParentSpaceSelectionChanged, ParentNameProp_T)
+						.OnGenerateWidget(this, &FRigDetails::MakeItemWidget)
+						.OnComboBoxOpening(this, &FRigDetails::OnComboBoxOopening, ParentNameProp_T, ArrayIndex, true)
+						.HasDownArrow(true)
+						[
+							SNew(STextBlock)
+							.Text(this, &FRigDetails::GetSelectedTextLabel, ParentNameProp_T)
+						]
 					]
 				]
 			]
@@ -368,10 +378,17 @@ void FRigDetails::GenerateRigControlArrayElementWidget(TSharedRef<IPropertyHandl
 					.WidthOverride(250)
 					.Content()
 					[
-						SNew(STextComboBox)
+						SAssignNew(ParentSpaceComboBoxes[ArrayIndex*2+1], SComboBox< TSharedPtr<FString> >)
 						.OptionsSource(&ParentNodeOptions)
 						.InitiallySelectedItem(ParentNodeOptions[ParentIndex_R])
 						.OnSelectionChanged(this, &FRigDetails::OnParentSpaceSelectionChanged, ParentNameProp_R)
+						.OnGenerateWidget(this, &FRigDetails::MakeItemWidget)
+						.OnComboBoxOpening(this, &FRigDetails::OnComboBoxOopening, ParentNameProp_R, ArrayIndex, true)
+						.HasDownArrow(true)
+						[
+							SNew(STextBlock)
+							.Text(this, &FRigDetails::GetSelectedTextLabel, ParentNameProp_R)
+						]
 					]
 				]
 			]
@@ -412,4 +429,45 @@ FReply FRigDetails::OnSetAllToParent()
 
 	return FReply::Handled();
 }
+
+/** Called to create a widget for each string */
+TSharedRef<SWidget> FRigDetails::MakeItemWidget(TSharedPtr<FString> StringItem)
+{
+	check(StringItem.IsValid());
+
+	return SNew(STextBlock)
+		.Text(*StringItem.Get());
+}
+/** Helper method to get the text for a given item in the combo box */
+FString FRigDetails::GetSelectedTextLabel(TSharedRef<IPropertyHandle> ParentSpacePropertyHandle) const
+{
+	FString DisplayText;
+
+	check (ParentSpacePropertyHandle->GetValueAsDisplayString(DisplayText) != FPropertyAccess::Fail);
+
+	return DisplayText;
+}
+
+void FRigDetails::OnComboBoxOopening(TSharedRef<IPropertyHandle> ParentSpacePropertyHandle, int32 ArrayIndex, bool bTranslation)
+{
+	FString PropertyValue = GetSelectedTextLabel(ParentSpacePropertyHandle);
+
+	// now find exact data
+	TArray<TSharedPtr<FString>> & ParentOptions = ParentSpaceOptionList[ArrayIndex];
+	TSharedPtr<FString> SelectedItem;
+	for (auto Option : ParentOptions)
+	{
+		if (*Option.Get() == PropertyValue)
+		{
+			SelectedItem = Option;
+			break;
+		}
+	}
+
+	int32 ComboBoxIndex = (bTranslation)? ArrayIndex*2 : ArrayIndex*2+1;
+	TSharedPtr< SComboBox<TSharedPtr<FString>> > ComboBox = ParentSpaceComboBoxes[ComboBoxIndex];
+
+	ComboBox->SetSelectedItem(SelectedItem);
+}
+
 #undef LOCTEXT_NAMESPACE
diff --git a/Engine/Source/Editor/DetailCustomizations/Private/RigDetails.h b/Engine/Source/Editor/DetailCustomizations/Private/RigDetails.h
--- a/Engine/Source/Editor/DetailCustomizations/Private/RigDetails.h
+++ b/Engine/Source/Editor/DetailCustomizations/Private/RigDetails.h
@@ -25,6 +25,7 @@ private:
 	// array custom boxes - these will stay around as long as this window is up
 	TArray<TSharedPtr<SEditableTextBox>>	DisplayNameTextBoxes;
 	TArray<TArray<TSharedPtr<FString>>>		ParentSpaceOptionList;
+	TArray<TSharedPtr<SComboBox< TSharedPtr<FString> >>>	ParentSpaceComboBoxes;
 
 	/** we only support one item */
 	TWeakObjectPtr<UObject> ItemBeingEdited;
@@ -37,6 +38,11 @@ private:
 
 	// combo box handler
 	void OnParentSpaceSelectionChanged(TSharedPtr<FString> SelectedItem, ESelectInfo::Type SelectInfo, TSharedRef<IPropertyHandle> ParentSpacePropertyHandle);
+	/** Called to create a widget for each string */
+	TSharedRef<SWidget> MakeItemWidget(TSharedPtr<FString> StringItem);
+	/** Helper method to get the text for a given item in the combo box */
+	FString GetSelectedTextLabel(TSharedRef<IPropertyHandle> ParentSpacePropertyHandle) const;
+	void OnComboBoxOopening(TSharedRef<IPropertyHandle> ParentSpacePropertyHandle, int32 ArrayIndex, bool bTranslation);
 
 	// button handlers
 	FReply OnSetAllToWorld();
