UMG - Slate color customization now has a way to toggle the 'Use Foreground' flag using the new inherit checkbox.

diff --git a/Engine/Source/Editor/DetailCustomizations/Private/SlateColorCustomization.cpp b/Engine/Source/Editor/DetailCustomizations/Private/SlateColorCustomization.cpp
--- a/Engine/Source/Editor/DetailCustomizations/Private/SlateColorCustomization.cpp
+++ b/Engine/Source/Editor/DetailCustomizations/Private/SlateColorCustomization.cpp
@@ -10,36 +10,90 @@ TSharedRef<IPropertyTypeCustomization> FSlateColorCustomization::MakeInstance()
 
 void FSlateColorCustomization::CustomizeHeader( TSharedRef<class IPropertyHandle> StructPropertyHandle, class FDetailWidgetRow& HeaderRow, IPropertyTypeCustomizationUtils& StructCustomizationUtils )
 {
-	uint32 NumChildren;
-	StructPropertyHandle->GetNumChildren( NumChildren );
+	static const FName ColorUseRuleKey(TEXT("ColorUseRule"));
+	static const FName SpecifiedColorKey(TEXT("SpecifiedColor"));
 
-	for( uint32 ChildIndex = 0; ChildIndex < NumChildren; ++ChildIndex )
-	{
-		const TSharedRef< IPropertyHandle > ChildHandle = StructPropertyHandle->GetChildHandle( ChildIndex ).ToSharedRef();
-
-		if ( ChildHandle->GetProperty()->GetName() == TEXT("ColorUseRule") )
-		{
-			ColorRuleHandle = ChildHandle;
-		}
-		else if ( ChildHandle->GetProperty()->GetName() == TEXT("SpecifiedColor") )
-		{
-			SpecifiedColorHandle = ChildHandle;
-		}
-	}
+	ColorRuleHandle = StructPropertyHandle->GetChildHandle(ColorUseRuleKey);
+	SpecifiedColorHandle = StructPropertyHandle->GetChildHandle(SpecifiedColorKey);
 
-	check( ColorRuleHandle.IsValid() );
-	check( SpecifiedColorHandle.IsValid() );
+	check(ColorRuleHandle.IsValid());
+	check(SpecifiedColorHandle.IsValid());
 }
 
 void FSlateColorCustomization::CustomizeChildren( TSharedRef<class IPropertyHandle> StructPropertyHandle, class IDetailChildrenBuilder& StructBuilder, IPropertyTypeCustomizationUtils& StructCustomizationUtils )
 {
 	SpecifiedColorHandle->SetOnPropertyValueChanged( FSimpleDelegate::CreateSP( SharedThis( this ), &FSlateColorCustomization::OnValueChanged ) );
-	StructBuilder.AddChildProperty( SpecifiedColorHandle.ToSharedRef() )
-		.DisplayName( StructPropertyHandle->GetPropertyDisplayName() )
-		.ToolTip( StructPropertyHandle->GetToolTipText() );
+	
+	IDetailPropertyRow& ColorRow = StructBuilder.AddChildProperty(SpecifiedColorHandle.ToSharedRef());
+
+	TSharedPtr<SWidget> NameWidget;
+	TSharedPtr<SWidget> ValueWidget;
+	FDetailWidgetRow Row;
+	ColorRow.GetDefaultWidgets(NameWidget, ValueWidget, Row);
+
+	ColorRow.CustomWidget(/*bShowChildren*/ true)
+		.NameContent()
+		[
+			SNew(STextBlock)
+			.Font(IDetailLayoutBuilder::GetDetailFont())
+			.Text(StructPropertyHandle->GetPropertyDisplayName())
+			.ToolTipText(StructPropertyHandle->GetToolTipText())
+		]
+		.ValueContent()
+		.MinDesiredWidth(250.0f)
+		.MaxDesiredWidth(250.0f)
+		[
+			SNew(SHorizontalBox)
+
+			+ SHorizontalBox::Slot()
+			.FillWidth(1.0f)
+			[
+				ValueWidget.ToSharedRef()
+			]
+
+			+ SHorizontalBox::Slot()
+			.AutoWidth()
+			.Padding(2, 0, 0, 0)
+			[
+				SNew(SCheckBox)
+				.IsChecked(this, &FSlateColorCustomization::GetForegroundCheckState)
+				.OnCheckStateChanged(this, &FSlateColorCustomization::HandleForegroundChanged)
+				[
+					SNew(STextBlock)
+					.Font(IDetailLayoutBuilder::GetDetailFont())
+					.Text(NSLOCTEXT("SlateColorCustomization", "Inherit", "Inherit"))
+					.ToolTipText(NSLOCTEXT("SlateColorCustomization", "InheritToolTip", "Uses the foreground color inherited down the widget hierarchy"))
+				]
+			]
+		];
 }
 
 void FSlateColorCustomization::OnValueChanged()
 {
 	ColorRuleHandle->SetValueFromFormattedString(TEXT("UseColor_Specified"));
 }
+
+ESlateCheckBoxState::Type FSlateColorCustomization::GetForegroundCheckState() const
+{
+	FString ColorRuleValue;
+	ColorRuleHandle->GetValueAsFormattedString(ColorRuleValue);
+
+	if ( ColorRuleValue == TEXT("UseColor_Foreground") )
+	{
+		return ESlateCheckBoxState::Checked;
+	}
+
+	return ESlateCheckBoxState::Unchecked;
+}
+
+void FSlateColorCustomization::HandleForegroundChanged(ESlateCheckBoxState::Type CheckedState)
+{
+	if ( CheckedState == ESlateCheckBoxState::Checked )
+	{
+		ColorRuleHandle->SetValueFromFormattedString(TEXT("UseColor_Foreground"));
+	}
+	else
+	{
+		ColorRuleHandle->SetValueFromFormattedString(TEXT("UseColor_Specified"));
+	}
+}
diff --git a/Engine/Source/Editor/DetailCustomizations/Private/SlateColorCustomization.h b/Engine/Source/Editor/DetailCustomizations/Private/SlateColorCustomization.h
--- a/Engine/Source/Editor/DetailCustomizations/Private/SlateColorCustomization.h
+++ b/Engine/Source/Editor/DetailCustomizations/Private/SlateColorCustomization.h
@@ -19,6 +19,10 @@ private:
 	 */
 	virtual void OnValueChanged();
 
+	ESlateCheckBoxState::Type GetForegroundCheckState() const;
+
+	void HandleForegroundChanged(ESlateCheckBoxState::Type CheckedState);
+
 private:
 
 	TSharedPtr<IPropertyHandle> ColorRuleHandle;
