UMG - The slate brush customization now shows images differently than border/box in the little preview image.

diff --git a/Engine/Source/Editor/DetailCustomizations/Private/SlateBrushCustomization.cpp b/Engine/Source/Editor/DetailCustomizations/Private/SlateBrushCustomization.cpp
--- a/Engine/Source/Editor/DetailCustomizations/Private/SlateBrushCustomization.cpp
+++ b/Engine/Source/Editor/DetailCustomizations/Private/SlateBrushCustomization.cpp
@@ -950,17 +950,37 @@ class SSlateBrushStaticPreview : public SCompoundWidget
 
 	SLATE_END_ARGS()
 
-	void Construct( const FArguments& InArgs, TSharedPtr<IPropertyHandle> InResourceObjectProperty )
+	void Construct(const FArguments& InArgs, TSharedPtr<IPropertyHandle> InResourceObjectProperty)
 	{
 		ResourceObjectProperty = InResourceObjectProperty;
 
 		ChildSlot
 		[
-			SNew(SBorder)
-			.BorderImage(this, &SSlateBrushStaticPreview::GetPropertyBrush)
+			SNew(SHorizontalBox)
+
+			+ SHorizontalBox::Slot()
+			.FillWidth(1.0f)
+			[
+				SNew(SBorder)
+				.Visibility(this, &SSlateBrushStaticPreview::GetPreviewVisibilityBorder)
+				.BorderImage(this, &SSlateBrushStaticPreview::GetPropertyBrush)
+				[
+					SNew(SSpacer)
+					.Size(FVector2D(1, 1))
+				]
+			]
+
+			+ SHorizontalBox::Slot()
+			.AutoWidth()
 			[
-				SNew(SSpacer)
-				.Size(FVector2D(1,1))
+				SNew(SBox)
+				.WidthOverride(18.0f)
+				.HeightOverride(18.0f)
+				[
+					SNew(SImage)
+					.Visibility(this, &SSlateBrushStaticPreview::GetPreviewVisibilityImage)
+					.Image(this, &SSlateBrushStaticPreview::GetPropertyBrush)
+				]
 			]
 		];
 	}
@@ -974,11 +994,23 @@ class SSlateBrushStaticPreview : public SCompoundWidget
 		TemporaryBrush = *static_cast<FSlateBrush*>( RawData[0] );
 	}
 
+private:
+
 	const FSlateBrush* GetPropertyBrush() const
 	{
 		return &TemporaryBrush;
 	}
 
+	EVisibility GetPreviewVisibilityBorder() const
+	{
+		return TemporaryBrush.DrawAs == ESlateBrushDrawType::Image ? EVisibility::Collapsed : EVisibility::Visible;
+	}
+
+	EVisibility GetPreviewVisibilityImage() const
+	{
+		return TemporaryBrush.DrawAs == ESlateBrushDrawType::Image ? EVisibility::Visible : EVisibility::Collapsed;
+	}
+
 private:
 
 	/**
