Editor: Source control integration and styling improvements for hardware targeting settings (WIP)
#Unshelved from pending changelist '2302592' and modified

diff --git a/Engine/Source/Editor/DetailCustomizations/Private/HardwareTargetingSettingsDetails.cpp b/Engine/Source/Editor/DetailCustomizations/Private/HardwareTargetingSettingsDetails.cpp
--- a/Engine/Source/Editor/DetailCustomizations/Private/HardwareTargetingSettingsDetails.cpp
+++ b/Engine/Source/Editor/DetailCustomizations/Private/HardwareTargetingSettingsDetails.cpp
@@ -4,6 +4,7 @@
 
 #include "HardwareTargetingSettingsDetails.h"
 #include "HardwareTargetingModule.h"
+#include "ISourceControlModule.h"
 
 #define LOCTEXT_NAMESPACE "FHardwareTargetingSettingsDetails"
 
@@ -12,21 +13,38 @@ TSharedRef<IDetailCustomization> FHardwareTargetingSettingsDetails::MakeInstance
 	return MakeShareable(new FHardwareTargetingSettingsDetails);
 }
 
-void FHardwareTargetingSettingsDetails::CustomizeDetails(IDetailLayoutBuilder& DetailBuilder)
-{	
-	IDetailCategoryBuilder& HardwareTargetingCategory = DetailBuilder.EditCategory(TEXT("Hardware Targeting"));
+class SRequiredDefaultConfig : public SCompoundWidget
+{
+public:
 	
-	IHardwareTargetingModule& HardwareTargeting = IHardwareTargetingModule::Get();
-	HardwareTargetingCategory.AddCustomRow(FString())
-	[
-		SNew(SBorder)
-		.BorderImage(FEditorStyle::GetBrush("NoBorder"))
-		.Padding(0)
-		.Visibility_Static([]()->EVisibility{
-			auto* Settings = GetMutableDefault<UHardwareTargetingSettings>();
-			return Settings->HasPendingChanges() ? EVisibility::Visible : EVisibility::Collapsed;
-		})
+	SLATE_BEGIN_ARGS(SRequiredDefaultConfig) {}
+
+	SLATE_END_ARGS()
+
+	~SRequiredDefaultConfig()
+	{
+		GetMutableDefault<UHardwareTargetingSettings>()->OnSettingChanged().RemoveAll(this);
+	}
+	
+public:
+
+	TMap<TWeakObjectPtr<UObject>, TSharedPtr<SRichTextBlock>> SettingRegions;
+
+	void Construct(const FArguments& InArgs)
+	{
+		LastStatusUpdate = 0;
+		GetMutableDefault<UHardwareTargetingSettings>()->OnSettingChanged().AddRaw(this, &SRequiredDefaultConfig::Update);
+
+		auto ApplyNow = []{
+			IHardwareTargetingModule& Module = IHardwareTargetingModule::Get();
+			Module.ApplyHardwareTargetingSettings();
+
+			return FReply::Handled();
+		};
+
+		ChildSlot
 		[
+
 			SNew(SVerticalBox)
 			+ SVerticalBox::Slot()
 			.AutoHeight()
@@ -47,92 +65,366 @@ void FHardwareTargetingSettingsDetails::CustomizeDetails(IDetailLayoutBuilder& D
 				[
 					SNew(SButton)
 					.Text(LOCTEXT("ApplyNow", "Apply Now"))
-					.OnClicked_Static([](){
-				
-						IHardwareTargetingModule& Module = IHardwareTargetingModule::Get();
-						Module.ApplyHardwareTargetingSettings();
-
-						return FReply::Handled();
-					})
+					.IsEnabled(this, &SRequiredDefaultConfig::CanApply)
+					.OnClicked_Static(ApplyNow)
 				]
 			]
+			
 			+ SVerticalBox::Slot()
-			.FillHeight(1.0f)
+			.AutoHeight()
 			[
+				SAssignNew(SettingsDescriptions, SUniformGridPanel)
+			]
+		];
+	}
+
+	static EVisibility GetAnyPendingChangesVisibility()
+	{
+		return GetMutableDefault<UHardwareTargetingSettings>()->HasPendingChanges() ? EVisibility::Visible : EVisibility::Collapsed;
+	}
+
+	void Initialize(IDetailLayoutBuilder& DetailBuilder, IDetailCategoryBuilder& PendingChangesCategory)
+	{
+		auto IsVisible = []()->EVisibility{
+			auto* Settings = GetMutableDefault<UHardwareTargetingSettings>();
+			return Settings->HasPendingChanges() ? EVisibility::Visible : EVisibility::Collapsed;
+		};
+
+		FText CategoryHeaderTooltip = LOCTEXT("CategoryHeaderTooltip", "List of properties modified in this project setting category");
+
+		IHardwareTargetingModule& Module = IHardwareTargetingModule::Get();
+		for (const FModifiedDefaultConfig& Settings : Module.GetPendingSettingsChanges())
+		{
+			TSharedRef<SRichTextBlock> EditPropertiesBlock =
 				SNew(SRichTextBlock)
-				.AutoWrapText(true)
+				.AutoWrapText(false)
 				.Justification(ETextJustify::Left)
-				.Text(TAttribute<FText>::Create(TAttribute<FText>::FGetter::CreateStatic([](){
-					return IHardwareTargetingModule::Get().QueryReadableDescriptionOfHardwareTargetingSettings();
-				})))
 				.TextStyle(FEditorStyle::Get(), "HardwareTargets.Normal")
- 				.DecoratorStyleSet(&FEditorStyle::Get())
+				.DecoratorStyleSet(&FEditorStyle::Get());
+
+			SettingRegions.Add(Settings.SettingsObject, EditPropertiesBlock);
+
+			FDetailWidgetRow& CategoryRow = PendingChangesCategory.AddCustomRow(Settings.CategoryHeading.ToString())
+				.Visibility(TAttribute<EVisibility>::Create(&SRequiredDefaultConfig::GetAnyPendingChangesVisibility))
+				.NameContent()
+				[
+					SNew(STextBlock)
+					.Text(Settings.CategoryHeading)
+					.ToolTipText(CategoryHeaderTooltip)
+					.Font(IDetailLayoutBuilder::GetDetailFont())
+				]
+				.ValueContent()
+				.MaxDesiredWidth(300.0f)
+				[
+					EditPropertiesBlock
+				];
+		}
+	}
+
+	bool CanApply() const
+	{
+		for (const auto& Pair : WritableStatus)
+		{
+			if (Pair.Value != EWritableStatus::Writable)
+			{
+				return false;
+			}
+		}
+
+		return true;
+	}
+
+	FString GetDefaultConfigPath(UObject* Settings) const
+	{
+		return FPaths::ConvertRelativePathToFull(FString::Printf(TEXT("%sDefault%s.ini"), *FPaths::SourceConfigDir(), *Settings->GetClass()->ClassConfigName.ToString()));
+	}
+
+	virtual void Tick(const FGeometry& AllottedGeometry, const double InCurrentTime, const float InDeltaTime)
+	{
+		if (!ISourceControlModule::Get().IsEnabled())
+		{
+			return;
+		}
+
+		const double CheckStatusIntervalS = 2;
+		if (InCurrentTime - LastStatusUpdate > CheckStatusIntervalS)
+		{
+			LastStatusUpdate = InCurrentTime;
+
+			for (auto& Pair : WritableStatus)
+			{
+				if (Pair.Key.IsValid() && Pair.Value == EWritableStatus::Unknown)
+				{
+					ISourceControlProvider& SourceControlProvider = ISourceControlModule::Get().GetProvider();
+					FSourceControlStatePtr SourceControlState = SourceControlProvider.GetState(GetDefaultConfigPath(Pair.Key.Get()), EStateCacheUsage::Use);
+					if (SourceControlState.IsValid())
+					{
+						Pair.Value = SourceControlState->IsCheckedOut() ? EWritableStatus::Writable : EWritableStatus::NotWritable;
+					}
+				}
+			}
+		}
+	}
+
+	bool AreSettingsWritable(TWeakObjectPtr<UObject> Settings) const
+	{
+		if (auto* Status = WritableStatus.Find(Settings))
+		{
+			if (*Status == EWritableStatus::Writable)
+			{
+				return true;
+			}
+		}
+		return false;
+	}
+
+	EVisibility GetWritableHintVisibility(TWeakObjectPtr<UObject> Settings) const
+	{
+		if (const auto* Status = WritableStatus.Find(Settings))
+		{
+			if (*Status == EWritableStatus::NotWritable)
+			{
+				return EVisibility::Visible;
+			}
+		}
+		return EVisibility::Collapsed;
+	}
+
+	EVisibility GetWritableThrobberVisibility(TWeakObjectPtr<UObject> Settings) const
+	{
+		if (const auto* Status = WritableStatus.Find(Settings))
+		{
+			if (*Status == EWritableStatus::Unknown)
+			{
+				return EVisibility::Visible;
+			}
+		}
+		return EVisibility::Collapsed;
+	}
+
+	void Update()
+	{
+		WritableStatus.Empty();
+		SettingsDescriptions->ClearChildren();
+
+		IHardwareTargetingModule& Module = IHardwareTargetingModule::Get();
+		int32 SlotIndex = 0;
+
+		// Run thru the settings and push changes to the existing settings, as well as build a list of inis that will need to be edited
+		TSet<FString> SeenConfigFiles;
+		for (const FModifiedDefaultConfig& Settings : Module.GetPendingSettingsChanges())
+		{
+			if (!Settings.SettingsObject.IsValid())
+			{
+				continue;
+			}
+
+			SettingRegions.FindChecked(Settings.SettingsObject)->SetText(Settings.Description);
+
+			if (!Settings.SettingsObject->GetClass()->HasAnyClassFlags(CLASS_Config | CLASS_DefaultConfig))
+			{
+				continue;
+			}
+			
+			EWritableStatus& Status = WritableStatus.Add(Settings.SettingsObject, EWritableStatus::Unknown);
+			FString ConfigFile = GetDefaultConfigPath(Settings.SettingsObject.Get());
+
+			if (!SeenConfigFiles.Contains(ConfigFile))
+			{
+				SeenConfigFiles.Add(ConfigFile);
+
+				if (ISourceControlModule::Get().IsEnabled())
+				{
+					ISourceControlModule::Get().QueueStatusUpdate(ConfigFile);
+
+					ISourceControlProvider& SourceControlProvider = ISourceControlModule::Get().GetProvider();
+					FSourceControlStatePtr SourceControlState = SourceControlProvider.GetState(ConfigFile, EStateCacheUsage::Use);
+					if (SourceControlState.IsValid())
+					{
+						Status = SourceControlState->IsCheckedOut() ? EWritableStatus::Writable : EWritableStatus::NotWritable;
+					}
+				}
+				else
+				{
+					const bool bIsWritable = !FPaths::FileExists(ConfigFile) || !IFileManager::Get().IsReadOnly(*ConfigFile);
+					Status = bIsWritable ? EWritableStatus::Writable : EWritableStatus::NotWritable;
+				}
+
+				SettingsDescriptions->AddSlot(0, SlotIndex)
+				[
+					SNew(SVerticalBox)
+
+					+ SVerticalBox::Slot()
+					.AutoHeight()
+					[
+						SNew(STextBlock)
+						.Text(FPaths::GetCleanFilename(ConfigFile))
+					]
+				 
+					+ SVerticalBox::Slot()
+					.AutoHeight()
+					[
+						SNew(STextBlock)
+						.Visibility(this, &SRequiredDefaultConfig::GetWritableHintVisibility, Settings.SettingsObject)
+						.Text(LOCTEXT("NotWritable", "This file is not writable"))
+					]
+				 
+					+ SVerticalBox::Slot()
+					.AutoHeight()
+					[
+						SNew(SBorder)
+						.Visibility(this, &SRequiredDefaultConfig::GetWritableThrobberVisibility, Settings.SettingsObject)
+						.BorderImage(FEditorStyle::GetBrush("NoBorder"))
+						.Padding(0)
+						[
+							SNew(SHorizontalBox)
+
+							+ SHorizontalBox::Slot()
+							.AutoWidth()
+							[
+								SNew(SThrobber)
+
+							]
+							+ SHorizontalBox::Slot()
+							.AutoWidth()
+							[
+								SNew(STextBlock)
+								.Text(LOCTEXT("CheckingFileStatus", "Checking file status"))
+							]
+						]
+					]
+				];
+			}
+
+			++SlotIndex;
+		}
+	}
+	
+	enum class EWritableStatus
+	{
+		Unknown, Writable, NotWritable
+	};
+	TMap<TWeakObjectPtr<UObject>, EWritableStatus> WritableStatus;
+
+	TSharedPtr<SUniformGridPanel> SettingsDescriptions;
+
+	double LastStatusUpdate;
+};
+
+void FHardwareTargetingSettingsDetails::CustomizeDetails(IDetailLayoutBuilder& DetailBuilder)
+{	
+	IDetailCategoryBuilder& HardwareTargetingCategory = DetailBuilder.EditCategory(TEXT("Hardware Targeting"));
+	IDetailCategoryBuilder& PendingChangesCategory = DetailBuilder.EditCategory(TEXT("Pending Changes"));
+
+	auto AnyPendingChangesVisible = []()->EVisibility{
+		auto* Settings = GetMutableDefault<UHardwareTargetingSettings>();
+		return Settings->HasPendingChanges() ? EVisibility::Visible : EVisibility::Collapsed;
+	};
+	auto NoPendingChangesVisible = []()->EVisibility{
+		auto* Settings = GetMutableDefault<UHardwareTargetingSettings>();
+		return Settings->HasPendingChanges() ? EVisibility::Collapsed : EVisibility::Visible;
+	};
+	//auto ModifiedText = []{
+	//	return IHardwareTargetingModule::Get().QueryReadableDescriptionOfHardwareTargetingSettings();
+	//};
+
+	TSharedRef<SRequiredDefaultConfig> ConfigWidget = SNew(SRequiredDefaultConfig);
+	IHardwareTargetingModule& HardwareTargeting = IHardwareTargetingModule::Get();
+	PendingChangesCategory.AddCustomRow(FString())
+	[
+		SNew(SVerticalBox)
+		+ SVerticalBox::Slot()
+		[
+			SNew(SBox)
+			.Visibility_Static(AnyPendingChangesVisible)
+			[
+				ConfigWidget
+			]
+		]
+		+SVerticalBox::Slot()
+		[
+			SNew(SBox)
+			.Visibility_Static(NoPendingChangesVisible)
+			[
+				SNew(STextBlock)
+				.Font(DetailBuilder.GetDetailFont())
+				.Text(LOCTEXT("NoPendingChangesMessage", "There are no pending settings changes."))
 			]
 		]
 	];
 
+	ConfigWidget->Initialize(DetailBuilder, PendingChangesCategory);
 
 	auto PropertyName = GET_MEMBER_NAME_CHECKED(UHardwareTargetingSettings, TargetedHardwareClass);
 	DetailBuilder.HideProperty(PropertyName);
 
 	TSharedRef<IPropertyHandle> Property = DetailBuilder.GetProperty(PropertyName);
-	HardwareTargetingCategory.AddCustomRow(LOCTEXT("HardwareTargetingOption", "Targeted Hardware:").ToString())
-	.NameContent()
-	[
-		SNew(STextBlock)
-		.Text(FText::FromString(Property->GetPropertyDisplayName()))
-		.Font(DetailBuilder.GetDetailFont())
-	]
-	.ValueContent()
-	[
-		SNew(SHorizontalBox)
+	{
+		auto SetPropertyValue = [](EHardwareClass::Type NewValue, TSharedRef<IPropertyHandle> Property){
+			Property->SetValue(uint8(NewValue));
+		};
+		auto GetPropertyValue = [](TSharedRef<IPropertyHandle> Property){
+			uint8 Value = 0;
+			Property->GetValue(Value);
+			return EHardwareClass::Type(Value);
+		};
 
-		+ SHorizontalBox::Slot()
-		.AutoWidth()
+		HardwareTargetingCategory.AddCustomRow(LOCTEXT("HardwareTargetingOption", "Targeted Hardware:").ToString())
+		.NameContent()
 		[
-			HardwareTargeting.MakeHardwareClassTargetCombo(
-				FOnHardwareClassChanged::CreateStatic([](EHardwareClass::Type NewValue, TSharedRef<IPropertyHandle> Property){
-					Property->SetValue(uint8(NewValue));
-				}, Property),
-				TAttribute<EHardwareClass::Type>::Create(TAttribute<EHardwareClass::Type>::FGetter::CreateStatic([](TSharedRef<IPropertyHandle> Property){
-					uint8 Value = 0;
-					Property->GetValue(Value);
-					return EHardwareClass::Type(Value);
-				}, Property))
-			)
+			SNew(STextBlock)
+			.Text(FText::FromString(Property->GetPropertyDisplayName()))
+			.Font(DetailBuilder.GetDetailFont())
 		]
-	];
+		.ValueContent()
+		[
+			SNew(SHorizontalBox)
+
+			+ SHorizontalBox::Slot()
+			.AutoWidth()
+			[
+				HardwareTargeting.MakeHardwareClassTargetCombo(
+					FOnHardwareClassChanged::CreateStatic(SetPropertyValue, Property),
+					TAttribute<EHardwareClass::Type>::Create(TAttribute<EHardwareClass::Type>::FGetter::CreateStatic(GetPropertyValue, Property))
+				)
+			]
+		];
+	}
 
 	PropertyName = GET_MEMBER_NAME_CHECKED(UHardwareTargetingSettings, DefaultGraphicsPerformance);
 	DetailBuilder.HideProperty(PropertyName);
 
 	Property = DetailBuilder.GetProperty(PropertyName);
-	HardwareTargetingCategory.AddCustomRow(LOCTEXT("GraphicsPreferenceOption", "Graphics Preference:").ToString())
-	.NameContent()
-	[
-		SNew(STextBlock)
-		.Text(FText::FromString(Property->GetPropertyDisplayName()))
-		.Font(DetailBuilder.GetDetailFont())
-	]
-	.ValueContent()
-	[
-		SNew(SHorizontalBox)
-
-		+ SHorizontalBox::Slot()
-		.AutoWidth()
+	{
+		auto SetPropertyValue = [](EGraphicsPreset::Type NewValue, TSharedRef<IPropertyHandle> Property){
+			Property->SetValue(uint8(NewValue));
+		};
+		auto GetPropertyValue = [](TSharedRef<IPropertyHandle> Property){
+			uint8 Value = 0;
+			Property->GetValue(Value);
+			return EGraphicsPreset::Type(Value);
+		};
+		HardwareTargetingCategory.AddCustomRow(LOCTEXT("GraphicsPreferenceOption", "Graphics Preference:").ToString())
+		.NameContent()
 		[
-			HardwareTargeting.MakeGraphicsPresetTargetCombo(
-				FOnGraphicsPresetChanged::CreateStatic([](EGraphicsPreset::Type NewValue, TSharedRef<IPropertyHandle> Property){
-					Property->SetValue(uint8(NewValue));
-				}, Property),
-				TAttribute<EGraphicsPreset::Type>::Create(TAttribute<EGraphicsPreset::Type>::FGetter::CreateStatic([](TSharedRef<IPropertyHandle> Property){
-					uint8 Value = 0;
-					Property->GetValue(Value);
-					return EGraphicsPreset::Type(Value);
-				}, Property))
-			)
+			SNew(STextBlock)
+			.Text(FText::FromString(Property->GetPropertyDisplayName()))
+			.Font(DetailBuilder.GetDetailFont())
 		]
-	];
+		.ValueContent()
+		[
+			SNew(SHorizontalBox)
+
+			+ SHorizontalBox::Slot()
+			.AutoWidth()
+			[
+				HardwareTargeting.MakeGraphicsPresetTargetCombo(
+					FOnGraphicsPresetChanged::CreateStatic(SetPropertyValue, Property),
+					TAttribute<EGraphicsPreset::Type>::Create(TAttribute<EGraphicsPreset::Type>::FGetter::CreateStatic(GetPropertyValue, Property))
+				)
+			]
+		];
+	}
 
 }
 
diff --git a/Engine/Source/Editor/EditorStyle/Private/SlateEditorStyle.cpp b/Engine/Source/Editor/EditorStyle/Private/SlateEditorStyle.cpp
--- a/Engine/Source/Editor/EditorStyle/Private/SlateEditorStyle.cpp
+++ b/Engine/Source/Editor/EditorStyle/Private/SlateEditorStyle.cpp
@@ -1429,19 +1429,14 @@ void FSlateEditorStyle::FStyle::SetupGeneralStyles()
 		FLinearColor EditorOrange = FLinearColor(0.728f, 0.364f, 0.003f);
 
 		FTextBlockStyle TargetSettingsNormal = FTextBlockStyle(NormalText)
-			.SetFont(TTF_CORE_FONT("Fonts/Roboto-Regular", 12))
-			.SetShadowOffset(FVector2D::UnitVector);
+			.SetFont(TTF_CORE_FONT("Fonts/Roboto-Regular", 8));
 
 		Set("HardwareTargets.Normal", TargetSettingsNormal);
 
  		Set("HardwareTargets.Strong", FTextBlockStyle(TargetSettingsNormal)
- 			.SetFont(TTF_CORE_FONT("Fonts/Roboto-Bold", 12))
- 			.SetShadowOffset(FVector2D::UnitVector));
-
-		Set("HardwareTargets.H1", FTextBlockStyle(TargetSettingsNormal)
+ 			.SetFont(TTF_CORE_FONT("Fonts/Roboto-Bold", 8))
 			.SetColorAndOpacity(EditorOrange)
-			.SetFont(TTF_CORE_FONT("Fonts/Roboto-Bold", 16))
-			.SetShadowOffset(FVector2D::UnitVector));
+ 			.SetShadowOffset(FVector2D::UnitVector));
 	}
 #endif
 
diff --git a/Engine/Source/Editor/HardwareTargeting/Private/HardwareTargetingModule.cpp b/Engine/Source/Editor/HardwareTargeting/Private/HardwareTargetingModule.cpp
--- a/Engine/Source/Editor/HardwareTargeting/Private/HardwareTargetingModule.cpp
+++ b/Engine/Source/Editor/HardwareTargeting/Private/HardwareTargetingModule.cpp
@@ -16,38 +16,47 @@
 
 struct FMetaSettingGatherer
 {
-	TSet<UObject*> DirtyObjects;
 	FTextBuilder DescriptionBuffer;
 
+	TMap<UObject*, FTextBuilder> DescriptionBuffers;
+	TMap<UObject*, FText> CategoryNames;
+
 	// Are we just displaying what would change, or actually changing things?
 	bool bReadOnly;
 
+	bool bIncludeUnmodifiedProperties;
+
 	FMetaSettingGatherer()
 		: bReadOnly(false)
+		, bIncludeUnmodifiedProperties(false)
 	{
 	}
 
 	void AddEntry(UObject* SettingsObject, UProperty* Property, FText NewValue, bool bModified)
 	{
-		if (bModified && !bReadOnly)
+		if (bModified || bIncludeUnmodifiedProperties)
 		{
-			DirtyObjects.Add(SettingsObject);
-
-			FPropertyChangedEvent ChangeEvent(Property, false, EPropertyChangeType::ValueSet);
-			SettingsObject->PostEditChangeProperty(ChangeEvent);
-		}
-
-		if (bReadOnly)
-		{
-			FText SettingDisplayName = Property->GetDisplayNameText();
-
-			FFormatNamedArguments Args;
-			Args.Add(TEXT("SettingName"), SettingDisplayName);
-			Args.Add(TEXT("SettingValue"), NewValue);
-
-			FText FormatString = bModified ? LOCTEXT("MetaSettingDisplayStringUnmodified", "{SettingName} is {SettingValue} <HardwareTargets.Strong>(modified)</>") : LOCTEXT("MetaSettingDisplayStringUnmodified", "{SettingName} is {SettingValue}");
-
-			DescriptionBuffer.AppendLine(FText::Format(FormatString, Args));
+			FTextBuilder& DescriptionBuffer = DescriptionBuffers.FindOrAdd(SettingsObject);
+
+			if (!bReadOnly)
+			{
+				FPropertyChangedEvent ChangeEvent(Property, false, EPropertyChangeType::ValueSet);
+				SettingsObject->PostEditChangeProperty(ChangeEvent);
+			}
+			else
+			{
+				FText SettingDisplayName = Property->GetDisplayNameText();
+
+				FFormatNamedArguments Args;
+				Args.Add(TEXT("SettingName"), SettingDisplayName);
+				Args.Add(TEXT("SettingValue"), NewValue);
+
+				FText FormatString = bModified ?
+					LOCTEXT("MetaSettingDisplayStringModified", "{SettingName} is {SettingValue} <HardwareTargets.Strong>(modified)</>") :
+					LOCTEXT("MetaSettingDisplayStringUnmodified", "{SettingName} is {SettingValue}");
+
+				DescriptionBuffer.AppendLine(FText::Format(FormatString, Args));
+			}
 		}
 	}
 
@@ -57,26 +66,10 @@ struct FMetaSettingGatherer
 	void Finalize()
 	{
 		check(!bReadOnly);
-		for (UObject* SettingsObject : DirtyObjects)
-		{
-			SettingsObject->UpdateDefaultConfigFile();
-		}
-	}
 
-	void StartSection(FText SectionHeading)
-	{
-		if (bReadOnly)
+		for (auto& Pair : DescriptionBuffers)
 		{
-			DescriptionBuffer.AppendLine(SectionHeading);
-			DescriptionBuffer.Indent();
-		}
-	}
-
-	void EndSection()
-	{
-		if (bReadOnly)
-		{
-			DescriptionBuffer.Unindent();
+			Pair.Key->UpdateDefaultConfigFile();
 		}
 	}
 };
@@ -126,7 +119,7 @@ public:
 	
 	// IHardwareTargetingModule interface
 	virtual void ApplyHardwareTargetingSettings() override;
-	virtual FText QueryReadableDescriptionOfHardwareTargetingSettings() override;
+	virtual TArray<FModifiedDefaultConfig> GetPendingSettingsChanges() override;
 	virtual TSharedRef<SWidget> MakeHardwareClassTargetCombo(FOnHardwareClassChanged OnChanged, TAttribute<EHardwareClass::Type> SelectedEnum) override;
 	virtual TSharedRef<SWidget> MakeGraphicsPresetTargetCombo(FOnGraphicsPresetChanged OnChanged, TAttribute<EGraphicsPreset::Type> SelectedEnum) override;
 	// End of IHardwareTargetingModule interface
@@ -153,21 +146,41 @@ void FHardwareTargetingModule::ShutdownModule()
 {
 }
 
-FText FHardwareTargetingModule::QueryReadableDescriptionOfHardwareTargetingSettings()
+TArray<FModifiedDefaultConfig> FHardwareTargetingModule::GetPendingSettingsChanges()
 {
 	// Gather and stringify the modified settings
 	FMetaSettingGatherer Gatherer;
 	Gatherer.bReadOnly = true;
+	Gatherer.bIncludeUnmodifiedProperties = true;
 	GatherSettings(Gatherer);
 
-	return Gatherer.DescriptionBuffer.ToText();
-}
+	TArray<FModifiedDefaultConfig> OutArray;
+	for (auto& Pair : Gatherer.DescriptionBuffers)
+	{
+		FModifiedDefaultConfig ModifiedConfig;
+		ModifiedConfig.SettingsObject = Pair.Key;
+		ModifiedConfig.Description = Pair.Value.ToText();
+		ModifiedConfig.CategoryHeading = Gatherer.CategoryNames.FindChecked(Pair.Key);
 
+		OutArray.Add(ModifiedConfig);
+	}
+	return OutArray;
+}
 	
 void FHardwareTargetingModule::GatherSettings(FMetaSettingGatherer& Builder)
 {
 	UHardwareTargetingSettings* Settings = GetMutableDefault<UHardwareTargetingSettings>();
 
+
+	if (Builder.bReadOnly)
+	{
+		// Force the category order and give nice descriptions
+		Builder.CategoryNames.Add(GetMutableDefault<URendererSettings>(), LOCTEXT("RenderingCategoryHeader", "Engine - Rendering"));
+		Builder.CategoryNames.Add(GetMutableDefault<UInputSettings>(), LOCTEXT("InputCategoryHeader", "Engine - Input"));
+		Builder.CategoryNames.Add(GetMutableDefault<UGameMapsSettings>(), LOCTEXT("MapsAndModesCategoryHeader", "Project - Maps & Modes"));
+	}
+
+
 	const bool bLowEndMobile = (Settings->TargetedHardwareClass == EHardwareClass::Mobile) && (Settings->DefaultGraphicsPerformance == EGraphicsPreset::Scalable);
 	const bool bAnyMobile = (Settings->TargetedHardwareClass == EHardwareClass::Mobile);
 	const bool bHighEndMobile = (Settings->TargetedHardwareClass == EHardwareClass::Mobile) && (Settings->DefaultGraphicsPerformance == EGraphicsPreset::Maximum);
@@ -175,8 +188,6 @@ void FHardwareTargetingModule::GatherSettings(FMetaSettingGatherer& Builder)
 	const bool bHighEndPC = (Settings->TargetedHardwareClass == EHardwareClass::Desktop) && (Settings->DefaultGraphicsPerformance == EGraphicsPreset::Maximum);
 
 	{
-		Builder.StartSection(LOCTEXT("RenderingHeading", "<HardwareTargets.H1>Default Rendering Settings</>"));
-
 		// Based roughly on https://docs.unrealengine.com/latest/INT/Platforms/Mobile/PostProcessEffects/index.html
 		UE_META_SETTING_ENTRY(Builder, URendererSettings, bMobileHDR, !bLowEndMobile);
 
@@ -196,27 +207,17 @@ void FHardwareTargetingModule::GatherSettings(FMetaSettingGatherer& Builder)
 		//@TODO: DOF setting doesn't exist yet
 		// UE_META_SETTING_ENTRY(Builder, URendererSettings, bDefaultFeatureDepthOfField, bHighEndPC);
 		UE_META_SETTING_ENTRY(Builder, URendererSettings, DefaultFeatureAntiAliasing, bHighEndPC ? EAntiAliasingMethodUI::AAM_TemporalAA : EAntiAliasingMethodUI::AAM_None);
-
-		Builder.EndSection();
 	}
 
 	{
-		Builder.StartSection(LOCTEXT("InputHeading", "<HardwareTargets.H1>Default Input Settings</>"));
-
 		// Mobile uses touch
 		UE_META_SETTING_ENTRY(Builder, UInputSettings, bUseMouseForTouch, bAnyMobile);
 		//@TODO: Use bAlwaysShowTouchInterface (sorta implied by bUseMouseForTouch)?
-
-		Builder.EndSection();
 	}
 
 	{
-		Builder.StartSection(LOCTEXT("GameHeading", "<HardwareTargets.H1>Default Game Settings</>"));
-		
 		// Tablets or phones are usually shared-screen multiplayer instead of split-screen
 		UE_META_SETTING_ENTRY(Builder, UGameMapsSettings, bUseSplitscreen, bAnyPC);
-
-		Builder.EndSection();
 	}
 }
 
diff --git a/Engine/Source/Editor/HardwareTargeting/Private/HardwareTargetingSettings.cpp b/Engine/Source/Editor/HardwareTargeting/Private/HardwareTargetingSettings.cpp
--- a/Engine/Source/Editor/HardwareTargeting/Private/HardwareTargetingSettings.cpp
+++ b/Engine/Source/Editor/HardwareTargeting/Private/HardwareTargetingSettings.cpp
@@ -13,7 +13,7 @@ UHardwareTargetingSettings::UHardwareTargetingSettings(const FPostConstructIniti
 	, DefaultGraphicsPerformance(EGraphicsPreset::Maximum)
 {
 }
-	
+
 bool UHardwareTargetingSettings::HasPendingChanges() const
 {
 	UHardwareTargetingSettings* Settings = GetMutableDefault<UHardwareTargetingSettings>();
@@ -38,3 +38,8 @@ bool UHardwareTargetingSettings::HasPendingChanges() const
 
 	return false;
 }
+
+void UHardwareTargetingSettings::PostEditChangeProperty( struct FPropertyChangedEvent& PropertyChangedEvent )
+{
+	SettingChangedEvent.Broadcast();
+}
\ No newline at end of file
diff --git a/Engine/Source/Editor/HardwareTargeting/Public/HardwareTargetingModule.h b/Engine/Source/Editor/HardwareTargeting/Public/HardwareTargetingModule.h
--- a/Engine/Source/Editor/HardwareTargeting/Public/HardwareTargetingModule.h
+++ b/Engine/Source/Editor/HardwareTargeting/Public/HardwareTargetingModule.h
@@ -8,6 +8,19 @@
 DECLARE_DELEGATE_OneParam(FOnHardwareClassChanged, EHardwareClass::Type)
 DECLARE_DELEGATE_OneParam(FOnGraphicsPresetChanged, EGraphicsPreset::Type)
 
+/** Struct specifying pending changes to a settings object */
+struct FModifiedDefaultConfig
+{
+	/** The settings object to which the description relates */
+	TWeakObjectPtr<UObject> SettingsObject;
+
+	/** Heading describing the name of the category */
+	FText CategoryHeading;
+
+	/** Text describing the pending changes to the settings */
+	FText Description;
+};
+
 class IHardwareTargetingModule : public IModuleInterface
 {
 public:
@@ -18,8 +31,8 @@ public:
 	/** Apply the current hardware targeting settings if they have changed */
 	virtual void ApplyHardwareTargetingSettings() = 0;
 
-	/** Queries what the the current hardware targeting settings would like to change */
-	virtual FText QueryReadableDescriptionOfHardwareTargetingSettings() = 0;
+	/** Gets a list of objects that are required to be writable in order to apply the settings */
+	virtual TArray<FModifiedDefaultConfig> GetPendingSettingsChanges() = 0;
 
 	/** Make a new combo box for choosing a hardware class target */
 	virtual TSharedRef<SWidget> MakeHardwareClassTargetCombo(FOnHardwareClassChanged OnChanged, TAttribute<EHardwareClass::Type> SelectedEnum) = 0;
diff --git a/Engine/Source/Editor/HardwareTargeting/Public/HardwareTargetingSettings.h b/Engine/Source/Editor/HardwareTargeting/Public/HardwareTargetingSettings.h
--- a/Engine/Source/Editor/HardwareTargeting/Public/HardwareTargetingSettings.h
+++ b/Engine/Source/Editor/HardwareTargeting/Public/HardwareTargetingSettings.h
@@ -41,4 +41,19 @@ class HARDWARETARGETING_API UHardwareTargetingSettings
 
 	/** Check if these settings have any pending changes that require action */
 	bool HasPendingChanges() const;
+
+#if WITH_EDITOR
+public:
+	/** Returns an event delegate that is executed when a setting has changed. */
+	DECLARE_EVENT(UHardwareTargetingSettings, FSettingChangedEvent);
+	FSettingChangedEvent& OnSettingChanged( ) { return SettingChangedEvent; }
+
+protected:
+	/** Called when a property on this object is changed */
+	virtual void PostEditChangeProperty( struct FPropertyChangedEvent& PropertyChangedEvent ) override;
+
+private:
+	/** Holds an event delegate that is executed when a setting has changed. */
+	FSettingChangedEvent SettingChangedEvent;
+#endif
 };
\ No newline at end of file
