More work to tidy up the target hardware section of the project settings

Applied settings are now written as a string rather than an int
Also added UObject::GetDefaultConfigFilename()

diff --git a/Engine/Source/Developer/SettingsEditor/Private/Widgets/SSettingsEditor.cpp b/Engine/Source/Developer/SettingsEditor/Private/Widgets/SSettingsEditor.cpp
--- a/Engine/Source/Developer/SettingsEditor/Private/Widgets/SSettingsEditor.cpp
+++ b/Engine/Source/Developer/SettingsEditor/Private/Widgets/SSettingsEditor.cpp
@@ -308,8 +308,7 @@ bool SSettingsEditor::CheckOutDefaultConfigFile( )
 
 FString SSettingsEditor::GetDefaultConfigFilePath( const TWeakObjectPtr<UObject>& SettingsObject ) const
 {
-	FString RelativeConfigFilePath = FString::Printf(TEXT("%sDefault%s.ini"), *FPaths::SourceConfigDir(), *SettingsObject->GetClass()->ClassConfigName.ToString());
-
+	FString RelativeConfigFilePath = SettingsObject->GetDefaultConfigFilename();
 	return FPaths::ConvertRelativePathToFull(RelativeConfigFilePath);
 }
 
diff --git a/Engine/Source/Developer/SharedSettingsWidgets/Private/SSettingsEditorCheckoutNotice.cpp b/Engine/Source/Developer/SharedSettingsWidgets/Private/SSettingsEditorCheckoutNotice.cpp
--- a/Engine/Source/Developer/SharedSettingsWidgets/Private/SSettingsEditorCheckoutNotice.cpp
+++ b/Engine/Source/Developer/SharedSettingsWidgets/Private/SSettingsEditorCheckoutNotice.cpp
@@ -19,7 +19,7 @@ void SSettingsEditorCheckoutNotice::Construct( const FArguments& InArgs )
 
 	DefaultConfigCheckOutTimer = 0.0f;
 	DefaultConfigCheckOutNeeded = false;
-	DefaultConfigQueryInProgress = false;
+	DefaultConfigQueryInProgress = true;
 
 	// default configuration notice
 	ChildSlot
@@ -223,7 +223,7 @@ void SSettingsEditorCheckoutNotice::Tick(const FGeometry& AllottedGeometry, cons
 	{
 		bool NewCheckOutNeeded = false;
 
-		DefaultConfigQueryInProgress = false;
+		DefaultConfigQueryInProgress = true;
 		FString CachedConfigFileName = ConfigFilePath.Get();
 		if (!CachedConfigFileName.IsEmpty())
 		{
@@ -240,6 +240,7 @@ void SSettingsEditorCheckoutNotice::Tick(const FGeometry& AllottedGeometry, cons
 			else
 			{
 				NewCheckOutNeeded = (FPaths::FileExists(CachedConfigFileName) && IFileManager::Get().IsReadOnly(*CachedConfigFileName));
+				DefaultConfigQueryInProgress = false;
 			}
 
 			// file has been checked in or reverted
diff --git a/Engine/Source/Editor/DetailCustomizations/Private/HardwareTargetingSettingsDetails.cpp b/Engine/Source/Editor/DetailCustomizations/Private/HardwareTargetingSettingsDetails.cpp
--- a/Engine/Source/Editor/DetailCustomizations/Private/HardwareTargetingSettingsDetails.cpp
+++ b/Engine/Source/Editor/DetailCustomizations/Private/HardwareTargetingSettingsDetails.cpp
@@ -74,7 +74,7 @@ public:
 			+ SVerticalBox::Slot()
 			.AutoHeight()
 			[
-				SAssignNew(SettingsDescriptions, SUniformGridPanel)
+				SAssignNew(CheckoutNotices, SVerticalBox)
 			]
 		];
 	}
@@ -137,15 +137,10 @@ public:
 		return true;
 	}
 
-	FString GetDefaultConfigPath(UObject* Settings) const
-	{
-		return FPaths::ConvertRelativePathToFull(FString::Printf(TEXT("%sDefault%s.ini"), *FPaths::SourceConfigDir(), *Settings->GetClass()->ClassConfigName.ToString()));
-	}
-
 	void Update()
 	{
 		FileWatcherWidgets.Reset();
-		SettingsDescriptions->ClearChildren();
+		CheckoutNotices->ClearChildren();
 
 		IHardwareTargetingModule& Module = IHardwareTargetingModule::Get();
 		int32 SlotIndex = 0;
@@ -167,7 +162,7 @@ public:
 			}
 			
 
-			FString ConfigFile = GetDefaultConfigPath(Settings.SettingsObject.Get());
+			FString ConfigFile = FPaths::ConvertRelativePathToFull(Settings.SettingsObject->GetDefaultConfigFilename());
 
 			if (!SeenConfigFiles.Contains(ConfigFile))
 			{
@@ -178,7 +173,8 @@ public:
 					.ConfigFilePath(ConfigFile);
 				FileWatcherWidgets.Add(FileWatcherWidget);
 
-				SettingsDescriptions->AddSlot(0, SlotIndex)
+				CheckoutNotices->AddSlot()
+				.Padding(FMargin(0.f, 0.f, 12.f, 5.f))
 				[
 					FileWatcherWidget
 				];
@@ -189,7 +185,7 @@ public:
 
 	TArray<TSharedPtr<SSettingsEditorCheckoutNotice>> FileWatcherWidgets;
 
-	TSharedPtr<SUniformGridPanel> SettingsDescriptions;
+	TSharedPtr<SVerticalBox> CheckoutNotices;
 
 	double LastStatusUpdate;
 };
@@ -207,9 +203,6 @@ void FHardwareTargetingSettingsDetails::CustomizeDetails(IDetailLayoutBuilder& D
 		auto* Settings = GetMutableDefault<UHardwareTargetingSettings>();
 		return Settings->HasPendingChanges() ? EVisibility::Collapsed : EVisibility::Visible;
 	};
-	//auto ModifiedText = []{
-	//	return IHardwareTargetingModule::Get().QueryReadableDescriptionOfHardwareTargetingSettings();
-	//};
 
 	TSharedRef<SRequiredDefaultConfig> ConfigWidget = SNew(SRequiredDefaultConfig);
 	IHardwareTargetingModule& HardwareTargeting = IHardwareTargetingModule::Get();
@@ -238,11 +231,14 @@ void FHardwareTargetingSettingsDetails::CustomizeDetails(IDetailLayoutBuilder& D
 
 	ConfigWidget->Initialize(DetailBuilder, PendingChangesCategory);
 
-	auto PropertyName = GET_MEMBER_NAME_CHECKED(UHardwareTargetingSettings, TargetedHardwareClass);
-	DetailBuilder.HideProperty(PropertyName);
+	TSharedPtr<SWidget> HardwareClassCombo, GraphicsPresetCombo;
 
-	TSharedRef<IPropertyHandle> Property = DetailBuilder.GetProperty(PropertyName);
+	// Setup the hardware class combo
 	{
+		auto PropertyName = GET_MEMBER_NAME_CHECKED(UHardwareTargetingSettings, TargetedHardwareClass);
+		DetailBuilder.HideProperty(PropertyName);
+
+		TSharedRef<IPropertyHandle> Property = DetailBuilder.GetProperty(PropertyName);
 		auto SetPropertyValue = [](EHardwareClass::Type NewValue, TSharedRef<IPropertyHandle> Property){
 			Property->SetValue(uint8(NewValue));
 		};
@@ -251,34 +247,19 @@ void FHardwareTargetingSettingsDetails::CustomizeDetails(IDetailLayoutBuilder& D
 			Property->GetValue(Value);
 			return EHardwareClass::Type(Value);
 		};
-
-		HardwareTargetingCategory.AddCustomRow(LOCTEXT("HardwareTargetingOption", "Targeted Hardware:").ToString())
-		.NameContent()
-		[
-			SNew(STextBlock)
-			.Text(FText::FromString(Property->GetPropertyDisplayName()))
-			.Font(DetailBuilder.GetDetailFont())
-		]
-		.ValueContent()
-		[
-			SNew(SHorizontalBox)
-
-			+ SHorizontalBox::Slot()
-			.AutoWidth()
-			[
-				HardwareTargeting.MakeHardwareClassTargetCombo(
-					FOnHardwareClassChanged::CreateStatic(SetPropertyValue, Property),
-					TAttribute<EHardwareClass::Type>::Create(TAttribute<EHardwareClass::Type>::FGetter::CreateStatic(GetPropertyValue, Property))
-				)
-			]
-		];
+		
+		HardwareClassCombo = HardwareTargeting.MakeHardwareClassTargetCombo(
+			FOnHardwareClassChanged::CreateStatic(SetPropertyValue, Property),
+			TAttribute<EHardwareClass::Type>::Create(TAttribute<EHardwareClass::Type>::FGetter::CreateStatic(GetPropertyValue, Property))
+		);
 	}
 
-	PropertyName = GET_MEMBER_NAME_CHECKED(UHardwareTargetingSettings, DefaultGraphicsPerformance);
-	DetailBuilder.HideProperty(PropertyName);
-
-	Property = DetailBuilder.GetProperty(PropertyName);
+	// Setup the graphics preset combo
 	{
+		auto PropertyName = GET_MEMBER_NAME_CHECKED(UHardwareTargetingSettings, DefaultGraphicsPerformance);
+		DetailBuilder.HideProperty(PropertyName);
+
+		TSharedRef<IPropertyHandle> Property = DetailBuilder.GetProperty(PropertyName);
 		auto SetPropertyValue = [](EGraphicsPreset::Type NewValue, TSharedRef<IPropertyHandle> Property){
 			Property->SetValue(uint8(NewValue));
 		};
@@ -287,28 +268,38 @@ void FHardwareTargetingSettingsDetails::CustomizeDetails(IDetailLayoutBuilder& D
 			Property->GetValue(Value);
 			return EGraphicsPreset::Type(Value);
 		};
-		HardwareTargetingCategory.AddCustomRow(LOCTEXT("GraphicsPreferenceOption", "Graphics Preference:").ToString())
-		.NameContent()
+		GraphicsPresetCombo = HardwareTargeting.MakeGraphicsPresetTargetCombo(
+			FOnGraphicsPresetChanged::CreateStatic(SetPropertyValue, Property),
+			TAttribute<EGraphicsPreset::Type>::Create(TAttribute<EGraphicsPreset::Type>::FGetter::CreateStatic(GetPropertyValue, Property))
+		);
+	}
+
+	HardwareTargetingCategory.AddCustomRow(LOCTEXT("HardwareTargetingOption", "Targeted Hardware:").ToString())
+	.NameContent()
+	[
+		SNew(STextBlock)
+		.Text(LOCTEXT("OptimizeProjectFor", "Optimize project settings for:"))
+		.Font(DetailBuilder.GetDetailFont())
+	]
+	.ValueContent()
+	.MaxDesiredWidth(0)
+	[
+		SNew(SHorizontalBox)
+
+		+ SHorizontalBox::Slot()
+		.Padding(FMargin(10.f, 0.f))
+		.AutoWidth()
 		[
-			SNew(STextBlock)
-			.Text(FText::FromString(Property->GetPropertyDisplayName()))
-			.Font(DetailBuilder.GetDetailFont())
+			HardwareClassCombo.ToSharedRef()
 		]
-		.ValueContent()
-		[
-			SNew(SHorizontalBox)
-
-			+ SHorizontalBox::Slot()
-			.AutoWidth()
-			[
-				HardwareTargeting.MakeGraphicsPresetTargetCombo(
-					FOnGraphicsPresetChanged::CreateStatic(SetPropertyValue, Property),
-					TAttribute<EGraphicsPreset::Type>::Create(TAttribute<EGraphicsPreset::Type>::FGetter::CreateStatic(GetPropertyValue, Property))
-				)
-			]
-		];
-	}
 
+		+ SHorizontalBox::Slot()
+		.Padding(FMargin(0.f, 0.f, 10.f, 0.f))
+		.AutoWidth()
+		[
+			GraphicsPresetCombo.ToSharedRef()
+		]
+	];
 }
 
 #undef LOCTEXT_NAMESPACE
diff --git a/Engine/Source/Editor/HardwareTargeting/Private/HardwareTargetingModule.cpp b/Engine/Source/Editor/HardwareTargeting/Private/HardwareTargetingModule.cpp
--- a/Engine/Source/Editor/HardwareTargeting/Private/HardwareTargetingModule.cpp
+++ b/Engine/Source/Editor/HardwareTargeting/Private/HardwareTargetingModule.cpp
@@ -64,14 +64,23 @@ struct FMetaSettingGatherer
 	template <typename ValueType>
 	static FText ValueToString(ValueType Value);
 
-	void Finalize()
+	bool Finalize()
 	{
 		check(!bReadOnly);
 
+		bool bSuccess = true;
 		for (auto& Pair : DescriptionBuffers)
 		{
+			const FString Filename = Pair.Key->GetDefaultConfigFilename();
+			const FDateTime BeforeTime = IFileManager::Get().GetTimeStamp(*Filename);
+
 			Pair.Key->UpdateDefaultConfigFile();
+
+			const FDateTime AfterTime = IFileManager::Get().GetTimeStamp(*Filename);
+			bSuccess = BeforeTime != AfterTime && bSuccess;
 		}
+
+		return bSuccess;
 	}
 };
 
@@ -233,12 +242,22 @@ void FHardwareTargetingModule::ApplyHardwareTargetingSettings()
 		FMetaSettingGatherer Builder;
 		Builder.bReadOnly = false;
 		GatherSettings(Builder);
-		Builder.Finalize();
+
+		const bool bSuccess = Builder.Finalize();
 
 		// Write out the 'did we apply' values
-		GConfig->SetInt(TEXT("AppliedHardwareTargetingSettings"), TEXT("AppliedDefaultGraphicsPerformance"), Settings->DefaultGraphicsPerformance, GEditorIni);
-		GConfig->SetInt(TEXT("AppliedHardwareTargetingSettings"), TEXT("AppliedTargetedHardwareClass"), Settings->TargetedHardwareClass, GEditorIni);
-		GConfig->Flush(false, GEditorIni);
+		if (bSuccess)
+		{
+			UEnum* Enum = FindObject<UEnum>(ANY_PACKAGE, TEXT("EGraphicsPreset"), true);
+			check(Enum);
+			GConfig->SetString(TEXT("AppliedHardwareTargetingSettings"), TEXT("AppliedDefaultGraphicsPerformance"), *Enum->GetEnumName(Settings->DefaultGraphicsPerformance), GEditorIni);
+
+			Enum = FindObject<UEnum>(ANY_PACKAGE, TEXT("EHardwareClass"), true);
+			check(Enum);
+			GConfig->SetString(TEXT("AppliedHardwareTargetingSettings"), TEXT("AppliedTargetedHardwareClass"), *Enum->GetEnumName(Settings->TargetedHardwareClass), GEditorIni);
+
+			GConfig->Flush(false, GEditorIni);
+		}
 	}
 }
 
diff --git a/Engine/Source/Editor/HardwareTargeting/Private/HardwareTargetingSettings.cpp b/Engine/Source/Editor/HardwareTargeting/Private/HardwareTargetingSettings.cpp
--- a/Engine/Source/Editor/HardwareTargeting/Private/HardwareTargetingSettings.cpp
+++ b/Engine/Source/Editor/HardwareTargeting/Private/HardwareTargetingSettings.cpp
@@ -19,21 +19,57 @@ bool UHardwareTargetingSettings::HasPendingChanges() const
 	UHardwareTargetingSettings* Settings = GetMutableDefault<UHardwareTargetingSettings>();
 
 	{
-		int32 AppliedHardwareEnumValue = 0;
-		bool bFoundKey = GConfig->GetInt(TEXT("AppliedHardwareTargetingSettings"), TEXT("AppliedTargetedHardwareClass"), AppliedHardwareEnumValue, GEditorIni);
-		if (!bFoundKey || (EHardwareClass::Type(AppliedHardwareEnumValue) != Settings->TargetedHardwareClass))
+		FString AppliedHardwareEnumValue;
+		const bool bFoundKey = GConfig->GetString(TEXT("AppliedHardwareTargetingSettings"), TEXT("AppliedTargetedHardwareClass"), AppliedHardwareEnumValue, GEditorIni);
+		if (!bFoundKey)
 		{
 			return true;
 		}
+		else if (AppliedHardwareEnumValue.IsNumeric())
+		{
+			// We used to write out an int instead of a string
+			const EHardwareClass::Type EnumValue = static_cast<EHardwareClass::Type>(FPlatformString::Atoi(*AppliedHardwareEnumValue));
+			if (EnumValue != Settings->TargetedHardwareClass)
+			{
+				return true;
+			}
+		}
+		else
+		{
+			UEnum* Enum = FindObject<UEnum>(ANY_PACKAGE, TEXT("EHardwareClass"), true);
+			check(Enum);
+			if (Enum->GetEnumName(Settings->TargetedHardwareClass) != AppliedHardwareEnumValue)
+			{
+				return true;
+			}
+		}
 	}
 
 	{
-		int32 AppliedGraphicsEnumValue = 0;
-		bool bFoundKey = GConfig->GetInt(TEXT("AppliedHardwareTargetingSettings"), TEXT("AppliedDefaultGraphicsPerformance"), AppliedGraphicsEnumValue, GEditorIni);
-		if (!bFoundKey || (EGraphicsPreset::Type(AppliedGraphicsEnumValue) != Settings->DefaultGraphicsPerformance))
+		FString AppliedGraphicsEnumValue;
+		const bool bFoundKey = GConfig->GetString(TEXT("AppliedHardwareTargetingSettings"), TEXT("AppliedDefaultGraphicsPerformance"), AppliedGraphicsEnumValue, GEditorIni);
+		if (!bFoundKey)
 		{
 			return true;
 		}
+		else if (AppliedGraphicsEnumValue.IsNumeric())
+		{
+			// We used to write out an int instead of a string
+			const EGraphicsPreset::Type EnumValue = EGraphicsPreset::Type(FPlatformString::Atoi(*AppliedGraphicsEnumValue));
+			if (EnumValue != Settings->DefaultGraphicsPerformance)
+			{
+				return true;
+			}
+		}
+		else
+		{
+			UEnum* Enum = FindObject<UEnum>(ANY_PACKAGE, TEXT("EGraphicsPreset"), true);
+			check(Enum);
+			if (Enum->GetEnumName(Settings->DefaultGraphicsPerformance) != AppliedGraphicsEnumValue)
+			{
+				return true;
+			}
+		}
 	}
 
 	return false;
diff --git a/Engine/Source/Runtime/CoreUObject/Private/UObject/Obj.cpp b/Engine/Source/Runtime/CoreUObject/Private/UObject/Obj.cpp
--- a/Engine/Source/Runtime/CoreUObject/Private/UObject/Obj.cpp
+++ b/Engine/Source/Runtime/CoreUObject/Private/UObject/Obj.cpp
@@ -1736,6 +1736,11 @@ void UObject::SaveConfig( uint64 Flags, const TCHAR* InFilename, FConfigCacheIni
 	}
 }
 
+FString UObject::GetDefaultConfigFilename() const
+{
+	return FString::Printf(TEXT("%sDefault%s.ini"), *FPaths::SourceConfigDir(), *GetClass()->ClassConfigName.ToString());
+}
+
 // @todo ini: Verify per object config objects
 void UObject::UpdateDefaultConfigFile()
 {
@@ -1743,7 +1748,7 @@ void UObject::UpdateDefaultConfigFile()
 	FConfigCacheIni Config;
 
 	// add an empty file to the config so it doesn't read in the original file (see FConfigCacheIni.Find())
-	FString DefaultIniName = FString::Printf(TEXT("%sDefault%s.ini"), *FPaths::SourceConfigDir(), *GetClass()->ClassConfigName.ToString());
+	FString DefaultIniName = GetDefaultConfigFilename();
 	FConfigFile& NewFile = Config.Add(DefaultIniName, FConfigFile());
 
 	// save the object properties to this file
diff --git a/Engine/Source/Runtime/CoreUObject/Public/UObject/UObject.h b/Engine/Source/Runtime/CoreUObject/Public/UObject/UObject.h
--- a/Engine/Source/Runtime/CoreUObject/Public/UObject/UObject.h
+++ b/Engine/Source/Runtime/CoreUObject/Public/UObject/UObject.h
@@ -626,6 +626,11 @@ public:
 	 */
 	void UpdateDefaultConfigFile();
 
+	/**
+	 * Get the default config filename for the specified UObject
+	 */
+	FString GetDefaultConfigFilename() const;
+
 	/**
 	 * Imports property values from an .ini file.
 	 *
