[345504] Blueprints - Fix issues with propagation of default rotation transform values during live editing.

- Modified FComponentTransformDetails to fix various issues with default rotation value propagation.
- Consolidated default transform value propagation code into FComponentEditorUtils to eliminate some redundancy.

diff --git a/Engine/Source/Editor/DetailCustomizations/Private/ComponentTransformDetails.cpp b/Engine/Source/Editor/DetailCustomizations/Private/ComponentTransformDetails.cpp
--- a/Engine/Source/Editor/DetailCustomizations/Private/ComponentTransformDetails.cpp
+++ b/Engine/Source/Editor/DetailCustomizations/Private/ComponentTransformDetails.cpp
@@ -46,41 +46,18 @@ private:
 template<typename T>
 static void PropagateTransformPropertyChange(UObject* InObject, UProperty* InProperty, const T& OldValue, const T& NewValue)
 {
-	check(InObject != NULL);
-	check(InProperty != NULL);
+	check(InObject != nullptr);
+	check(InProperty != nullptr);
 
 	TArray<UObject*> ArchetypeInstances;
+	TSet<USceneComponent*> UpdatedComponents;
 	FComponentEditorUtils::GetArchetypeInstances(InObject, ArchetypeInstances);
 	for(int32 InstanceIndex = 0; InstanceIndex < ArchetypeInstances.Num(); ++InstanceIndex)
 	{
 		USceneComponent* InstancedSceneComponent = FComponentEditorUtils::GetSceneComponent(ArchetypeInstances[InstanceIndex], InObject);
-		if(InstancedSceneComponent != NULL)
+		if(InstancedSceneComponent != nullptr && !UpdatedComponents.Contains(InstancedSceneComponent))
 		{
-			// Propagate the change only if the current instanced value matches the previous default value
-			T* CurValue = InProperty->ContainerPtrToValuePtr<T>(InstancedSceneComponent);
-			if(CurValue != NULL && *CurValue == OldValue)
-			{
-				// Ensure that this instance will be included in any undo/redo operations, and record it into the transaction buffer.
-				// Note: We don't do this for components that originate from script, because they will be re-instanced from the template after an undo, so there is no need to record them.
-				if(!InstancedSceneComponent->bCreatedByConstructionScript)
-				{
-					InstancedSceneComponent->SetFlags(RF_Transactional);
-					InstancedSceneComponent->Modify();
-				}
-
-				// We must also modify the owner, because we'll need script components to be reconstructed as part of an undo operation.
-				AActor* Owner = InstancedSceneComponent->GetOwner();
-				if(Owner != NULL)
-				{
-					Owner->Modify();
-				}
-
-				// Change the property value
-				*CurValue = NewValue;
-
-				// Re-register the component with the scene so that transforms are updated for display
-				InstancedSceneComponent->ReregisterComponent();
-			}
+			FComponentEditorUtils::PropagateTransformPropertyChange(InstancedSceneComponent, InProperty, OldValue, NewValue, UpdatedComponents);
 		}
 	}
 }
@@ -96,23 +73,6 @@ FComponentTransformDetails::FComponentTransformDetails( const TArray< TWeakObjec
 {
 	GConfig->GetBool(TEXT("SelectionDetails"), TEXT("PreserveScaleRatio"), bPreserveScaleRatio, GEditorUserSettingsIni);
 
-	// Capture selected actor rotations so that we can adjust them without worrying about the Quat conversions affecting the raw values
-	for( int32 ObjectIndex = 0; ObjectIndex < SelectedObjects.Num(); ++ObjectIndex )
-	{
-		TWeakObjectPtr<UObject> ObjectPtr = InSelectedObjects[ObjectIndex];
-		if( ObjectPtr.IsValid() )
-		{
-			UObject* Object = ObjectPtr.Get();
-			
-			USceneComponent* RootComponent = FComponentEditorUtils::GetSceneComponent( Object );
-
-			if( RootComponent )
-			{
-				FRotator& RelativeRotation = ObjectToRelativeRotationMap.FindOrAdd(Object);
-				RelativeRotation = RootComponent->RelativeRotation;
-			}
-		}
-	}
 }
 
 TSharedRef<SWidget> FComponentTransformDetails::BuildTransformFieldLabel( ETransformField::Type TransformField )
@@ -941,7 +901,7 @@ void FComponentTransformDetails::CacheTransform()
 			if( RootComponent )
 			{
 				Loc = RootComponent->RelativeLocation;
-				Rot = bEditingRotationInUI ? ObjectToRelativeRotationMap.FindOrAdd(Object) : RootComponent->RelativeRotation;
+				Rot = (bEditingRotationInUI && !Object->HasAnyFlags(RF_ClassDefaultObject|RF_DefaultSubObject)) ? ObjectToRelativeRotationMap.FindOrAdd(Object) : RootComponent->RelativeRotation;
 				Scale = RootComponent->RelativeScale3D;
 
 				if( ObjectIndex == 0 )
@@ -1100,12 +1060,14 @@ void FComponentTransformDetails::OnSetRotation( float NewValue, bool bCommitted,
 				USceneComponent* RootComponent = FComponentEditorUtils::GetSceneComponent( Object );
 				if( RootComponent )
 				{
-					FRotator& RelativeRotation = ObjectToRelativeRotationMap.FindOrAdd(Object);
+					const bool bIsEditingTemplateObject = Object->HasAnyFlags(RF_ClassDefaultObject|RF_DefaultSubObject);
+
+					FRotator& RelativeRotation = (bEditingRotationInUI && !bIsEditingTemplateObject) ? ObjectToRelativeRotationMap.FindOrAdd(Object) : RootComponent->RelativeRotation;
 					FRotator OldRelativeRotation = RelativeRotation;
 
 					float& ValueToChange = Axis == 0 ? RelativeRotation.Roll : Axis == 1 ? RelativeRotation.Pitch : RelativeRotation.Yaw;
 
-					if( bCommitted || ValueToChange != NewValue )
+					if( ValueToChange != NewValue )
 					{
 						if( !bBeganTransaction && bCommitted )
 						{
@@ -1150,25 +1112,28 @@ void FComponentTransformDetails::OnSetRotation( float NewValue, bool bCommitted,
 
 						ValueToChange = NewValue;
 
-						if( SelectedActorInfo.NumSelected == 0 )
-						{
-							// HACK: Set directly if no actors are selected since this causes Rot->Quat->Rot conversion issues
-							// (recalculates relative rotation from quat which can give an equivalent but different value than the user typed)
-							RootComponent->RelativeRotation = RelativeRotation;
-						}
-						else
+						if(!bIsEditingTemplateObject)
 						{
-							RootComponent->SetRelativeRotation( RelativeRotation );
+							if( SelectedActorInfo.NumSelected == 0 )
+							{
+								// HACK: Set directly if no actors are selected since this causes Rot->Quat->Rot conversion issues
+								// (recalculates relative rotation from quat which can give an equivalent but different value than the user typed)
+								RootComponent->RelativeRotation = RelativeRotation;
+							}
+							else
+							{
+								RootComponent->SetRelativeRotation( RelativeRotation );
+							}
 						}
 
 						AActor* ObjectAsActor = Cast<AActor>( Object );
-						if( ObjectAsActor && !ObjectAsActor->HasAnyFlags(RF_ClassDefaultObject) )
+						if( ObjectAsActor && !bIsEditingTemplateObject )
 						{
 							ObjectAsActor->ReregisterAllComponents();
 						}
 
 						// If this is a default object or subobject, propagate the change out to any current instances of this object
-						if(Object->HasAnyFlags(RF_ClassDefaultObject|RF_DefaultSubObject))
+						if(bIsEditingTemplateObject)
 						{
 							PropagateTransformPropertyChange(Object, RelativeRotationProperty, OldRelativeRotation, RelativeRotation);
 						}
@@ -1187,7 +1152,7 @@ void FComponentTransformDetails::OnSetRotation( float NewValue, bool bCommitted,
 								Object->PostEditChangeProperty( PropertyChangedEvent );	
 							}
 					
-							if (!Object->HasAnyFlags(RF_ClassDefaultObject|RF_DefaultSubObject))
+							if (!bIsEditingTemplateObject)
 							{
 								// The actor is done moving
 								GEditor->BroadcastEndObjectMovement( *Object );
@@ -1257,6 +1222,9 @@ void FComponentTransformDetails::OnBeginRotatonSlider()
 
 				UProperty* RelativeRotationProperty = FindField<UProperty>( USceneComponent::StaticClass(), "RelativeRotation" );
 				Object->PreEditChange( RelativeRotationProperty );
+
+				// Add/update cached rotation value prior to slider interaction
+				ObjectToRelativeRotationMap.FindOrAdd(Object) = RootComponent->RelativeRotation;
 			}
 		}
 	}
diff --git a/Engine/Source/Editor/Kismet/Private/SCSEditorViewportClient.cpp b/Engine/Source/Editor/Kismet/Private/SCSEditorViewportClient.cpp
--- a/Engine/Source/Editor/Kismet/Private/SCSEditorViewportClient.cpp
+++ b/Engine/Source/Editor/Kismet/Private/SCSEditorViewportClient.cpp
@@ -68,40 +68,6 @@ namespace
 		} 
 	}
 
-	template<typename T>
-	bool PropagateTransformPropertyChange(AActor* InOwner, USceneComponent* InSceneComponent, T& CurrentValue, const T& OldDefaultValue, const T& NewDefaultValue)
-	{
-		check(InSceneComponent);
-
-		// Only propagate the change if the current value matches the previous default value
-		if(NewDefaultValue != OldDefaultValue && CurrentValue == OldDefaultValue)
-		{
-			// Ensure that this instance will be included in any undo/redo operations, and record it into the transaction buffer.
-			// Note: We don't do this for components that originate from script, because they will be re-instanced from the template after an undo, so there is no need to record them.
-			if(!InSceneComponent->bCreatedByConstructionScript)
-			{
-				InSceneComponent->SetFlags(RF_Transactional);
-				InSceneComponent->Modify();
-			}
-
-			// We must also modify the owner, because we'll need script components to be reconstructed as part of an undo operation.
-			if(InOwner != NULL)
-			{
-				InOwner->Modify();
-			}
-
-			// Modify the value
-			CurrentValue = NewDefaultValue;
-
-			// Re-register the component with the scene so that transforms are updated for display
-			InSceneComponent->ReregisterComponent();
-
-			return true;
-		}
-
-		return false;
-	}
-
 	// Determine whether or not the given node has a parent node that is not the root node, is movable and is selected
 	bool IsMovableParentNodeSelected(const FSCSEditorTreeNodePtrType& NodePtr, const TArray<FSCSEditorTreeNodePtrType>& SelectedNodes)
 	{
@@ -419,7 +385,7 @@ bool FSCSEditorViewportClient::InputWidgetDelta( FViewport* Viewport, EAxisList:
 								}
 							}
 
-							FComponentEditorUtils::FTransformData OldTransform(*SelectedTemplate);
+							FComponentEditorUtils::FTransformData OldDefaultTransform(*SelectedTemplate);
 
 							TSharedPtr<ISCSEditorCustomization> Customization = BlueprintEditorPtr.Pin()->CustomizeSCSEditor(SceneComp);
 							if(Customization.IsValid() && Customization->HandleViewportDrag(SceneComp, SelectedTemplate, Drag, Rot, ModifiedScale, GetWidgetLocation()))
@@ -450,13 +416,13 @@ bool FSCSEditorViewportClient::InputWidgetDelta( FViewport* Viewport, EAxisList:
 								UpdatedComponents.Add(SelectedTemplate);
 							}
 
-							FComponentEditorUtils::FTransformData NewTransform(*SelectedTemplate);
+							FComponentEditorUtils::FTransformData NewDefaultTransform(*SelectedTemplate);
 
 							if (SelectedNodePtr->IsNative())
 							{
 								if (ensure(SelectedTemplate->HasAnyFlags(RF_DefaultSubObject)))
 								{
-									FComponentEditorUtils::PropagateTransformPropertyChange(SelectedTemplate, OldTransform, NewTransform, UpdatedComponents);
+									FComponentEditorUtils::PropagateTransformPropertyChange(SelectedTemplate, OldDefaultTransform, NewDefaultTransform, UpdatedComponents);
 								}
 							}
 
@@ -501,11 +467,11 @@ bool FSCSEditorViewportClient::InputWidgetDelta( FViewport* Viewport, EAxisList:
 										{
 											const bool bIsProcessingPreviewActor = (ArchetypeInstance == PreviewActor);
 											SceneComp = Cast<USceneComponent>(SelectedNodePtr->FindComponentInstanceInActor(ArchetypeInstance, bIsProcessingPreviewActor));
-											if(!bIsProcessingPreviewActor && SceneComp != NULL)
+											if(!bIsProcessingPreviewActor && SceneComp != nullptr && !UpdatedComponents.Contains(SceneComp))
 											{
-												PropagateTransformPropertyChange(ArchetypeInstance, SceneComp, SceneComp->RelativeLocation, OldRelativeLocation, SelectedTemplate->RelativeLocation);
-												PropagateTransformPropertyChange(ArchetypeInstance, SceneComp, SceneComp->RelativeRotation, OldRelativeRotation, SelectedTemplate->RelativeRotation);
-												PropagateTransformPropertyChange(ArchetypeInstance, SceneComp, SceneComp->RelativeScale3D, OldRelativeScale3D, SelectedTemplate->RelativeScale3D);
+												FComponentEditorUtils::PropagateTransformPropertyChange(SceneComp, SceneComp->RelativeLocation, OldRelativeLocation, SelectedTemplate->RelativeLocation, UpdatedComponents);
+												FComponentEditorUtils::PropagateTransformPropertyChange(SceneComp, SceneComp->RelativeRotation, OldRelativeRotation, SelectedTemplate->RelativeRotation, UpdatedComponents);
+												FComponentEditorUtils::PropagateTransformPropertyChange(SceneComp, SceneComp->RelativeScale3D,  OldRelativeScale3D,  SelectedTemplate->RelativeScale3D,  UpdatedComponents);
 											}
 										}
 									}
diff --git a/Engine/Source/Editor/UnrealEd/Private/Kismet2/ComponentEditorUtils.cpp b/Engine/Source/Editor/UnrealEd/Private/Kismet2/ComponentEditorUtils.cpp
--- a/Engine/Source/Editor/UnrealEd/Private/Kismet2/ComponentEditorUtils.cpp
+++ b/Engine/Source/Editor/UnrealEd/Private/Kismet2/ComponentEditorUtils.cpp
@@ -88,86 +88,37 @@ void FComponentEditorUtils::GetArchetypeInstances( UObject* Object, TArray<UObje
 	}
 }
 
-template<typename T>
-struct FComponentTransformPropertyChangeHelper
-{
-	T* CurValuePtr;
-	const T OldVal;
-	const T NewVal;
-
-	FComponentTransformPropertyChangeHelper(USceneComponent* Instance, const UProperty* Property, T InOldVal, T InNewVal ) 
-		: CurValuePtr(NULL)
-		, OldVal(InOldVal)
-		, NewVal(InNewVal)
-	{
-		check(Instance && Property);
-		CurValuePtr = Property->ContainerPtrToValuePtr<T>(Instance);
-		check(CurValuePtr);
-	}
-
-	bool ShouldBeChanged() const
-	{
-		return (*CurValuePtr == OldVal);
-	}
-
-	void TryChange()
-	{
-		if (ShouldBeChanged())
-		{
-			*CurValuePtr = NewVal;
-		}
-	}
-};
-
 void FComponentEditorUtils::PropagateTransformPropertyChange(
-	class USceneComponent* SceneComponentTemplate,
-	const FTransformData& TransformOld,
-	const FTransformData& TransformNew,
+	class USceneComponent* InSceneComponentTemplate,
+	const FTransformData& OldDefaultTransform,
+	const FTransformData& NewDefaultTransform,
 	TSet<class USceneComponent*>& UpdatedComponents)
 {
-	check(SceneComponentTemplate != NULL);
+	check(InSceneComponentTemplate != nullptr);
 
 	TArray<UObject*> ArchetypeInstances;
-	FComponentEditorUtils::GetArchetypeInstances(SceneComponentTemplate, ArchetypeInstances);
+	FComponentEditorUtils::GetArchetypeInstances(InSceneComponentTemplate, ArchetypeInstances);
 	for(int32 InstanceIndex = 0; InstanceIndex < ArchetypeInstances.Num(); ++InstanceIndex)
 	{
-		USceneComponent* InstancedSceneComponent = FComponentEditorUtils::GetSceneComponent(ArchetypeInstances[InstanceIndex], SceneComponentTemplate);
-		if(InstancedSceneComponent && !UpdatedComponents.Contains(InstancedSceneComponent))
+		USceneComponent* InstancedSceneComponent = FComponentEditorUtils::GetSceneComponent(ArchetypeInstances[InstanceIndex], InSceneComponentTemplate);
+		if(InstancedSceneComponent != nullptr && !UpdatedComponents.Contains(InstancedSceneComponent))
 		{
 			static const UProperty* RelativeLocationProperty = FindFieldChecked<UProperty>( USceneComponent::StaticClass(), "RelativeLocation" );
-			static const UProperty* RelativeRotationProperty = FindFieldChecked<UProperty>( USceneComponent::StaticClass(), "RelativeRotation" );
-			static const UProperty* RelativeScale3DProperty = FindFieldChecked<UProperty>( USceneComponent::StaticClass(), "RelativeScale3D" );
-
-			FComponentTransformPropertyChangeHelper<FVector> Location(InstancedSceneComponent, RelativeLocationProperty, TransformOld.Trans, TransformNew.Trans);
-			FComponentTransformPropertyChangeHelper<FRotator> Rotation(InstancedSceneComponent, RelativeRotationProperty, TransformOld.Rot, TransformNew.Rot);
-			FComponentTransformPropertyChangeHelper<FVector> Scale(InstancedSceneComponent, RelativeScale3DProperty, TransformOld.Scale, TransformNew.Scale);
-			
-			const bool bShouldBeChanged = Location.ShouldBeChanged() || Rotation.ShouldBeChanged() || Scale.ShouldBeChanged();
-			if(bShouldBeChanged)
+			if(RelativeLocationProperty != nullptr)
 			{
-				// Ensure that this instance will be included in any undo/redo operations, and record it into the transaction buffer.
-				// Note: We don't do this for components that originate from script, because they will be re-instanced from the template after an undo, so there is no need to record them.
-				if(!InstancedSceneComponent->bCreatedByConstructionScript)
-				{
-					InstancedSceneComponent->SetFlags(RF_Transactional);
-					InstancedSceneComponent->Modify();
-				}
-
-				// We must also modify the owner, because we'll need script components to be reconstructed as part of an undo operation.
-				AActor* Owner = InstancedSceneComponent->GetOwner();
-				if(Owner != NULL)
-				{
-					Owner->Modify();
-				}
-
-				Location.TryChange();
-				Rotation.TryChange();
-				Scale.TryChange();
+				PropagateTransformPropertyChange(InstancedSceneComponent, RelativeLocationProperty, OldDefaultTransform.Trans, NewDefaultTransform.Trans, UpdatedComponents);
+			}
 
-				// Re-register the component with the scene so that transforms are updated for display
-				InstancedSceneComponent->ReregisterComponent();
+			static const UProperty* RelativeRotationProperty = FindFieldChecked<UProperty>( USceneComponent::StaticClass(), "RelativeRotation" );
+			if(RelativeRotationProperty != nullptr)
+			{
+				PropagateTransformPropertyChange(InstancedSceneComponent, RelativeRotationProperty, OldDefaultTransform.Rot, NewDefaultTransform.Rot, UpdatedComponents);
+			}
 
-				UpdatedComponents.Add(InstancedSceneComponent);
+			static const UProperty* RelativeScale3DProperty = FindFieldChecked<UProperty>( USceneComponent::StaticClass(), "RelativeScale3D" );
+			if(RelativeScale3DProperty != nullptr)
+			{
+				PropagateTransformPropertyChange(InstancedSceneComponent, RelativeScale3DProperty, OldDefaultTransform.Scale, NewDefaultTransform.Scale, UpdatedComponents);
 			}
 		}
 	}
diff --git a/Engine/Source/Editor/UnrealEd/Public/Kismet2/ComponentEditorUtils.h b/Engine/Source/Editor/UnrealEd/Public/Kismet2/ComponentEditorUtils.h
--- a/Engine/Source/Editor/UnrealEd/Public/Kismet2/ComponentEditorUtils.h
+++ b/Engine/Source/Editor/UnrealEd/Public/Kismet2/ComponentEditorUtils.h
@@ -26,6 +26,55 @@ public:
 		FVector Scale;
 	};
 
-	static void PropagateTransformPropertyChange(class USceneComponent* SceneComponentTemplate, const FTransformData& TransformOld, const FTransformData& TransformNew, TSet<class USceneComponent*>& UpdatedComponents);
+	// Given a template, propagates a default transform change to all instances of the template
+	static void PropagateTransformPropertyChange(class USceneComponent* InSceneComponentTemplate, const FTransformData& OldDefaultTransform, const FTransformData& NewDefaultTransform, TSet<class USceneComponent*>& UpdatedComponents);
 
+	// Given an instance of a template and a property, propagates a default value change to the instance (only if applicable)
+	template<typename T>
+	static void PropagateTransformPropertyChange(class USceneComponent* InSceneComponent, const class UProperty* InProperty, const T& OldDefaultValue, const T& NewDefaultValue, TSet<class USceneComponent*>& UpdatedComponents)
+	{
+		check(InProperty != nullptr);
+		check(InSceneComponent != nullptr);
+
+		T* CurrentValue = InProperty->ContainerPtrToValuePtr<T>(InSceneComponent);
+		if(CurrentValue != nullptr)
+		{
+			PropagateTransformPropertyChange(InSceneComponent, *CurrentValue, OldDefaultValue, NewDefaultValue, UpdatedComponents);
+		}
+	}
+
+	// Given an instance of a template and a current value, propagates a default value change to the instance (only if applicable)
+	template<typename T>
+	static void PropagateTransformPropertyChange(class USceneComponent* InSceneComponent, T& CurrentValue, const T& OldDefaultValue, const T& NewDefaultValue, TSet<class USceneComponent*>& UpdatedComponents)
+	{
+		check(InSceneComponent != nullptr);
+
+		// Propagate the change only if the current instanced value matches the previous default value (otherwise this could overwrite any per-instance override)
+		if(NewDefaultValue != OldDefaultValue && CurrentValue == OldDefaultValue)
+		{
+			// Ensure that this instance will be included in any undo/redo operations, and record it into the transaction buffer.
+			// Note: We don't do this for components that originate from script, because they will be re-instanced from the template after an undo, so there is no need to record them.
+			if(!InSceneComponent->bCreatedByConstructionScript)
+			{
+				InSceneComponent->SetFlags(RF_Transactional);
+				InSceneComponent->Modify();
+			}
+
+			// We must also modify the owner, because we'll need script components to be reconstructed as part of an undo operation.
+			AActor* Owner = InSceneComponent->GetOwner();
+			if(Owner != nullptr)
+			{
+				Owner->Modify();
+			}
+
+			// Modify the value
+			CurrentValue = NewDefaultValue;
+
+			// Re-register the component with the scene so that transforms are updated for display
+			InSceneComponent->ReregisterComponent();
+
+			// Add the component to the set of updated instances
+			UpdatedComponents.Add(InSceneComponent);
+		}
+	}
 };
\ No newline at end of file
