Media: Made it more clear that the media URL can also be a file on disk

diff --git a/Engine/Source/Editor/DetailCustomizations/Private/MediaPlayerCustomization.cpp b/Engine/Source/Editor/DetailCustomizations/Private/MediaPlayerCustomization.cpp
--- a/Engine/Source/Editor/DetailCustomizations/Private/MediaPlayerCustomization.cpp
+++ b/Engine/Source/Editor/DetailCustomizations/Private/MediaPlayerCustomization.cpp
@@ -34,7 +34,10 @@ void FMediaPlayerCustomization::CustomizeDetails( IDetailLayoutBuilder& DetailBu
 					+ SHorizontalBox::Slot()
 						.AutoWidth()
 						[
-							UrlProperty->CreatePropertyNameWidget()
+							SNew(STextBlock)
+								.Font(IDetailLayoutBuilder::GetDetailFont())
+								.Text(LOCTEXT("FileOrUrlPropertyName", "File or URL"))
+								.ToolTipText(UrlProperty->GetToolTipText())
 						]
 
 					+ SHorizontalBox::Slot()
diff --git a/Engine/Source/Runtime/MediaAssets/Public/MediaPlayer.h b/Engine/Source/Runtime/MediaAssets/Public/MediaPlayer.h
--- a/Engine/Source/Runtime/MediaAssets/Public/MediaPlayer.h
+++ b/Engine/Source/Runtime/MediaAssets/Public/MediaPlayer.h
@@ -299,7 +299,7 @@ protected:
 	UPROPERTY(EditAnywhere, BlueprintReadOnly, Category=Source)
 	TEnumAsByte<enum EMediaPlayerStreamModes> StreamMode;
 
-	/** The URL to the media file to be played. */
+	/** The path or URL to the media file to be played. */
 	UPROPERTY(EditAnywhere, Category=Source)
 	FString URL;
 
