Media: Made 'All Movie Files' the default file type filter

diff --git a/Engine/Source/Editor/DetailCustomizations/Private/MediaPlayerCustomization.cpp b/Engine/Source/Editor/DetailCustomizations/Private/MediaPlayerCustomization.cpp
--- a/Engine/Source/Editor/DetailCustomizations/Private/MediaPlayerCustomization.cpp
+++ b/Engine/Source/Editor/DetailCustomizations/Private/MediaPlayerCustomization.cpp
@@ -368,7 +368,7 @@ FString FMediaPlayerCustomization::HandleUrlPickerFileTypeFilter( ) const
 		AllFilters += TEXT("|") + Format.Value.ToString() + TEXT(" (*.") + Format.Key + TEXT(")|*.") + Format.Key;
 	}
 
-	Filter += TEXT("|All movie files (") + AllExtensions + TEXT(")|") + AllExtensions + AllFilters;
+	Filter = TEXT("All movie files (") + AllExtensions + TEXT(")|") + AllExtensions + TEXT("|") + Filter + AllFilters;
 
 	return Filter;
 }
