Editor: Fix loctext key collision that broke mobility combo button tooltips

diff --git a/Engine/Source/Editor/DetailCustomizations/Private/SceneComponentDetails.cpp b/Engine/Source/Editor/DetailCustomizations/Private/SceneComponentDetails.cpp
--- a/Engine/Source/Editor/DetailCustomizations/Private/SceneComponentDetails.cpp
+++ b/Engine/Source/Editor/DetailCustomizations/Private/SceneComponentDetails.cpp
@@ -493,7 +493,7 @@ void FSceneComponentDetails::CustomizeDetails( IDetailLayoutBuilder& DetailBuild
 			.Style(FEditorStyle::Get(), "Property.ToggleButton.Start")
 			.IsChecked(this, &FSceneComponentDetails::IsMobilityActive, MobilityPropertyWeak, EComponentMobility::Static)
 			.OnCheckStateChanged(this, &FSceneComponentDetails::OnMobilityChanged, MobilityPropertyWeak, EComponentMobility::Static)
-			.ToolTipText(LOCTEXT("Mobility_Movable", "A static object can't be changed in game.\n● Allows Baked Lighting\n● Fastest Rendering"))
+			.ToolTipText(LOCTEXT("Mobility_Static_Tooltip", "A static object can't be changed in game.\n● Allows Baked Lighting\n● Fastest Rendering"))
 			[
 				SNew(SHorizontalBox)
 
@@ -532,7 +532,7 @@ void FSceneComponentDetails::CustomizeDetails( IDetailLayoutBuilder& DetailBuild
 			.IsChecked(this, &FSceneComponentDetails::IsMobilityActive, MobilityPropertyWeak, EComponentMobility::Stationary)
 			.Style(FEditorStyle::Get(), ( ColumnIndex == 0 ) ? "Property.ToggleButton.Start" : "Property.ToggleButton.Middle")
 			.OnCheckStateChanged(this, &FSceneComponentDetails::OnMobilityChanged, MobilityPropertyWeak, EComponentMobility::Stationary)
-			.ToolTipText(LOCTEXT("Mobility_Movable", "A stationary light will only have its shadowing and bounced lighting from static geometry baked by Lightmass, all other lighting will be dynamic.  It can change color and intensity in game.\n● Can't Move\n● Allows Partial Baked Lighting\n● Dynamic Shadows"))
+			.ToolTipText(LOCTEXT("Mobility_Stationary_Tooltip", "A stationary light will only have its shadowing and bounced lighting from static geometry baked by Lightmass, all other lighting will be dynamic.  It can change color and intensity in game.\n● Can't Move\n● Allows Partial Baked Lighting\n● Dynamic Shadows"))
 			[
 				SNew(SHorizontalBox)
 
@@ -569,7 +569,7 @@ void FSceneComponentDetails::CustomizeDetails( IDetailLayoutBuilder& DetailBuild
 		.IsChecked(this, &FSceneComponentDetails::IsMobilityActive, MobilityPropertyWeak, EComponentMobility::Movable)
 		.Style(FEditorStyle::Get(), ( ColumnIndex == 0 ) ? "Property.ToggleButton" : "Property.ToggleButton.End")
 		.OnCheckStateChanged(this, &FSceneComponentDetails::OnMobilityChanged, MobilityPropertyWeak, EComponentMobility::Movable)
-		.ToolTipText(LOCTEXT("Mobility_Movable", "Movable objects can be moved and changed in game.\n● Totally Dynamic\n● Allows Dynamic Shadows\n● Slowest Rendering"))
+		.ToolTipText(LOCTEXT("Mobility_Movable_Tooltip", "Movable objects can be moved and changed in game.\n● Totally Dynamic\n● Allows Dynamic Shadows\n● Slowest Rendering"))
 		[
 			SNew(SHorizontalBox)
 
