Add option to force animation only import from FBX file containing skeletal mesh data.

diff --git a/Engine/Config/BaseEditorUserSettings.ini b/Engine/Config/BaseEditorUserSettings.ini
--- a/Engine/Config/BaseEditorUserSettings.ini
+++ b/Engine/Config/BaseEditorUserSettings.ini
@@ -550,6 +550,7 @@ bResampleAnimations=False
 bImportRigidMesh=False
 bCombineMeshes=True
 bCreatePhysicsAsset=True
+bImportMesh=True
 
 
 [/Script/UnrealEd.FbxStaticMeshImportData]
diff --git a/Engine/Source/Editor/DetailCustomizations/Private/FbxImportUIDetails.cpp b/Engine/Source/Editor/DetailCustomizations/Private/FbxImportUIDetails.cpp
--- a/Engine/Source/Editor/DetailCustomizations/Private/FbxImportUIDetails.cpp
+++ b/Engine/Source/Editor/DetailCustomizations/Private/FbxImportUIDetails.cpp
@@ -42,18 +42,32 @@ void FFbxImportUIDetails::CustomizeDetails( IDetailLayoutBuilder& DetailBuilder
 
 	MeshCategory.GetDefaultProperties(CategoryDefaultProperties);
 
-	EFBXImportType ImportType = ImportUI->MeshTypeToImport;
-	switch(ImportType)
+	switch(ImportUI->MeshTypeToImport)
 	{
 		case FBXIT_StaticMesh:
 			CollectChildPropertiesRecursive(StaticMeshDataProp, ExtraProperties);
 			break;
 		case FBXIT_SkeletalMesh:
-			CollectChildPropertiesRecursive(SkeletalMeshDataProp, ExtraProperties);
+			if(ImportUI->bImportMesh)
+			{
+				CollectChildPropertiesRecursive(SkeletalMeshDataProp, ExtraProperties);
+			}
+			else
+			{
+				ImportUI->MeshTypeToImport = FBXIT_Animation;
+			}
 			break;
 		default:
 			break;
 	}
+	EFBXImportType ImportType = ImportUI->MeshTypeToImport;
+
+	if(ImportUI->OriginalImportType == FBXIT_SkeletalMesh)
+	{
+		TSharedRef<IPropertyHandle> ImportMeshProp = DetailBuilder.GetProperty(GET_MEMBER_NAME_CHECKED(UFbxImportUI, bImportMesh));
+		ImportMeshProp->SetOnPropertyValueChanged(FSimpleDelegate::CreateSP(this, &FFbxImportUIDetails::ImportMeshToggleChanged));
+		MeshCategory.AddProperty(ImportMeshProp);
+	}
 
 	if(ImportType != FBXIT_Animation)
 	{
@@ -93,6 +107,18 @@ void FFbxImportUIDetails::CustomizeDetails( IDetailLayoutBuilder& DetailBuilder
 
 	// Animation Category
 	IDetailCategoryBuilder& AnimCategory = DetailBuilder.EditCategory("Animation", TEXT(""), ECategoryPriority::Important);
+
+	CategoryDefaultProperties.Empty();
+	AnimCategory.GetDefaultProperties(CategoryDefaultProperties);
+	for(TSharedRef<IPropertyHandle> Handle : CategoryDefaultProperties)
+	{
+		FString MetaData = Handle->GetMetaData(TEXT("ImportType"));
+		if(!IsImportTypeMetaDataValid(ImportType, MetaData))
+		{
+			DetailBuilder.HideProperty(Handle);
+		}
+	}
+
 	if(ImportType == FBXIT_Animation || ImportType == FBXIT_SkeletalMesh)
 	{
 		ExtraProperties.Empty();
@@ -210,4 +236,13 @@ void FFbxImportUIDetails::MeshImportModeChanged()
 	}
 }
 
+void FFbxImportUIDetails::ImportMeshToggleChanged()
+{
+	if(CachedDetailBuilder)
+	{
+		ImportUI->MeshTypeToImport = ImportUI->bImportMesh ? FBXIT_SkeletalMesh: FBXIT_Animation;
+		CachedDetailBuilder->ForceRefreshDetails();
+	}
+}
+
 #undef LOCTEXT_NAMESPACE
diff --git a/Engine/Source/Editor/DetailCustomizations/Private/FbxImportUIDetails.h b/Engine/Source/Editor/DetailCustomizations/Private/FbxImportUIDetails.h
--- a/Engine/Source/Editor/DetailCustomizations/Private/FbxImportUIDetails.h
+++ b/Engine/Source/Editor/DetailCustomizations/Private/FbxImportUIDetails.h
@@ -23,6 +23,9 @@ public:
 	/** Called if the mesh mode (static / skeletal) changes */
 	void MeshImportModeChanged();
 
+	/** Called if the import mesh option for skeletal meshes is changed */
+	void ImportMeshToggleChanged();
+
 	TWeakObjectPtr<UFbxImportUI> ImportUI;		// The UI data object being customised
 	IDetailLayoutBuilder* CachedDetailBuilder;	// The detail builder for this cusomtomisation
 
diff --git a/Engine/Source/Editor/UnrealEd/Classes/Factories/FbxImportUI.h b/Engine/Source/Editor/UnrealEd/Classes/Factories/FbxImportUI.h
--- a/Engine/Source/Editor/UnrealEd/Classes/Factories/FbxImportUI.h
+++ b/Engine/Source/Editor/UnrealEd/Classes/Factories/FbxImportUI.h
@@ -51,6 +51,10 @@ class UFbxImportUI : public UObject
 	UPROPERTY(EditAnywhere, Category = Mesh, meta = (ImportType = "StaticMesh|SkeletalMesh"))
 	bool bImportAsSkeletal;
 
+	/** Whether to import the mesh. Allows animation only import when importing a skeletal mesh. */
+	UPROPERTY(EditAnywhere, config, Category=Mesh, meta=(ImportType="SkeletalMesh"))
+	bool bImportMesh;
+
 	/** For static meshes, enabling this option will combine all meshes in the FBX into a single monolithic mesh in Unreal */
 	UPROPERTY(EditAnywhere, AdvancedDisplay, config, Category=Mesh, meta=(ToolTip="If enabled, combines all meshes into a single mesh", ImportType="StaticMesh"))
 	uint32 bCombineMeshes:1;
@@ -80,7 +84,7 @@ class UFbxImportUI : public UObject
 	uint32 bImportRigidMesh:1;
 
 	/** Enable this option to use default sample rate for the imported animation at 30 frames per second */
-	UPROPERTY(EditAnywhere, AdvancedDisplay, config, Category=Animation, meta=(editcondition = "bImportAnimations", ToolTip="If enabled, samples all animation curves to 30 FPS", ImportType="SkeletalMesh"))
+	UPROPERTY(EditAnywhere, AdvancedDisplay, config, Category=Animation, meta=(editcondition = "bImportAnimations", ToolTip="If enabled, samples all animation curves to 30 FPS", ImportType="SkeletalMesh|Animation"))
 	uint32 bUseDefaultSampleRate:1;
 
 	/** Whether to automatically create Unreal materials for materials found in the FBX scene */
