Fixed "mirror" option in scale details context menu so that it persists.
#ttp 347906 - Editor: righclick "mirror" for scale values in details window is not working
#branch UE4
#proj Editor.DetailCustomizations

diff --git a/Engine/Source/Editor/DetailCustomizations/Private/ComponentTransformDetails.cpp b/Engine/Source/Editor/DetailCustomizations/Private/ComponentTransformDetails.cpp
--- a/Engine/Source/Editor/DetailCustomizations/Private/ComponentTransformDetails.cpp
+++ b/Engine/Source/Editor/DetailCustomizations/Private/ComponentTransformDetails.cpp
@@ -865,17 +865,20 @@ void FComponentTransformDetails::ExtendZScaleContextMenu( FMenuBuilder& MenuBuil
 
 void FComponentTransformDetails::OnXScaleMirrored()
 {
+	FSlateApplication::Get().ClearKeyboardFocus(EKeyboardFocusCause::Mouse);
 	ScaleObject( 1.0f, 0, true, LOCTEXT( "MirrorActorScaleX", "Mirror actor scale X" ) );
 }
 
 void FComponentTransformDetails::OnYScaleMirrored()
 {
-	ScaleObject( 1.0f, 1, true, LOCTEXT( "MirrorActorScaleY", "Mirror actor scale Y" ) );
+	FSlateApplication::Get().ClearKeyboardFocus(EKeyboardFocusCause::Mouse);
+	ScaleObject(1.0f, 1, true, LOCTEXT("MirrorActorScaleY", "Mirror actor scale Y"));
 }
 
 void FComponentTransformDetails::OnZScaleMirrored()
 {
-	ScaleObject( 1.0f, 2, true, LOCTEXT( "MirrorActorScaleZ", "Mirror actor scale Z" ) );
+	FSlateApplication::Get().ClearKeyboardFocus(EKeyboardFocusCause::Mouse);
+	ScaleObject(1.0f, 2, true, LOCTEXT("MirrorActorScaleZ", "Mirror actor scale Z"));
 }
 
 /**
