customize destructible mesh to make editor a bit easier to use

diff --git a/Engine/Source/Editor/DetailCustomizations/Private/DestructibleMeshDetails.cpp b/Engine/Source/Editor/DetailCustomizations/Private/DestructibleMeshDetails.cpp
new file mode 100644
--- /dev/null
+++ b/Engine/Source/Editor/DetailCustomizations/Private/DestructibleMeshDetails.cpp
@@ -0,0 +1,72 @@
+// Copyright 1998-2014 Epic Games, Inc. All Rights Reserved.
+
+#include "DetailCustomizationsPrivatePCH.h"
+#include "DestructibleMeshDetails.h"
+#include "ScopedTransaction.h"
+#include "ObjectEditorUtils.h"
+#include "IDocumentation.h"
+
+#define LOCTEXT_NAMESPACE "DestructibleMeshDetails"
+
+TSharedRef<IDetailCustomization> FDestructibleMeshDetails::MakeInstance()
+{
+	return MakeShareable(new FDestructibleMeshDetails);
+}
+
+void AddStructToDetails(FName CategoryName, FName PropertyName, IDetailLayoutBuilder& DetailBuilder, bool bInline = true, bool bAdvanced = false)
+{
+	IDetailCategoryBuilder& Category = DetailBuilder.EditCategory(CategoryName, TEXT(""), ECategoryPriority::Important);
+	TSharedPtr<IPropertyHandle> Params = DetailBuilder.GetProperty(PropertyName);
+	if (Params.IsValid())
+	{
+		EPropertyLocation::Type PropertyLocation = bAdvanced ? EPropertyLocation::Advanced : EPropertyLocation::Default;
+		if (bInline)
+		{
+			uint32 NumChildren = 0;
+			Params->GetNumChildren(NumChildren);
+
+			// add all collision properties
+			for (uint32 ChildIndex = 0; ChildIndex < NumChildren; ++ChildIndex)
+			{
+				TSharedPtr<IPropertyHandle> ChildProperty = Params->GetChildHandle(ChildIndex);
+				Category.AddProperty(ChildProperty, PropertyLocation);
+			}
+		}
+		else
+		{
+			Category.AddProperty(Params, PropertyLocation);
+		}
+	}
+}
+
+void FDestructibleMeshDetails::CustomizeDetails(IDetailLayoutBuilder& DetailBuilder)
+{
+		//we always hide bodysetup as it's not useful in this editor
+		TSharedPtr<IPropertyHandle> BodySetupHandler = DetailBuilder.GetProperty("BodySetup");
+		if (BodySetupHandler.IsValid())
+		{
+			DetailBuilder.HideProperty(BodySetupHandler);
+		}
+		
+		//rest of customization is just moving stuff out of DefaultDestructibleParameters so it's nicer to view
+		TSharedPtr<IPropertyHandle> DefaultParams = DetailBuilder.GetProperty("DefaultDestructibleParameters");
+		if (DefaultParams.IsValid() == false)
+		{
+			return;
+		}
+
+		AddStructToDetails("Damage", "DefaultDestructibleParameters.DamageParameters", DetailBuilder);
+		AddStructToDetails("Damage", "DefaultDestructibleParameters.AdvancedParameters", DetailBuilder, true, true);
+		AddStructToDetails("Debris", "DefaultDestructibleParameters.DebrisParameters", DetailBuilder);
+		AddStructToDetails("Flags", "DefaultDestructibleParameters.Flags", DetailBuilder);
+		AddStructToDetails("HierarchyDepth", "DefaultDestructibleParameters.SpecialHierarchyDepths", DetailBuilder);
+		AddStructToDetails("HierarchyDepth", "DefaultDestructibleParameters.DepthParameters", DetailBuilder, false, true);
+
+		
+
+		//hide the default params as we've taken everything out of it
+		DetailBuilder.HideProperty(DefaultParams);
+}
+
+#undef LOCTEXT_NAMESPACE
+
diff --git a/Engine/Source/Editor/DetailCustomizations/Private/DestructibleMeshDetails.h b/Engine/Source/Editor/DetailCustomizations/Private/DestructibleMeshDetails.h
new file mode 100644
--- /dev/null
+++ b/Engine/Source/Editor/DetailCustomizations/Private/DestructibleMeshDetails.h
@@ -0,0 +1,17 @@
+// Copyright 1998-2014 Epic Games, Inc. All Rights Reserved.
+
+#pragma once
+
+class FDestructibleMeshDetails : public IDetailCustomization
+{
+public:
+	/** Makes a new instance of this detail layout class for a specific detail view requesting it */
+	static TSharedRef<IDetailCustomization> MakeInstance();
+
+	/** IDetailCustomization interface */
+	virtual void CustomizeDetails( IDetailLayoutBuilder& DetailBuilder ) override;
+
+private:
+	TArray< TWeakObjectPtr<UObject> > ObjectsCustomized;
+};
+
diff --git a/Engine/Source/Editor/DetailCustomizations/Private/DetailCustomizations.cpp b/Engine/Source/Editor/DetailCustomizations/Private/DetailCustomizations.cpp
--- a/Engine/Source/Editor/DetailCustomizations/Private/DetailCustomizations.cpp
+++ b/Engine/Source/Editor/DetailCustomizations/Private/DetailCustomizations.cpp
@@ -44,6 +44,7 @@
 #include "BodyInstanceCustomization.h"
 #include "VehicleTransmissionDataCustomization.h"
 #include "BodySetupDetails.h"
+#include "DestructibleMeshDetails.h"
 #include "Runtime/Engine/Classes/PhysicsEngine/BodySetup.h"
 #include "SlateBrushCustomization.h"
 #include "SlateSoundCustomization.h"
@@ -227,6 +228,7 @@ void FDetailCustomizationsModule::RegisterObjectCustomizations()
 	RegisterCustomClassLayout("SoundWave", FOnGetDetailCustomizationInstance::CreateStatic(&FSoundWaveDetails::MakeInstance));
 	RegisterCustomClassLayout("DialogueWave", FOnGetDetailCustomizationInstance::CreateStatic(&FDialogueWaveDetails::MakeInstance));
 	RegisterCustomClassLayout("BodySetup", FOnGetDetailCustomizationInstance::CreateStatic(&FBodySetupDetails::MakeInstance));
+	RegisterCustomClassLayout("DestructibleMesh", FOnGetDetailCustomizationInstance::CreateStatic(&FDestructibleMeshDetails::MakeInstance));
 	RegisterCustomClassLayout("PhysicsConstraintTemplate", FOnGetDetailCustomizationInstance::CreateStatic(&FPhysicsConstraintComponentDetails::MakeInstance));
 	RegisterCustomClassLayout("PhysicsConstraintComponent", FOnGetDetailCustomizationInstance::CreateStatic(&FPhysicsConstraintComponentDetails::MakeInstance));
 	RegisterCustomClassLayout("CollisionProfile", FOnGetDetailCustomizationInstance::CreateStatic(&FCollisionProfileDetails::MakeInstance));
diff --git a/Engine/Source/Runtime/Engine/Classes/Engine/DestructibleMesh.h b/Engine/Source/Runtime/Engine/Classes/Engine/DestructibleMesh.h
--- a/Engine/Source/Runtime/Engine/Classes/Engine/DestructibleMesh.h
+++ b/Engine/Source/Runtime/Engine/Classes/Engine/DestructibleMesh.h
@@ -418,7 +418,7 @@ class UDestructibleMesh : public USkeletalMesh
 	class UPhysicalMaterial* DestructiblePhysicalMaterial_DEPRECATED;
 
 	/** Fracture effects for each fracture level, unless overridden in the component. */
-	UPROPERTY(EditAnywhere, editfixedsize, Category=DestructibleMesh)
+	UPROPERTY(EditAnywhere, editfixedsize, Category=Effects)
 	TArray<struct FFractureEffect> FractureEffects;
 
 	/** Physics data.  Fields from BodySetup which are relevant to the DestructibleMesh will be used. */
