Add mass override

diff --git a/Engine/Source/Editor/DetailCustomizations/Private/BodySetupDetails.cpp b/Engine/Source/Editor/DetailCustomizations/Private/BodySetupDetails.cpp
--- a/Engine/Source/Editor/DetailCustomizations/Private/BodySetupDetails.cpp
+++ b/Engine/Source/Editor/DetailCustomizations/Private/BodySetupDetails.cpp
@@ -35,23 +35,6 @@ void FBodySetupDetails::CustomizeDetails( IDetailLayoutBuilder& DetailBuilder )
 			// Get the objects being customized so we can enable/disable editing of 'Simulate Physics'
 			DetailBuilder.GetObjectsBeingCustomized(ObjectsCustomized);
 
-			PhysicsCategory.AddCustomRow(TEXT("Mass"), false)
-				.IsEnabled(TAttribute<bool>(this, &FBodySetupDetails::IsBodyMassReadOnly))
-				.NameContent()
-				[
-					SNew (STextBlock)
-					.Text(NSLOCTEXT("MassInKG", "MassInKG_Name", "Mass in KG"))
-					.ToolTip( IDocumentation::Get()->CreateToolTip(LOCTEXT("MassInKG", "Mass of the body in KG"), NULL, TEXT("Shared/Physics"), TEXT("MassInKG")) )
-					.Font( IDetailLayoutBuilder::GetDetailFont() )
-				]
-			.ValueContent()
-				[
-					SNew(SEditableTextBox)
-					.Text(this, &FBodySetupDetails::OnGetBodyMass)
-					.IsReadOnly(this, &FBodySetupDetails::IsBodyMassReadOnly)
-					.Font(IDetailLayoutBuilder::GetDetailFont())
-				];
-
 			// add all properties of this now - after adding 
 			for (uint32 ChildIndex=0; ChildIndex < NumChildren; ++ChildIndex)
 			{
@@ -67,6 +50,43 @@ void FBodySetupDetails::CustomizeDetails( IDetailLayoutBuilder& DetailBuilder )
 					//also hide bAutoWeld for phat
 					continue;
 				}
+				else if (ChildProperty->GetProperty()->GetName() == TEXT("MassInKg"))
+				{
+					PhysicsCategory.AddCustomRow(TEXT("Mass"), false)
+						.IsEnabled(TAttribute<bool>(this, &FBodySetupDetails::IsBodyMassEnabled))
+						.NameContent()
+						[
+							ChildProperty->CreatePropertyNameWidget()
+						]
+					.ValueContent()
+						[
+							SNew(SVerticalBox)
+							+ SVerticalBox::Slot()
+							.AutoHeight()
+							[
+								SNew(SEditableTextBox)
+								.Text(this, &FBodySetupDetails::OnGetBodyMass)
+								.IsReadOnly(this, &FBodySetupDetails::IsBodyMassReadOnly)
+								.Font(IDetailLayoutBuilder::GetDetailFont())
+								.Visibility(this, &FBodySetupDetails::IsMassVisible, false)
+							]
+
+							+ SVerticalBox::Slot()
+								.AutoHeight()
+								[
+									SNew(SVerticalBox)
+									.Visibility(this, &FBodySetupDetails::IsMassVisible, true)
+									+ SVerticalBox::Slot()
+									.AutoHeight()
+									[
+										ChildProperty->CreatePropertyValueWidget()
+									]
+								]
+
+						];
+
+					continue;
+				}
 				if (Category == TEXT("Physics"))
 				{
 					PhysicsCategory.AddProperty(ChildProperty);
@@ -114,5 +134,37 @@ FText FBodySetupDetails::OnGetBodyMass() const
 	return FText::AsNumber(Mass);
 }
 
+EVisibility FBodySetupDetails::IsMassVisible(bool bOverrideMass) const
+{
+	bool bIsMassReadOnly = IsBodyMassReadOnly();
+	if (bOverrideMass)
+	{
+		return bIsMassReadOnly ? EVisibility::Collapsed : EVisibility::Visible;
+	}
+	else
+	{
+		return bIsMassReadOnly ? EVisibility::Visible : EVisibility::Collapsed;
+	}
+}
+
+
+bool FBodySetupDetails::IsBodyMassReadOnly() const
+{
+	for (auto ObjectIt = ObjectsCustomized.CreateConstIterator(); ObjectIt; ++ObjectIt)
+	{
+		if (ObjectIt->IsValid() && (*ObjectIt)->IsA(UBodySetup::StaticClass()))
+		{
+			UBodySetup* BS = Cast<UBodySetup>(ObjectIt->Get());
+			if (BS->DefaultInstance.bOverrideMass == false)
+			{
+				return true;
+			}
+		}
+	}
+
+	return false;
+}
+
+
 #undef LOCTEXT_NAMESPACE
 
diff --git a/Engine/Source/Editor/DetailCustomizations/Private/BodySetupDetails.h b/Engine/Source/Editor/DetailCustomizations/Private/BodySetupDetails.h
--- a/Engine/Source/Editor/DetailCustomizations/Private/BodySetupDetails.h
+++ b/Engine/Source/Editor/DetailCustomizations/Private/BodySetupDetails.h
@@ -12,8 +12,9 @@ public:
 	virtual void CustomizeDetails( IDetailLayoutBuilder& DetailBuilder ) override;
 
 	FText OnGetBodyMass() const;
-	bool IsBodyMassReadOnly() const { return true; }
+	bool IsBodyMassReadOnly() const;
 	bool IsBodyMassEnabled() const { return !IsBodyMassReadOnly(); }
+	EVisibility IsMassVisible(bool bOverrideMass) const;
 
 private:
 	TArray< TWeakObjectPtr<UObject> > ObjectsCustomized;
diff --git a/Engine/Source/Editor/DetailCustomizations/Private/PrimitiveComponentDetails.cpp b/Engine/Source/Editor/DetailCustomizations/Private/PrimitiveComponentDetails.cpp
--- a/Engine/Source/Editor/DetailCustomizations/Private/PrimitiveComponentDetails.cpp
+++ b/Engine/Source/Editor/DetailCustomizations/Private/PrimitiveComponentDetails.cpp
@@ -119,13 +119,31 @@ EVisibility FPrimitiveComponentDetails::IsAutoWeldVisible() const
 {
 	for (int32 i = 0; i < ObjectsCustomized.Num(); ++i)
 	{
-		if (ObjectsCustomized[i].IsValid() && !ObjectsCustomized[i]->IsA(UStaticMeshComponent::StaticClass()))
+		if (ObjectsCustomized[i].IsValid())
+		if(ObjectsCustomized[i]->IsA(UStaticMeshComponent::StaticClass()))
 		{
-			return EVisibility::Collapsed;
+			return EVisibility::Visible;
+		}
+		else if (ObjectsCustomized[i]->IsA(USkeletalMeshComponent::StaticClass()))
+		{
+			return EVisibility::Visible;
 		}
 	}
 
-	return EVisibility::Visible;
+	return EVisibility::Collapsed;
+}
+
+EVisibility FPrimitiveComponentDetails::IsMassVisible(bool bOverrideMass) const
+{
+	bool bIsMassReadOnly = IsBodyMassReadOnly();
+	if (bOverrideMass)
+	{
+		return bIsMassReadOnly ? EVisibility::Collapsed : EVisibility::Visible;
+	}
+	else
+	{
+		return bIsMassReadOnly ? EVisibility::Visible : EVisibility::Collapsed;
+	}
 }
 
 void FPrimitiveComponentDetails::CustomizeDetails( IDetailLayoutBuilder& DetailBuilder )
@@ -150,34 +168,20 @@ void FPrimitiveComponentDetails::CustomizeDetails( IDetailLayoutBuilder& DetailB
 			DetailBuilder.GetObjectsBeingCustomized(ObjectsCustomized);
 
 			bool bDisplayMass = true;
+			bool bDisplayMassOverride = true;
 
 			for (int32 i = 0; i < ObjectsCustomized.Num(); ++i)
 			{
 				if (ObjectsCustomized[i].IsValid() && ObjectsCustomized[i]->IsA(UDestructibleComponent::StaticClass()))
 				{
 					bDisplayMass = false;
-					break;
+					bDisplayMassOverride = false;
 				}
-			}
 
-			if (bDisplayMass)
-			{
-				PhysicsCategory.AddCustomRow(TEXT("Mass"), false)
-					.IsEnabled(TAttribute<bool>(this, &FPrimitiveComponentDetails::IsBodyMassEnabled))
-					.NameContent()
-					[
-						SNew(STextBlock)
-						.Text(NSLOCTEXT("MassInKG", "MassInKG_Name", "Mass in KG"))
-						.ToolTip(IDocumentation::Get()->CreateToolTip(LOCTEXT("MassInKG", "Mass of the body in KG"), NULL, TEXT("Shared/Physics"), TEXT("MassInKG")))
-						.Font(IDetailLayoutBuilder::GetDetailFont())
-					]
-					.ValueContent()
-					[
-						SNew(SEditableTextBox)
-						.Text(this, &FPrimitiveComponentDetails::OnGetBodyMass)
-						.IsReadOnly(this, &FPrimitiveComponentDetails::IsBodyMassReadOnly)
-						.Font(IDetailLayoutBuilder::GetDetailFont())
-					];
+				if (ObjectsCustomized[i].IsValid() && ObjectsCustomized[i]->IsA(USkeletalMeshComponent::StaticClass()))
+				{
+					bDisplayMassOverride = false;
+				}
 			}
 
 			// add all physics properties now - after adding mass
@@ -213,6 +217,52 @@ void FPrimitiveComponentDetails::CustomizeDetails( IDetailLayoutBuilder& DetailB
 						PhysicsCategory.AddProperty(ChildProperty).Visibility(TAttribute<EVisibility>(this, &FPrimitiveComponentDetails::IsAutoWeldVisible))
 																  .EditCondition(TAttribute<bool>(this, &FPrimitiveComponentDetails::IsAutoWeldEditable), NULL);
 					}
+					else if (PropName == TEXT("bOverrideMass"))
+					{
+						if (bDisplayMassOverride)
+						{
+							PhysicsCategory.AddProperty(ChildProperty);
+						}
+					}
+					else if (PropName == TEXT("MassInKg"))
+					{
+						
+						if (bDisplayMass)
+						{
+							PhysicsCategory.AddCustomRow(TEXT("Mass"), false)
+								.IsEnabled(TAttribute<bool>(this, &FPrimitiveComponentDetails::IsBodyMassEnabled))
+								.NameContent()
+								[
+									ChildProperty->CreatePropertyNameWidget()
+								]
+							.ValueContent()
+								[
+									SNew(SVerticalBox)
+									+ SVerticalBox::Slot()
+									.AutoHeight()
+									[
+										SNew(SEditableTextBox)
+										.Text(this, &FPrimitiveComponentDetails::OnGetBodyMass)
+										.IsReadOnly(this, &FPrimitiveComponentDetails::IsBodyMassReadOnly)
+										.Font(IDetailLayoutBuilder::GetDetailFont())
+										.Visibility(this, &FPrimitiveComponentDetails::IsMassVisible, false)
+									]
+
+									+ SVerticalBox::Slot()
+									.AutoHeight()
+									[
+										SNew(SVerticalBox)
+										.Visibility(this, &FPrimitiveComponentDetails::IsMassVisible, true)
+										+ SVerticalBox::Slot()
+										.AutoHeight()
+										[
+											ChildProperty->CreatePropertyValueWidget()
+										]
+									]
+									
+								];
+						}
+					}
 					else
 					{
 						PhysicsCategory.AddProperty(ChildProperty);
@@ -307,6 +357,22 @@ FText FPrimitiveComponentDetails::OnGetBodyMass() const
 	return FText::AsNumber(Mass);
 }
 
+bool FPrimitiveComponentDetails::IsBodyMassReadOnly() const
+{
+	for (auto ObjectIt = ObjectsCustomized.CreateConstIterator(); ObjectIt; ++ObjectIt)
+	{
+		if (ObjectIt->IsValid() && (*ObjectIt)->IsA(UPrimitiveComponent::StaticClass()))
+		{
+			if (UPrimitiveComponent* Comp = Cast<UPrimitiveComponent>(ObjectIt->Get()))
+			{
+				if (Comp->BodyInstance.bOverrideMass == false) { return true; }
+			}
+		}
+	}
+
+	return false;
+}
+
 
 #undef LOCTEXT_NAMESPACE
 
diff --git a/Engine/Source/Editor/DetailCustomizations/Private/PrimitiveComponentDetails.h b/Engine/Source/Editor/DetailCustomizations/Private/PrimitiveComponentDetails.h
--- a/Engine/Source/Editor/DetailCustomizations/Private/PrimitiveComponentDetails.h
+++ b/Engine/Source/Editor/DetailCustomizations/Private/PrimitiveComponentDetails.h
@@ -28,7 +28,7 @@ private:
 	bool IsUseAsyncEditable() const;
 
 	FText OnGetBodyMass() const;
-	bool IsBodyMassReadOnly() const { return true; }
+	bool IsBodyMassReadOnly() const;
 	bool IsBodyMassEnabled() const { return !IsBodyMassReadOnly(); }
 
 private:
@@ -39,5 +39,6 @@ private:
 
 	bool IsAutoWeldEditable() const;
 	EVisibility IsAutoWeldVisible() const;
+	EVisibility IsMassVisible(bool bOverrideMass) const;
 };
 
diff --git a/Engine/Source/Runtime/Engine/Classes/PhysicsEngine/BodyInstance.h b/Engine/Source/Runtime/Engine/Classes/PhysicsEngine/BodyInstance.h
--- a/Engine/Source/Runtime/Engine/Classes/PhysicsEngine/BodyInstance.h
+++ b/Engine/Source/Runtime/Engine/Classes/PhysicsEngine/BodyInstance.h
@@ -198,6 +198,15 @@ public:
 	UPROPERTY()
 	uint32 bUpdateMassWhenScaleChanges:1;
 
+	/** If true, mass will not be automatically computed and you must set it directly */
+	UPROPERTY(EditAnywhere, BlueprintReadOnly, Category=Physics)
+	uint32 bOverrideMass : 1;
+
+	/**Mass of the body in KG. By default we compute this based on physical material and mass scale.
+	   *@see bOverrideMass to set this directly */
+	UPROPERTY(EditAnywhere, BlueprintReadOnly, Category = Physics, meta = (editcondition = "bOverrideMass", ClampMin = "0.001", UIMin = "0.001"))
+	float MassInKg;
+
 	/** Locks physical movement along specified axis.*/
 	UPROPERTY(EditAnywhere, Category = Physics, meta=(DisplayName="Locked Axis"))
 	TEnumAsByte<ELockedAxis::Type> LockedAxisMode;
diff --git a/Engine/Source/Runtime/Engine/Private/PhysicsEngine/BodyInstance.cpp b/Engine/Source/Runtime/Engine/Private/PhysicsEngine/BodyInstance.cpp
--- a/Engine/Source/Runtime/Engine/Private/PhysicsEngine/BodyInstance.cpp
+++ b/Engine/Source/Runtime/Engine/Private/PhysicsEngine/BodyInstance.cpp
@@ -203,6 +203,8 @@ FBodyInstance::FBodyInstance()
 , bEnableGravity(true)
 , bUseAsyncScene(false)
 , bUpdateMassWhenScaleChanges(false)
+, bOverrideMass(false)
+, MassInKg(100.f)
 , LockedAxisMode(0)
 , CustomLockedAxis(FVector::ZeroVector)
 , DOFConstraint(NULL)
@@ -2760,22 +2762,27 @@ void FBodyInstance::UpdateMassProperties()
 	PxRigidDynamic* PRigidDynamic = GetPxRigidDynamic();
 	if ((PRigidDynamic != NULL) && (GetNumSimShapes(PRigidDynamic) > 0))
 	{
-		// First, reset mass to default
 
 		// physical material - nothing can weigh less than hydrogen (0.09 kg/m^3)
 		float DensityKGPerCubicUU = FMath::Max(KgPerM3ToKgPerCm3(0.09f), gPerCm3ToKgPerCm3(PhysMat->Density));
 		PxRigidBodyExt::updateMassAndInertia(*PRigidDynamic, DensityKGPerCubicUU);
 
-		// Then scale mass to avoid big differences between big and small objects.
+		//grab OldMass so we can apply new mass while maintaining inertia tensor
 		float OldMass = PRigidDynamic->getMass();
+		float NewMass = 0.f;
 
-		float UsePow = FMath::Clamp<float>(PhysMat->RaiseMassToPower, KINDA_SMALL_NUMBER, 1.f);
-		float NewMass = FMath::Pow(OldMass, UsePow);
-
-		// Apply user-defined mass scaling.
-		NewMass *= FMath::Clamp<float>(MassScale, 0.01f, 100.0f);
-
-		//UE_LOG(LogPhysics, Log,  TEXT("OldMass: %f NewMass: %f"), OldMass, NewMass );
+		if (bOverrideMass == false)
+		{
+			float UsePow = FMath::Clamp<float>(PhysMat->RaiseMassToPower, KINDA_SMALL_NUMBER, 1.f);
+			NewMass = FMath::Pow(OldMass, UsePow);
+			
+			// Apply user-defined mass scaling.
+			NewMass *= FMath::Clamp<float>(MassScale, 0.01f, 100.0f);
+		}
+		else
+		{
+			NewMass = FMath::Max(MassInKg, 0.001f);	//min weight of 1g
+		}
 
 		check (NewMass > 0.f);
 
@@ -2801,12 +2808,22 @@ void FBodyInstance::UpdateMassProperties()
 	{
 		//@TODO: BOX2D: Implement COMNudge, Unreal 'funky' mass algorithm, etc... for UpdateMassProperties (if we don't update the formula, we need to update the displayed mass in the details panel)
 
-		// Unreal material density is in g/cm^3, and Box2D density is in kg/m^2
-		// physical material - nothing can weigh less than hydrogen (0.09 kg/m^3)
-		const float DensityKGPerCubicCM = FMath::Max(0.00009f, PhysMat->Density * 0.001f);
-		const float DensityKGPerCubicM = DensityKGPerCubicCM * 1000.0f;
-		const float DensityKGPerSquareM = DensityKGPerCubicM * 0.1f; //@TODO: BOX2D: Should there be a thickness property for mass calculations?
-		const float MassScaledDensity = DensityKGPerSquareM * FMath::Clamp<float>(MassScale, 0.01f, 100.0f);
+		float MassScaledDensity = 0.f;
+		if (bOverrideMass == false)
+		{
+			// Unreal material density is in g/cm^3, and Box2D density is in kg/m^2
+			// physical material - nothing can weigh less than hydrogen (0.09 kg/m^3)
+			float DensityKGPerCubicCM = FMath::Max(KgPerM3ToKgPerCm3(0.09f), gPerCm3ToKgPerCm3(PhysMat->Density));
+			const float DensityKGPerCubicM = DensityKGPerCubicCM * 1000.0f;
+			const float DensityKGPerSquareM = DensityKGPerCubicM * 0.1f; //@TODO: BOX2D: Should there be a thickness property for mass calculations?
+			MassScaledDensity = DensityKGPerSquareM * FMath::Clamp<float>(MassScale, 0.01f, 100.0f);
+		}
+		else
+		{
+			MassScaledDensity = FMath::Max(MassInKg, 0.001f);	//min weight of 1g	//TODO: this is actually wrong because we're assuming mass and density are the same thing, but good enough for now
+		}
+
+		check(MassScaledDensity > 0.f);
 
 		// Apply the density
 		for (b2Fixture* Fixture = BodyInstancePtr->GetFixtureList(); Fixture; Fixture = Fixture->GetNext())
diff --git a/Engine/Source/Runtime/Engine/Private/PrimitiveComponentPhysics.cpp b/Engine/Source/Runtime/Engine/Private/PrimitiveComponentPhysics.cpp
--- a/Engine/Source/Runtime/Engine/Private/PrimitiveComponentPhysics.cpp
+++ b/Engine/Source/Runtime/Engine/Private/PrimitiveComponentPhysics.cpp
@@ -382,6 +382,11 @@ float UPrimitiveComponent::GetMass() const
 
 float UPrimitiveComponent::CalculateMass(FName)
 {
+	if (BodyInstance.bOverrideMass)
+	{
+		return BodyInstance.MassInKg;
+	}
+
 	if (BodyInstance.BodySetup.IsValid())
 	{
 		return BodyInstance.BodySetup->CalculateMass(this);
