Fix mac and undo accidental typo submission

diff --git a/Engine/Source/Editor/DetailCustomizations/Private/PrimitiveComponentDetails.cpp b/Engine/Source/Editor/DetailCustomizations/Private/PrimitiveComponentDetails.cpp
--- a/Engine/Source/Editor/DetailCustomizations/Private/PrimitiveComponentDetails.cpp
+++ b/Engine/Source/Editor/DetailCustomizations/Private/PrimitiveComponentDetails.cpp
@@ -119,12 +119,7 @@ EVisibility FPrimitiveComponentDetails::IsAutoWeldVisible() const
 {
 	for (int32 i = 0; i < ObjectsCustomized.Num(); ++i)
 	{
-		if (ObjectsCustomized[i].IsValid())
-		if(ObjectsCustomized[i]->IsA(UStaticMeshComponent::StaticClass()))
-		{
-			return EVisibility::Visible;
-		}
-		else if (ObjectsCustomized[i]->IsA(USkeletalMeshComponent::StaticClass()))
+		if (ObjectsCustomized[i].IsValid() && !ObjectsCustomized[i]->IsA(UStaticMeshComponent::StaticClass()))
 		{
 			return EVisibility::Visible;
 		}
