fix auto weld being accidently hidden as result of merge issues

diff --git a/Engine/Source/Editor/DetailCustomizations/Private/PrimitiveComponentDetails.cpp b/Engine/Source/Editor/DetailCustomizations/Private/PrimitiveComponentDetails.cpp
--- a/Engine/Source/Editor/DetailCustomizations/Private/PrimitiveComponentDetails.cpp
+++ b/Engine/Source/Editor/DetailCustomizations/Private/PrimitiveComponentDetails.cpp
@@ -121,11 +121,11 @@ EVisibility FPrimitiveComponentDetails::IsAutoWeldVisible() const
 	{
 		if (ObjectsCustomized[i].IsValid() && !ObjectsCustomized[i]->IsA(UStaticMeshComponent::StaticClass()))
 		{
-			return EVisibility::Visible;
+			return EVisibility::Collapsed;
 		}
 	}
 
-	return EVisibility::Collapsed;
+	return EVisibility::Visible;
 }
 
 EVisibility FPrimitiveComponentDetails::IsMassVisible(bool bOverrideMass) const
