Media: Fixed spelling error (TTP# 348611)

diff --git a/Engine/Source/Editor/DetailCustomizations/Private/MediaPlayerCustomization.cpp b/Engine/Source/Editor/DetailCustomizations/Private/MediaPlayerCustomization.cpp
--- a/Engine/Source/Editor/DetailCustomizations/Private/MediaPlayerCustomization.cpp
+++ b/Engine/Source/Editor/DetailCustomizations/Private/MediaPlayerCustomization.cpp
@@ -107,7 +107,7 @@ void FMediaPlayerCustomization::CustomizeDetails( IDetailLayoutBuilder& DetailBu
 			.NameContent()
 			[
 				SNew(STextBlock)
-					.Text(LOCTEXT("ForwardRatesUnthinned", "Forward Rates (Unhinned)"))
+					.Text(LOCTEXT("ForwardRatesUnthinned", "Forward Rates (Unthinned)"))
 					.Font(FEditorStyle::GetFontStyle(TEXT("PropertyWindow.NormalFont")))
 			]
 			.ValueContent()
