Media: Added tooltips to media details (TTP# 348610)

diff --git a/Engine/Source/Editor/DetailCustomizations/Private/MediaPlayerCustomization.cpp b/Engine/Source/Editor/DetailCustomizations/Private/MediaPlayerCustomization.cpp
--- a/Engine/Source/Editor/DetailCustomizations/Private/MediaPlayerCustomization.cpp
+++ b/Engine/Source/Editor/DetailCustomizations/Private/MediaPlayerCustomization.cpp
@@ -58,11 +58,12 @@ void FMediaPlayerCustomization::CustomizeDetails( IDetailLayoutBuilder& DetailBu
 					SNew(SFilePathPicker)
 						.BrowseButtonImage(FEditorStyle::GetBrush("PropertyWindow.Button_Ellipsis"))
 						.BrowseButtonStyle(FEditorStyle::Get(), "HoverHintOnly")
-						.BrowseButtonToolTip(LOCTEXT("UrlBrowseButtonToolTipText", "Choose a file from this computer"))
+						.BrowseButtonToolTip(LOCTEXT("UrlBrowseButtonToolTip", "Choose a file from this computer"))
 						.BrowseDirectory(FPaths::GameContentDir() / TEXT("Movies"))
 						.FilePath(this, &FMediaPlayerCustomization::HandleUrlPickerFilePath)
 						.FileTypeFilter(this, &FMediaPlayerCustomization::HandleUrlPickerFileTypeFilter)
 						.OnPathPicked(this, &FMediaPlayerCustomization::HandleUrlPickerPathPicked)
+						.ToolTipText(LOCTEXT("UrlToolTip", "The path to a media file on this computer, or a URL to a media source on the internet"))
 				];
 		}
 	}
@@ -84,6 +85,7 @@ void FMediaPlayerCustomization::CustomizeDetails( IDetailLayoutBuilder& DetailBu
 				SNew(STextBlock)
 					.Text(this, &FMediaPlayerCustomization::HandleDurationTextBlockText)
 					.Font(FEditorStyle::GetFontStyle(TEXT("PropertyWindow.NormalFont")))
+					.ToolTipText(LOCTEXT("DurationToolTip", "The total duration of this media, i.e. how long it plays"))
 			];
 
 		// forward rates (thinned)
@@ -100,6 +102,7 @@ void FMediaPlayerCustomization::CustomizeDetails( IDetailLayoutBuilder& DetailBu
 				SNew(STextBlock)
 					.Text(this, &FMediaPlayerCustomization::HandleSupportedRatesTextBlockText, EMediaPlaybackDirections::Forward, false)
 					.Font(FEditorStyle::GetFontStyle(TEXT("PropertyWindow.NormalFont")))
+					.ToolTipText(LOCTEXT("ThinnedForwardRateToolTip", "The rate at which this media can be played thinned (with skipping frames) in the forward direction."))
 			];
 
 		// forward rates (unthinned)
@@ -116,6 +119,7 @@ void FMediaPlayerCustomization::CustomizeDetails( IDetailLayoutBuilder& DetailBu
 				SNew(STextBlock)
 					.Text(this, &FMediaPlayerCustomization::HandleSupportedRatesTextBlockText, EMediaPlaybackDirections::Forward, true)
 					.Font(FEditorStyle::GetFontStyle(TEXT("PropertyWindow.NormalFont")))
+					.ToolTipText(LOCTEXT("UnthinnedForwardRateToolTip", "The rate at which this media can be played unthinned (without skipping frames) in the forward direction."))
 			];
 
 		// reverse rates (thinned)
@@ -132,6 +136,7 @@ void FMediaPlayerCustomization::CustomizeDetails( IDetailLayoutBuilder& DetailBu
 				SNew(STextBlock)
 					.Text(this, &FMediaPlayerCustomization::HandleSupportedRatesTextBlockText, EMediaPlaybackDirections::Reverse, false)
 					.Font(FEditorStyle::GetFontStyle(TEXT("PropertyWindow.NormalFont")))
+					.ToolTipText(LOCTEXT("ThinnedReverseRateToolTip", "The rate at which this media can be played thinned (with skipping frames) in the reverse direction."))
 			];
 
 		// reverse rates (unthinned)
@@ -148,6 +153,7 @@ void FMediaPlayerCustomization::CustomizeDetails( IDetailLayoutBuilder& DetailBu
 				SNew(STextBlock)
 					.Text(this, &FMediaPlayerCustomization::HandleSupportedRatesTextBlockText, EMediaPlaybackDirections::Reverse, true)
 					.Font(FEditorStyle::GetFontStyle(TEXT("PropertyWindow.NormalFont")))
+					.ToolTipText(LOCTEXT("UnthinnedReverseRateToolTip", "The rate at which this media can be played unthinned (without skipping frames) in the reverse direction."))
 			];
 
 		// supports scrubbing
@@ -164,6 +170,7 @@ void FMediaPlayerCustomization::CustomizeDetails( IDetailLayoutBuilder& DetailBu
 				SNew(STextBlock)
 					.Text(this, &FMediaPlayerCustomization::HandleSupportsScrubbingTextBlockText)
 					.Font(FEditorStyle::GetFontStyle(TEXT("PropertyWindow.NormalFont")))
+					.ToolTipText(LOCTEXT("ScrubbingToolTip", "Whether this media supports scrubbing, i.e. manually controlling the playback speed and position."))
 			];
 
 		// supports seeking
@@ -180,6 +187,7 @@ void FMediaPlayerCustomization::CustomizeDetails( IDetailLayoutBuilder& DetailBu
 				SNew(STextBlock)
 					.Text(this, &FMediaPlayerCustomization::HandleSupportsSeekingTextBlockText)
 					.Font(FEditorStyle::GetFontStyle(TEXT("PropertyWindow.NormalFont")))
+					.ToolTipText(LOCTEXT("SeekingToolTip", "Whether this media supports seeking, i.e. jumping to a specific position."))
 			];
 	}
 }
