Fix crash on reset all in rig

diff --git a/Engine/Source/Editor/DetailCustomizations/Private/RigDetails.cpp b/Engine/Source/Editor/DetailCustomizations/Private/RigDetails.cpp
--- a/Engine/Source/Editor/DetailCustomizations/Private/RigDetails.cpp
+++ b/Engine/Source/Editor/DetailCustomizations/Private/RigDetails.cpp
@@ -287,6 +287,11 @@ void FRigDetails::GenerateTransformBaseArrayElementWidget(TSharedRef<IPropertyHa
 	URig * Rig = Cast<URig>(ItemBeingEdited.Get());
 	check (Rig);
 	const TArray<FNode> & Nodes = Rig->GetNodes();
+	if ( Nodes.Num() <= 0 )
+	{
+		return;
+	}
+
 	int32 NodeIndex = 0, ParentIndex_T=0, ParentIndex_R=0;
 	const FNode& CurNode = Nodes[ArrayIndex];
 	for(auto Node : Nodes)
@@ -427,7 +432,7 @@ void FRigDetails::OnParentSpaceSelectionChanged(TSharedPtr<FString> SelectedItem
 	{
 		if(SelectedItem.IsValid())
 		{
-			check(ParentSpacePropertyHandle->SetValueFromFormattedString(*SelectedItem.Get()) != FPropertyAccess::Fail);
+			ParentSpacePropertyHandle->SetValueFromFormattedString(*SelectedItem.Get());
 		}
 	}
 }
@@ -471,9 +476,12 @@ FString FRigDetails::GetSelectedTextLabel(TSharedRef<IPropertyHandle> ParentSpac
 {
 	FString DisplayText;
 
-	check (ParentSpacePropertyHandle->GetValueAsDisplayString(DisplayText) != FPropertyAccess::Fail);
+	if (ParentSpacePropertyHandle->GetValueAsDisplayString(DisplayText) != FPropertyAccess::Fail)
+	{
+		return DisplayText;
+	}
 
-	return DisplayText;
+	return TEXT("Unknown");
 }
 
 void FRigDetails::OnComboBoxOopening(TSharedRef<IPropertyHandle> ParentSpacePropertyHandle, int32 ArrayIndex, bool bTranslation)
@@ -501,15 +509,18 @@ void FRigDetails::OnComboBoxOopening(TSharedRef<IPropertyHandle> ParentSpaceProp
 void FRigDetails::OnAdvancedCheckBoxStateChanged(ESlateCheckBoxState::Type NewState, TSharedRef<IPropertyHandle> PropertyHandle)
 {
 	bool bValue = (NewState == ESlateCheckBoxState::Checked)? true : false;
-	check (PropertyHandle->SetValue(bValue) != FPropertyAccess::Fail);
+	PropertyHandle->SetValue(bValue);
 }
 
 ESlateCheckBoxState::Type FRigDetails::AdvancedCheckBoxIsChecked(TSharedRef<IPropertyHandle> PropertyHandle) const
 {
-	bool bValue;
+	bool bValue = false;
 	// multi value doesn't work in array, so i'm not handling multi value
-	check(PropertyHandle->GetValue(bValue) != FPropertyAccess::Fail);
+	if (PropertyHandle->GetValue(bValue) != FPropertyAccess::Fail)
+	{
+		return (bValue)? ESlateCheckBoxState::Checked: ESlateCheckBoxState::Unchecked;
+	}
 
-	return (bValue)? ESlateCheckBoxState::Checked: ESlateCheckBoxState::Unchecked;
+	return ESlateCheckBoxState::Undetermined;
 }
 #undef LOCTEXT_NAMESPACE
