Revised documentationactor link code.

diff --git a/Engine/Source/Editor/DetailCustomizations/Private/DocumentationActorDetails.cpp b/Engine/Source/Editor/DetailCustomizations/Private/DocumentationActorDetails.cpp
--- a/Engine/Source/Editor/DetailCustomizations/Private/DocumentationActorDetails.cpp
+++ b/Engine/Source/Editor/DetailCustomizations/Private/DocumentationActorDetails.cpp
@@ -86,7 +86,7 @@ FText FDocumentationActorDetails::OnGetButtonTooltipText() const
 	switch (LinkType)
 	{
 		case EDocumentationActorType::UDNLink:
-			Tooltip = FText::Format(LOCTEXT("OpenUDNLink", "Opens the UDN link:{0}"), FText::FromString(SelectedDocumentationActor->DocumentLink));
+			Tooltip = FText::Format(LOCTEXT("OpenUDNLink", "Opens the help link:{0}"), FText::FromString(SelectedDocumentationActor->DocumentLink));
 			break;
 		case EDocumentationActorType::URLLink:
 			Tooltip = FText::Format(LOCTEXT("OpenURLLink", "Opens the Web page:{0}"), FText::FromString(SelectedDocumentationActor->DocumentLink));
diff --git a/Engine/Source/Runtime/Engine/Private/DocumentationActor.cpp b/Engine/Source/Runtime/Engine/Private/DocumentationActor.cpp
--- a/Engine/Source/Runtime/Engine/Private/DocumentationActor.cpp
+++ b/Engine/Source/Runtime/Engine/Private/DocumentationActor.cpp
@@ -81,14 +81,13 @@ void ADocumentationActor::UpdateLinkType()
 	{
 		LinkType = EDocumentationActorType::URLLink;
 	}
-	else if (IDocumentation::Get()->PageExists(DocumentLink) == true )
+	else
 	{
 		LinkType = EDocumentationActorType::UDNLink;
 	}
-	else
+#else
+	LinkType = EDocumentationActorType::None;
 #endif
-	{
-		LinkType = EDocumentationActorType::InvalidLink;
-	}
+	
 }
 
