Fixed transform customizating swapping roll and yaw

diff --git a/Engine/Source/Editor/DetailCustomizations/Private/MathStructProxyCustomizations.cpp b/Engine/Source/Editor/DetailCustomizations/Private/MathStructProxyCustomizations.cpp
--- a/Engine/Source/Editor/DetailCustomizations/Private/MathStructProxyCustomizations.cpp
+++ b/Engine/Source/Editor/DetailCustomizations/Private/MathStructProxyCustomizations.cpp
@@ -159,7 +159,7 @@ void FMatrixStructCustomization::CustomizeChildren(TSharedRef<class IPropertyHan
 			+ SHorizontalBox::Slot()
 			.Padding(FMargin(0.0f, 2.0f, 3.0f, 2.0f))
 			[
-				MakeNumericProxyWidget<FRotator, float>(StructPropertyHandle, CachedRotationYaw, LOCTEXT("RotationYaw", "X"), FLinearColor::White, SNumericEntryBox<float>::RedLabelBackgroundColor)
+				MakeNumericProxyWidget<FRotator, float>(StructPropertyHandle, CachedRotationRoll, LOCTEXT("RotationRoll", "X"), FLinearColor::White, SNumericEntryBox<float>::RedLabelBackgroundColor)
 			]
 			+ SHorizontalBox::Slot()
 			.Padding(FMargin(0.0f, 2.0f, 3.0f, 2.0f))
@@ -169,7 +169,7 @@ void FMatrixStructCustomization::CustomizeChildren(TSharedRef<class IPropertyHan
 			+ SHorizontalBox::Slot()
 			.Padding(FMargin(0.0f, 2.0f, 0.0f, 2.0f))
 			[
-				MakeNumericProxyWidget<FRotator, float>(StructPropertyHandle, CachedRotationRoll, LOCTEXT("RotationRoll", "Z"), FLinearColor::White, SNumericEntryBox<float>::BlueLabelBackgroundColor)
+				MakeNumericProxyWidget<FRotator, float>(StructPropertyHandle, CachedRotationYaw, LOCTEXT("RotationYaw", "Z"), FLinearColor::White, SNumericEntryBox<float>::BlueLabelBackgroundColor)
 			]
 		];
 
