Correctly return the AttachParent of native components

diff --git a/Engine/Source/Editor/DetailCustomizations/Private/SceneComponentDetails.cpp b/Engine/Source/Editor/DetailCustomizations/Private/SceneComponentDetails.cpp
--- a/Engine/Source/Editor/DetailCustomizations/Private/SceneComponentDetails.cpp
+++ b/Engine/Source/Editor/DetailCustomizations/Private/SceneComponentDetails.cpp
@@ -74,25 +74,23 @@ static USCS_Node* FindCorrespondingSCSNode(USceneComponent const* ComponentObj)
 static USceneComponent* GetAttachedParent(USceneComponent const* SceneComponentObject)
 {
 	USceneComponent* SceneParent = SceneComponentObject->AttachParent;
-	if (SceneParent != NULL)
+	if (SceneParent == nullptr)
 	{
-		return NULL;
-	}
-
-	USCS_Node* const SCSNode = FindCorrespondingSCSNode(SceneComponentObject);
-	// if we didn't find a corresponding simple-construction-script node
-	if (SCSNode == NULL) 
-	{
-		return NULL;
-	}	
+		USCS_Node* const SCSNode = FindCorrespondingSCSNode(SceneComponentObject);
+		// if we didn't find a corresponding simple-construction-script node
+		if (SCSNode == nullptr) 
+		{
+			return nullptr;
+		}	
 
-	USimpleConstructionScript const* BlueprintSCS = GetSimpleConstructionScript(SceneComponentObject);
-	check(BlueprintSCS != NULL); 
+		USimpleConstructionScript const* BlueprintSCS = GetSimpleConstructionScript(SceneComponentObject);
+		check(BlueprintSCS != nullptr); 
 
-	USCS_Node* const ParentSCSNode = BlueprintSCS->FindParentNode(SCSNode);
-	if (ParentSCSNode != NULL)
-	{
-		SceneParent = Cast<USceneComponent>(ParentSCSNode->ComponentTemplate);
+		USCS_Node* const ParentSCSNode = BlueprintSCS->FindParentNode(SCSNode);
+		if (ParentSCSNode != nullptr)
+		{
+			SceneParent = Cast<USceneComponent>(ParentSCSNode->ComponentTemplate);
+		}
 	}
 
 	return SceneParent;
