Changed the Static Mesh LOD Group property widget in the static mesh import dialog to a dropdown box instead of a free text field, thus restricting the LOD Group which can be entered to defined LOD groups.
#ttp 349335 - EDITOR: FBX: CRASH: Changing the Static Mesh LODGroup name on import and open mesh will cause crash
#branch UE4
#proj Editor.DetailCustomizations

diff --git a/Engine/Source/Editor/DetailCustomizations/Private/FbxImportUIDetails.cpp b/Engine/Source/Editor/DetailCustomizations/Private/FbxImportUIDetails.cpp
--- a/Engine/Source/Editor/DetailCustomizations/Private/FbxImportUIDetails.cpp
+++ b/Engine/Source/Editor/DetailCustomizations/Private/FbxImportUIDetails.cpp
@@ -10,6 +10,13 @@
 FFbxImportUIDetails::FFbxImportUIDetails()
 {
 	CachedDetailBuilder = nullptr;
+
+	LODGroupNames.Reset();
+	UStaticMesh::GetLODGroups(LODGroupNames);
+	for (int32 GroupIndex = 0; GroupIndex < LODGroupNames.Num(); ++GroupIndex)
+	{
+		LODGroupOptions.Add(MakeShareable(new FString(LODGroupNames[GroupIndex].GetPlainNameString())));
+	}
 }
 
 TSharedRef<IDetailCustomization> FFbxImportUIDetails::MakeInstance()
@@ -101,7 +108,16 @@ void FFbxImportUIDetails::CustomizeDetails( IDetailLayoutBuilder& DetailBuilder
 			else
 			{
 				// No override, add to default mesh category
-				MeshCategory.AddProperty(Handle);
+				IDetailPropertyRow& PropertyRow = MeshCategory.AddProperty(Handle);
+
+				UProperty* Property = Handle->GetProperty();
+				if (Property != nullptr)
+				{
+					if (Property->GetFName() == GET_MEMBER_NAME_CHECKED(UFbxStaticMeshImportData, StaticMeshLODGroup))
+					{
+						SetStaticMeshLODGroupWidget(PropertyRow, Handle);
+					}
+				}
 			}
 		}
 	}
@@ -199,6 +215,52 @@ void FFbxImportUIDetails::CustomizeDetails( IDetailLayoutBuilder& DetailBuilder
 	}
 }
 
+void FFbxImportUIDetails::SetStaticMeshLODGroupWidget(IDetailPropertyRow& PropertyRow, const TSharedPtr<IPropertyHandle>& Handle)
+{
+	TSharedPtr<SWidget> NameWidget;
+	TSharedPtr<SWidget> ValueWidget;
+	FDetailWidgetRow Row;
+	PropertyRow.GetDefaultWidgets(NameWidget, ValueWidget, Row);
+
+	FName InitialValue;
+	ensure(Handle->GetValue(InitialValue) == FPropertyAccess::Success);
+	int32 GroupIndex = LODGroupNames.Find(InitialValue);
+	check(GroupIndex != INDEX_NONE);
+
+	StaticMeshLODGroupPropertyHandle = Handle;
+	TWeakPtr<IPropertyHandle> HandlePtr = Handle;
+
+	const bool bShowChildren = true;
+	PropertyRow.CustomWidget(bShowChildren)
+		.NameContent()
+		.MinDesiredWidth(Row.NameWidget.MinWidth)
+		.MaxDesiredWidth(Row.NameWidget.MaxWidth)
+		[
+			NameWidget.ToSharedRef()
+		]
+		.ValueContent()
+		.MinDesiredWidth(Row.ValueWidget.MinWidth)
+		.MaxDesiredWidth(Row.ValueWidget.MaxWidth)
+		[
+			SNew(STextComboBox)
+			.OptionsSource(&LODGroupOptions)
+			.InitiallySelectedItem(LODGroupOptions[GroupIndex])
+			.OnSelectionChanged(this, &FFbxImportUIDetails::OnLODGroupChanged, HandlePtr)
+		];
+}
+
+void FFbxImportUIDetails::OnLODGroupChanged(TSharedPtr<FString> NewValue, ESelectInfo::Type SelectInfo, TWeakPtr<IPropertyHandle> HandlePtr)
+{
+	TSharedPtr<IPropertyHandle> Handle = HandlePtr.Pin();
+	if (Handle.IsValid())
+	{
+		int32 GroupIndex = LODGroupOptions.Find(NewValue);
+		check(GroupIndex != INDEX_NONE);
+		ensure(Handle->SetValue(LODGroupNames[GroupIndex]) == FPropertyAccess::Success);
+	}
+}
+
+
 void FFbxImportUIDetails::CollectChildPropertiesRecursive(TSharedPtr<IPropertyHandle> Node, TArray<TSharedPtr<IPropertyHandle>>& OutProperties)
 {
 	uint32 NodeNumChildren = 0;
diff --git a/Engine/Source/Editor/DetailCustomizations/Private/FbxImportUIDetails.h b/Engine/Source/Editor/DetailCustomizations/Private/FbxImportUIDetails.h
--- a/Engine/Source/Editor/DetailCustomizations/Private/FbxImportUIDetails.h
+++ b/Engine/Source/Editor/DetailCustomizations/Private/FbxImportUIDetails.h
@@ -33,4 +33,17 @@ private:
 
 	/** Use MakeInstance to create an instance of this class */
 	FFbxImportUIDetails();
+
+	/** Sets a custom widget for the StaticMeshLODGroup property */
+	void SetStaticMeshLODGroupWidget(IDetailPropertyRow& PropertyRow, const TSharedPtr<IPropertyHandle>& Handle);
+
+	/** Called when the StaticMeshLODGroup spinbox is changed */
+	void OnLODGroupChanged(TSharedPtr<FString> NewValue, ESelectInfo::Type SelectInfo, TWeakPtr<IPropertyHandle> HandlePtr);
+
+	/** LOD group options. */
+	TArray<FName> LODGroupNames;
+	TArray<TSharedPtr<FString>> LODGroupOptions;
+
+	/** Cached StaticMeshLODGroup property handle */
+	TSharedPtr<IPropertyHandle> StaticMeshLODGroupPropertyHandle;
 };
