Media: Fixed media player customization's incorrect handling of relative paths to local movie files

diff --git a/Engine/Source/Editor/DetailCustomizations/Private/MediaPlayerCustomization.cpp b/Engine/Source/Editor/DetailCustomizations/Private/MediaPlayerCustomization.cpp
--- a/Engine/Source/Editor/DetailCustomizations/Private/MediaPlayerCustomization.cpp
+++ b/Engine/Source/Editor/DetailCustomizations/Private/MediaPlayerCustomization.cpp
@@ -384,27 +384,16 @@ FString FMediaPlayerCustomization::HandleUrlPickerFileTypeFilter( ) const
 
 void FMediaPlayerCustomization::HandleUrlPickerPathPicked( const FString& PickedPath )
 {
-	FString FullPath;
+	const FString FullGameContentDir = FPaths::ConvertRelativePathToFull(FPaths::GameContentDir());
+	FString FullUrl = FPaths::ConvertRelativePathToFull(FPaths::IsRelative(PickedPath) ? FPaths::GameContentDir() / PickedPath : PickedPath);
 
-	if (!PickedPath.IsEmpty() && FPaths::IsRelative(PickedPath))
+	if (FullUrl.StartsWith(FullGameContentDir))
 	{
-		FullPath = FPaths::ConvertRelativePathToFull(PickedPath);
-	}
-	else
-	{
-		FullPath = PickedPath;
-		FPaths::NormalizeFilename(FullPath);
-	}
-
-	FString FullGameContentDir = FPaths::ConvertRelativePathToFull(FPaths::GameContentDir());
-
-	if (FullPath.StartsWith(FullGameContentDir))
-	{
-		FPaths::MakePathRelativeTo(FullPath, *FullGameContentDir);
-		FullPath = FString(TEXT("./")) + FullPath;
+		FPaths::MakePathRelativeTo(FullUrl, *FullGameContentDir);
+		FullUrl = FString(TEXT("./")) + FullUrl;
 	}
 
-	UrlProperty->SetValue(FullPath);
+	UrlProperty->SetValue(FullUrl);
 }
 
 
