MediaPlayer: Fixed icon alignment in details customization

diff --git a/Engine/Source/Editor/DetailCustomizations/Private/MediaPlayerCustomization.cpp b/Engine/Source/Editor/DetailCustomizations/Private/MediaPlayerCustomization.cpp
--- a/Engine/Source/Editor/DetailCustomizations/Private/MediaPlayerCustomization.cpp
+++ b/Engine/Source/Editor/DetailCustomizations/Private/MediaPlayerCustomization.cpp
@@ -33,6 +33,7 @@ void FMediaPlayerCustomization::CustomizeDetails( IDetailLayoutBuilder& DetailBu
 
 					+ SHorizontalBox::Slot()
 						.AutoWidth()
+						.VAlign(VAlign_Center)
 						[
 							SNew(STextBlock)
 								.Font(IDetailLayoutBuilder::GetDetailFont())
@@ -43,6 +44,7 @@ void FMediaPlayerCustomization::CustomizeDetails( IDetailLayoutBuilder& DetailBu
 					+ SHorizontalBox::Slot()
 						.FillWidth(1.0f)
 						.HAlign(HAlign_Left)
+						.VAlign(VAlign_Center)
 						.Padding(4.0f, 0.0f, 0.0f, 0.0f)
 						[
 							SNew(SImage)
