Expand Action/Axis mappings when clicking + buttons to add a new one

Setting group expansion at the point an action is added didn't work as the property view is rebuilt before it can be passed on.
Instead replaced existing code to set the expansion state of renamed nodes so that it can be used for any expansion settings that need to be delayed until next tick.

TTP#348688 - NUXF-293 Input binds: Clicking + buttons does not expand the nodes

diff --git a/Engine/Source/Editor/DetailCustomizations/Private/InputSettingsDetails.cpp b/Engine/Source/Editor/DetailCustomizations/Private/InputSettingsDetails.cpp
--- a/Engine/Source/Editor/DetailCustomizations/Private/InputSettingsDetails.cpp
+++ b/Engine/Source/Editor/DetailCustomizations/Private/InputSettingsDetails.cpp
@@ -27,8 +27,8 @@ void FActionMappingsNodeBuilder::Tick( float DeltaTime )
 	if (GroupsRequireRebuild())
 	{
 		RebuildChildren();
-		HandleRenamedGroupExpansion();
 	}
+	HandleDelayedGroupExpansion();
 }
 
 void FActionMappingsNodeBuilder::GenerateHeaderRowContent( FDetailWidgetRow& NodeRow )
@@ -74,17 +74,18 @@ void FActionMappingsNodeBuilder::GenerateChildContent( IDetailChildrenBuilder& C
 	{
 		FMappingSet& MappingSet = GroupedMappings[Index];
 
-		TSharedRef<SWidget> AddButton = PropertyCustomizationHelpers::MakeAddButton( FSimpleDelegate::CreateSP( this, &FActionMappingsNodeBuilder::AddActionMappingToGroupButton_OnClick, MappingSet), 
-			LOCTEXT("AddActionMappingToGroupToolTip", "Adds Action Mapping to Group") );
-
-		TSharedRef<SWidget> RemoveButton = PropertyCustomizationHelpers::MakeDeleteButton( FSimpleDelegate::CreateSP( this, &FActionMappingsNodeBuilder::RemoveActionMappingGroupButton_OnClick, MappingSet), 
-			LOCTEXT("RemoveActionMappingGroupToolTip", "Removes Action Mapping Group") );
-
 		FString GroupNameString(TEXT("ActionMappings."));
 		MappingSet.SharedName.AppendString(GroupNameString);
 		FName GroupName(*GroupNameString);
 		IDetailGroup& ActionMappingGroup = ChildrenBuilder.AddChildGroup(GroupName, MappingSet.SharedName.ToString());
 		MappingSet.DetailGroup = &ActionMappingGroup;
+
+		TSharedRef<SWidget> AddButton = PropertyCustomizationHelpers::MakeAddButton(FSimpleDelegate::CreateSP(this, &FActionMappingsNodeBuilder::AddActionMappingToGroupButton_OnClick, MappingSet),
+			LOCTEXT("AddActionMappingToGroupToolTip", "Adds Action Mapping to Group"));
+
+		TSharedRef<SWidget> RemoveButton = PropertyCustomizationHelpers::MakeDeleteButton(FSimpleDelegate::CreateSP(this, &FActionMappingsNodeBuilder::RemoveActionMappingGroupButton_OnClick, MappingSet),
+			LOCTEXT("RemoveActionMappingGroupToolTip", "Removes Action Mapping Group"));
+
 		ActionMappingGroup.HeaderRow()
 		[
 			SNew( SHorizontalBox )
@@ -159,6 +160,7 @@ void FActionMappingsNodeBuilder::AddActionMappingButton_OnClick()
 		}
 		while (!bFoundUniqueName);
 
+		DelayedGroupExpansionStates.Add(TPairInitializer<FName, bool>(NewActionMappingName, true));
 		FInputActionKeyMapping NewMapping(NewActionMappingName);
 		InputSettings->ActionMappings.Add(NewMapping);
 
@@ -192,8 +194,7 @@ void FActionMappingsNodeBuilder::OnActionMappingNameCommitted(const FText& InNam
 
 		if (MappingSet.DetailGroup)
 		{
-			RenamedGroupExpansionState.Key = NewName;
-			RenamedGroupExpansionState.Value = MappingSet.DetailGroup->GetExpansionState();
+			DelayedGroupExpansionStates.Add(TPairInitializer<FName, bool>(NewName, MappingSet.DetailGroup->GetExpansionState()));
 
 			// Don't want to save expansion state of old name
 			MappingSet.DetailGroup->ToggleExpansion(false);
@@ -214,6 +215,7 @@ void FActionMappingsNodeBuilder::AddActionMappingToGroupButton_OnClick(const FMa
 		InputSettings->Modify();
 		ActionMappingsPropertyHandle->NotifyPreChange();
 
+		DelayedGroupExpansionStates.Add(TPairInitializer<FName, bool>(MappingSet.SharedName, true));
 		FInputActionKeyMapping NewMapping(MappingSet.SharedName);
 		InputSettings->ActionMappings.Add(NewMapping);
 
@@ -294,19 +296,22 @@ void FActionMappingsNodeBuilder::RebuildGroupedMappings()
 	}
 }
 
-void FActionMappingsNodeBuilder::HandleRenamedGroupExpansion()
+void FActionMappingsNodeBuilder::HandleDelayedGroupExpansion()
 {
-	if (RenamedGroupExpansionState.Key != NAME_None)
+	if (DelayedGroupExpansionStates.Num() > 0)
 	{
-		for (auto& MappingSet : GroupedMappings)
+		for (auto GroupState : DelayedGroupExpansionStates)
 		{
-			if (MappingSet.SharedName == RenamedGroupExpansionState.Key)
+			for (auto& MappingSet : GroupedMappings)
 			{
-				MappingSet.DetailGroup->ToggleExpansion(RenamedGroupExpansionState.Value);
-				break;
+				if (MappingSet.SharedName == GroupState.Key)
+				{
+					MappingSet.DetailGroup->ToggleExpansion(GroupState.Value);
+					break;
+				}
 			}
 		}
-		RenamedGroupExpansionState.Key = NAME_None;
+		DelayedGroupExpansionStates.Empty();
 	}
 }
 
@@ -329,8 +334,8 @@ void FAxisMappingsNodeBuilder::Tick( float DeltaTime )
 	if (GroupsRequireRebuild())
 	{
 		RebuildChildren();
-		HandleRenamedGroupExpansion();
 	}
+	HandleDelayedGroupExpansion();
 }
 
 void FAxisMappingsNodeBuilder::GenerateHeaderRowContent( FDetailWidgetRow& NodeRow )
@@ -376,17 +381,18 @@ void FAxisMappingsNodeBuilder::GenerateChildContent( IDetailChildrenBuilder& Chi
 	{
 		FMappingSet& MappingSet = GroupedMappings[Index];
 
-		TSharedRef<SWidget> AddButton = PropertyCustomizationHelpers::MakeAddButton( FSimpleDelegate::CreateSP( this, &FAxisMappingsNodeBuilder::AddAxisMappingToGroupButton_OnClick, MappingSet), 
-			LOCTEXT("AddAxisMappingToGroupToolTip", "Adds Axis Mapping to Group") );
-
-		TSharedRef<SWidget> RemoveButton = PropertyCustomizationHelpers::MakeDeleteButton( FSimpleDelegate::CreateSP( this, &FAxisMappingsNodeBuilder::RemoveAxisMappingGroupButton_OnClick, MappingSet), 
-			LOCTEXT("RemoveAxisMappingGroupToolTip", "Removes Axis Mapping Group") );
-
 		FString GroupNameString(TEXT("AxisMappings."));
 		MappingSet.SharedName.AppendString(GroupNameString);
 		FName GroupName(*GroupNameString);
 		IDetailGroup& AxisMappingGroup = ChildrenBuilder.AddChildGroup(GroupName, MappingSet.SharedName.ToString());
 		MappingSet.DetailGroup = &AxisMappingGroup;
+
+		TSharedRef<SWidget> AddButton = PropertyCustomizationHelpers::MakeAddButton(FSimpleDelegate::CreateSP(this, &FAxisMappingsNodeBuilder::AddAxisMappingToGroupButton_OnClick, MappingSet),
+			LOCTEXT("AddAxisMappingToGroupToolTip", "Adds Axis Mapping to Group"));
+
+		TSharedRef<SWidget> RemoveButton = PropertyCustomizationHelpers::MakeDeleteButton(FSimpleDelegate::CreateSP(this, &FAxisMappingsNodeBuilder::RemoveAxisMappingGroupButton_OnClick, MappingSet),
+			LOCTEXT("RemoveAxisMappingGroupToolTip", "Removes Axis Mapping Group"));
+
 		AxisMappingGroup.HeaderRow()
 		[
 			SNew( SHorizontalBox )
@@ -461,6 +467,7 @@ void FAxisMappingsNodeBuilder::AddAxisMappingButton_OnClick()
 		}
 		while (!bFoundUniqueName);
 
+		DelayedGroupExpansionStates.Add(TPairInitializer<FName, bool>(NewAxisMappingName, true));
 		FInputAxisKeyMapping NewMapping(NewAxisMappingName);
 		InputSettings->AxisMappings.Add(NewMapping);
 
@@ -494,8 +501,7 @@ void FAxisMappingsNodeBuilder::OnAxisMappingNameCommitted(const FText& InName, E
 
 		if (MappingSet.DetailGroup)
 		{
-			RenamedGroupExpansionState.Key = NewName;
-			RenamedGroupExpansionState.Value = MappingSet.DetailGroup->GetExpansionState();
+			DelayedGroupExpansionStates.Add(TPairInitializer<FName, bool>(NewName, MappingSet.DetailGroup->GetExpansionState()));
 
 			// Don't want to save expansion state of old name
 			MappingSet.DetailGroup->ToggleExpansion(false);
@@ -516,6 +522,7 @@ void FAxisMappingsNodeBuilder::AddAxisMappingToGroupButton_OnClick(const FMappin
 		InputSettings->Modify();
 		AxisMappingsPropertyHandle->NotifyPreChange();
 
+		DelayedGroupExpansionStates.Add(TPairInitializer<FName, bool>(MappingSet.SharedName, true));
 		FInputAxisKeyMapping NewMapping(MappingSet.SharedName);
 		InputSettings->AxisMappings.Add(NewMapping);
 
@@ -596,19 +603,22 @@ void FAxisMappingsNodeBuilder::RebuildGroupedMappings()
 	}
 }
 
-void FAxisMappingsNodeBuilder::HandleRenamedGroupExpansion()
+void FAxisMappingsNodeBuilder::HandleDelayedGroupExpansion()
 {
-	if (RenamedGroupExpansionState.Key != NAME_None)
+	if (DelayedGroupExpansionStates.Num() > 0)
 	{
-		for (auto& MappingSet : GroupedMappings)
+		for (auto GroupState : DelayedGroupExpansionStates)
 		{
-			if (MappingSet.SharedName == RenamedGroupExpansionState.Key)
+			for (auto& MappingSet : GroupedMappings)
 			{
-				MappingSet.DetailGroup->ToggleExpansion(RenamedGroupExpansionState.Value);
-				break;
+				if (MappingSet.SharedName == GroupState.Key)
+				{
+					MappingSet.DetailGroup->ToggleExpansion(GroupState.Value);
+					break;
+				}
 			}
 		}
-		RenamedGroupExpansionState.Key = NAME_None;
+		DelayedGroupExpansionStates.Empty();
 	}
 }
 
diff --git a/Engine/Source/Editor/DetailCustomizations/Private/InputSettingsDetails.h b/Engine/Source/Editor/DetailCustomizations/Private/InputSettingsDetails.h
--- a/Engine/Source/Editor/DetailCustomizations/Private/InputSettingsDetails.h
+++ b/Engine/Source/Editor/DetailCustomizations/Private/InputSettingsDetails.h
@@ -43,8 +43,8 @@ private:
 	{
 		OnRebuildChildren.ExecuteIfBound();
 	}
-	/** Makes sure that renamed groups have their expansion set to their previous state */
-	void HandleRenamedGroupExpansion();
+	/** Makes sure that groups have their expansion set after any rebuilding */
+	void HandleDelayedGroupExpansion();
 
 private:
 	/** Called to rebuild the children of the detail tree */
@@ -58,7 +58,7 @@ private:
 
 	TArray<FMappingSet> GroupedMappings;
 
-	TPair<FName,bool> RenamedGroupExpansionState;
+	TArray<TPair<FName, bool>> DelayedGroupExpansionStates;
 };
 
 class FAxisMappingsNodeBuilder : public IDetailCustomNodeBuilder, public TSharedFromThis<FAxisMappingsNodeBuilder>
@@ -88,8 +88,8 @@ private:
 	{
 		OnRebuildChildren.ExecuteIfBound();
 	}
-	/** Makes sure that renamed groups have their expansion set to their previous state */
-	void HandleRenamedGroupExpansion();
+	/** Makes sure that groups have their expansion set after any rebuilding */
+	void HandleDelayedGroupExpansion();
 
 private:
 	/** Called to rebuild the children of the detail tree */
@@ -103,7 +103,7 @@ private:
 
 	TArray<FMappingSet> GroupedMappings;
 
-	TPair<FName,bool> RenamedGroupExpansionState;
+	TArray<TPair<FName, bool>> DelayedGroupExpansionStates;
 };
 
 class FInputSettingsDetails : public IDetailCustomization
