Add description and documentation anchor for input actions and axis

diff --git a/Engine/Source/Editor/DetailCustomizations/Private/InputSettingsDetails.cpp b/Engine/Source/Editor/DetailCustomizations/Private/InputSettingsDetails.cpp
--- a/Engine/Source/Editor/DetailCustomizations/Private/InputSettingsDetails.cpp
+++ b/Engine/Source/Editor/DetailCustomizations/Private/InputSettingsDetails.cpp
@@ -3,6 +3,7 @@
 #include "DetailCustomizationsPrivatePCH.h"
 #include "InputSettingsDetails.h"
 #include "ScopedTransaction.h"
+#include "IDocumentation.h"
 
 #define LOCTEXT_NAMESPACE "InputSettingsDetails"
 
@@ -639,6 +640,24 @@ void FInputSettingsDetails::CustomizeDetails(class IDetailLayoutBuilder& DetailB
 
 	IDetailCategoryBuilder& MappingsDetailCategoryBuilder = DetailBuilder.EditCategory(BindingsCategory);
 
+	MappingsDetailCategoryBuilder.AddCustomRow(TEXT("Action Axis Mappings"))
+	[
+		SNew(SHorizontalBox)
+		+SHorizontalBox::Slot()
+		.FillWidth(1)
+		[
+			SNew(STextBlock)
+			.Font(IDetailLayoutBuilder::GetDetailFont())
+			.AutoWrapText(true)
+			.Text(LOCTEXT("Mappings_Description", "Action and Axis Mappings provide a mechanism to conveniently map keys and axes to input behaviors by inserting a layer of indirection between the input behavior and the keys that invoke it. Action Mappings are for key presses and releases, while Axis Mappings allow for inputs that have a continuous range."))
+		]
+		+SHorizontalBox::Slot()
+		.AutoWidth()
+		[
+			IDocumentation::Get()->CreateAnchor(FString("Gameplay/Input"))
+		]
+	];
+
 	// Custom Action Mappings
 	const TSharedPtr<IPropertyHandle> ActionMappingsPropertyHandle = DetailBuilder.GetProperty(ActionMappings, UInputSettings::StaticClass());
 	ActionMappingsPropertyHandle->MarkHiddenByCustomization();
