Adding a curve linear color struct that can be added as a property and edited in details panel.

Removed code that explicitly checks for UCurveLinearColor and added functions to curve owner interface so that this struct can display the gradient editor as well.
Tidied up some of the curve customisation that this was based on, only including what's needed, forward declarations in single place, removed unused comments.

UE-4434 - GitHub 529 : Curve linear color + Customization

Integrated from PR #529 by iniside

diff --git a/Engine/Source/Editor/DetailCustomizations/Private/CurveColorCustomization.cpp b/Engine/Source/Editor/DetailCustomizations/Private/CurveColorCustomization.cpp
new file mode 100644
--- /dev/null
+++ b/Engine/Source/Editor/DetailCustomizations/Private/CurveColorCustomization.cpp
@@ -0,0 +1,435 @@
+// Copyright 1998-2014 Epic Games, Inc. All Rights Reserved.
+
+#include "DetailCustomizationsPrivatePCH.h"
+#include "CurveColorCustomization.h"
+#include "Dialogs/DlgPickAssetPath.h"
+#include "PackageTools.h"
+#include "MiniCurveEditor.h"
+#include "AssetRegistryModule.h"
+#include "SCurveEditor.h"
+
+#define LOCTEXT_NAMESPACE "CurveColorCustomization"
+
+const FVector2D FCurveColorCustomization::DEFAULT_WINDOW_SIZE = FVector2D(800, 500);
+
+TSharedRef<IPropertyTypeCustomization> FCurveColorCustomization::MakeInstance()
+{
+	return MakeShareable( new FCurveColorCustomization );
+}
+
+FCurveColorCustomization::~FCurveColorCustomization()
+{
+	DestroyPopOutWindow();
+}
+
+FCurveColorCustomization::FCurveColorCustomization()
+	: RuntimeCurve(NULL)
+	, Owner(NULL)
+	, ViewMinInput(0.0f)
+	, ViewMaxInput(5.0f)
+{
+}
+
+void FCurveColorCustomization::CustomizeHeader(TSharedRef<IPropertyHandle> InStructPropertyHandle, FDetailWidgetRow& HeaderRow, IPropertyTypeCustomizationUtils& StructCustomizationUtils)
+{
+	this->StructPropertyHandle = InStructPropertyHandle;
+
+	TArray<UObject*> OuterObjects;
+	StructPropertyHandle->GetOuterObjects(OuterObjects);
+
+	TArray<void*> StructPtrs;
+	StructPropertyHandle->AccessRawData( StructPtrs );
+	check(StructPtrs.Num()!=0);
+
+	if (StructPtrs.Num() == 1)
+	{
+		RuntimeCurve = reinterpret_cast<FRuntimeCurveLinearColor*>(StructPtrs[0]);
+
+		if (OuterObjects.Num() == 1)
+		{
+			Owner = OuterObjects[0];
+		}
+
+		HeaderRow
+			.NameContent()
+			[
+				InStructPropertyHandle->CreatePropertyNameWidget( TEXT( "" ), TEXT( "" ), false )
+			]
+			.ValueContent()
+			.HAlign(HAlign_Fill)
+			.MinDesiredWidth(200)
+			[
+				SNew(SBorder)
+				.VAlign(VAlign_Fill)
+				.OnMouseDoubleClick(this, &FCurveColorCustomization::OnCurvePreviewDoubleClick)
+				[
+					SAssignNew(CurveWidget, SCurveEditor)
+					.ViewMinInput(this, &FCurveColorCustomization::GetViewMinInput)
+					.ViewMaxInput(this, &FCurveColorCustomization::GetViewMaxInput)
+					.TimelineLength(this, &FCurveColorCustomization::GetTimelineLength)
+					.OnSetInputViewRange(this, &FCurveColorCustomization::SetInputViewRange)
+					.HideUI(false)
+					.DesiredSize(FVector2D(300, 150))
+				]
+			];
+
+		check(CurveWidget.IsValid());
+		if (RuntimeCurve->ExternalCurve)
+		{
+			CurveWidget->SetCurveOwner(RuntimeCurve->ExternalCurve, false);
+		}
+		else
+		{
+			CurveWidget->SetCurveOwner(this);
+		}
+	}
+	else
+	{
+		HeaderRow
+			.NameContent()
+			[
+				InStructPropertyHandle->CreatePropertyNameWidget( TEXT( "" ), TEXT( "" ), false )
+			]
+			.ValueContent()
+			[
+				SNew(SBorder)
+				.VAlign(VAlign_Fill)
+				[
+					SNew(STextBlock)
+					.Text(LOCTEXT("MultipleCurves", "Multiple Curves - unable to modify"))
+				]
+			];
+	}
+}
+
+void FCurveColorCustomization::CustomizeChildren(TSharedRef<IPropertyHandle> InStructPropertyHandle, IDetailChildrenBuilder& StructBuilder, IPropertyTypeCustomizationUtils& StructCustomizationUtils)
+{
+	uint32 NumChildren = 0;
+	StructPropertyHandle->GetNumChildren(NumChildren);
+
+	for( uint32 ChildIndex = 0; ChildIndex < NumChildren; ++ChildIndex )
+	{
+		TSharedPtr<IPropertyHandle> Child = StructPropertyHandle->GetChildHandle( ChildIndex );
+
+		if( Child->GetProperty()->GetName() == TEXT("ExternalCurve") )
+		{
+			ExternalCurveHandle = Child;
+
+			FSimpleDelegate OnCurveChangedDelegate = FSimpleDelegate::CreateSP( this, &FCurveColorCustomization::OnExternalCurveChanged, InStructPropertyHandle );
+			Child->SetOnPropertyValueChanged(OnCurveChangedDelegate);
+
+			StructBuilder.AddChildContent(TEXT("ExternalCurve"))
+				.NameContent()
+				[
+					Child->CreatePropertyNameWidget()
+				]
+				.ValueContent()
+				[
+					SNew(SVerticalBox)
+					+SVerticalBox::Slot()
+					.AutoHeight()
+					[
+						SNew(SHorizontalBox)
+						+SHorizontalBox::Slot()
+						[
+							Child->CreatePropertyValueWidget()
+						]
+						+SHorizontalBox::Slot()
+						.AutoWidth()
+						.VAlign(VAlign_Center)
+						.Padding(1,0)
+						[
+							SNew(SButton)
+							.ButtonStyle( FEditorStyle::Get(), "NoBorder" )
+							.ContentPadding(1.f)
+							.ToolTipText(LOCTEXT("ConvertInternalCurveTooltip", "Convert to Internal Color Curve"))
+							.OnClicked(this, &FCurveColorCustomization::OnConvertButtonClicked)
+							.IsEnabled(this, &FCurveColorCustomization::IsConvertButtonEnabled)
+							[
+								SNew(SImage)
+								.Image( FEditorStyle::GetBrush(TEXT("PropertyWindow.Button_Clear")) )
+							]
+						]
+					]
+					+SVerticalBox::Slot()
+					.AutoHeight()
+					[
+						SNew(SButton)
+						.HAlign(HAlign_Center)
+						.Text( LOCTEXT( "CreateAssetButton", "Create External Curve" ) )
+						.ToolTipText(LOCTEXT( "CreateAssetTooltip", "Create a new Color Curve asset from this curve") )
+						.OnClicked(this, &FCurveColorCustomization::OnCreateButtonClicked)
+						.IsEnabled(this, &FCurveColorCustomization::IsCreateButtonEnabled)
+					]
+				];
+		}
+		else
+		{
+			StructBuilder.AddChildProperty(Child.ToSharedRef());
+		}
+	}
+}
+
+TArray<FRichCurveEditInfoConst> FCurveColorCustomization::GetCurves() const
+{
+	TArray<FRichCurveEditInfoConst> Curves;
+	for (int Index = 0; Index < 4; Index++)
+	{
+		Curves.Add(FRichCurveEditInfoConst(&RuntimeCurve->ColorCurves[Index]));
+	}
+	return Curves;
+}
+
+TArray<FRichCurveEditInfo> FCurveColorCustomization::GetCurves()
+{
+	TArray<FRichCurveEditInfo> Curves;
+	for (int Index = 0; Index < 4; Index++)
+	{
+		Curves.Add(FRichCurveEditInfo(&RuntimeCurve->ColorCurves[Index]));
+	}
+	return Curves;
+}
+
+UObject* FCurveColorCustomization::GetOwner()
+{
+	return Owner;
+}
+
+void FCurveColorCustomization::ModifyOwner()
+{
+	if (Owner)
+	{
+		Owner->Modify(true);
+	}
+}
+
+void FCurveColorCustomization::MakeTransactional()
+{
+	if (Owner)
+	{
+		Owner->SetFlags(Owner->GetFlags() | RF_Transactional);
+	}
+}
+
+void FCurveColorCustomization::OnCurveChanged()
+{
+	StructPropertyHandle->NotifyPostChange();
+}
+
+FLinearColor FCurveColorCustomization::GetLinearColorValue(float InTime) const
+{
+	if (RuntimeCurve)
+	{
+		return RuntimeCurve->GetLinearColorValue(InTime);
+	}
+	return FLinearColor::Black; 
+}
+
+bool FCurveColorCustomization::HasAnyAlphaKeys() const
+{
+	if (RuntimeCurve)
+	{
+		return RuntimeCurve->ColorCurves[3].GetNumKeys() > 0;
+	}
+	return false;
+}
+
+float FCurveColorCustomization::GetTimelineLength() const
+{
+	return 0.f;
+}
+
+void FCurveColorCustomization::SetInputViewRange(float InViewMinInput, float InViewMaxInput)
+{
+	ViewMaxInput = InViewMaxInput;
+	ViewMinInput = InViewMinInput;
+}
+
+void FCurveColorCustomization::OnExternalCurveChanged(TSharedRef<IPropertyHandle> CurvePropertyHandle)
+{
+	if (RuntimeCurve)
+	{
+		if (RuntimeCurve->ExternalCurve)
+		{
+			CurveWidget->SetCurveOwner(RuntimeCurve->ExternalCurve, false);
+		}
+		else
+		{
+			CurveWidget->SetCurveOwner(this);
+		}
+
+		CurvePropertyHandle->NotifyPostChange();
+	}
+}
+
+FReply FCurveColorCustomization::OnCreateButtonClicked()
+{
+	if (CurveWidget.IsValid())
+	{
+		FString DefaultAsset = FPackageName::GetLongPackagePath(Owner->GetOutermost()->GetName()) + TEXT("/") + Owner->GetName() + TEXT("_ExternalCurve");
+
+		TSharedRef<SDlgPickAssetPath> NewCurveDlg =
+			SNew(SDlgPickAssetPath)
+			.Title(LOCTEXT("NewCurveDialogTitle", "Choose Location for External Curve Asset"))
+			.DefaultAssetPath(FText::FromString(DefaultAsset));
+
+		if (NewCurveDlg->ShowModal() != EAppReturnType::Cancel)
+		{
+			FString Package(NewCurveDlg->GetFullAssetPath().ToString());
+			FString Name(NewCurveDlg->GetAssetName().ToString());
+			FString Group(TEXT(""));
+
+			// Find (or create!) the desired package for this object
+			UPackage* Pkg = CreatePackage(NULL, *Package);
+			UPackage* OutermostPkg = Pkg->GetOutermost();
+
+			TArray<UPackage*> TopLevelPackages;
+			TopLevelPackages.Add( OutermostPkg );
+			if (!PackageTools::HandleFullyLoadingPackages(TopLevelPackages, LOCTEXT("CreateANewObject", "Create a new object")))
+			{
+				// User aborted.
+				return FReply::Handled();
+			}
+
+			if (!PromptUserIfExistingObject(Name, Package, Group, Pkg))
+			{
+				return FReply::Handled();
+			}
+
+			// PromptUserIfExistingObject may have GCed and recreated our outermost package - re-acquire it here.
+			OutermostPkg = Pkg->GetOutermost();
+
+			// Create a new asset and set it as the external curve
+			FName AssetName = *Name;
+			UCurveLinearColor* NewCurve = Cast<UCurveLinearColor>(CurveWidget->CreateCurveObject(UCurveLinearColor::StaticClass(), Pkg, AssetName));
+			if( NewCurve )
+			{
+				// run through points of editor data and add to external curve
+				for (int32 Index = 0; Index < 4; Index++)
+				{
+					CopyCurveData(&RuntimeCurve->ColorCurves[Index], &NewCurve->FloatCurves[Index]);
+				}
+
+				// Set the new object as the sole selection.
+				USelection* SelectionSet = GEditor->GetSelectedObjects();
+				SelectionSet->DeselectAll();
+				SelectionSet->Select( NewCurve );
+
+				// Notify the asset registry
+				FAssetRegistryModule::AssetCreated(NewCurve);
+
+				// Mark the package dirty...
+				OutermostPkg->MarkPackageDirty();
+
+				// Make sure expected type of pointer passed to SetValue, so that it's not interpreted as a bool
+				const UObject* NewObject = NewCurve;
+				ExternalCurveHandle->SetValue(NewObject);
+			}
+		}
+	}
+	return FReply::Handled();
+}
+
+bool FCurveColorCustomization::IsCreateButtonEnabled() const
+{
+	return CurveWidget.IsValid() && RuntimeCurve != NULL && RuntimeCurve->ExternalCurve == NULL;
+}
+
+FReply FCurveColorCustomization::OnConvertButtonClicked()
+{
+	if (RuntimeCurve && RuntimeCurve->ExternalCurve)
+	{
+		// clear points of editor data
+		for (int32 Index = 0; Index < 4; Index++)
+		{
+			RuntimeCurve->ColorCurves[Index].Reset();
+		}
+
+		// run through points of external curve and add to editor data
+		for (int32 Index = 0; Index < 4; Index++)
+		{
+			CopyCurveData(&RuntimeCurve->ExternalCurve->FloatCurves[Index], &RuntimeCurve->ColorCurves[Index]);
+		}
+
+		// null out external curve
+		const UObject* NullObject = NULL;
+		ExternalCurveHandle->SetValue(NullObject);
+	}
+	return FReply::Handled();
+}
+
+bool FCurveColorCustomization::IsConvertButtonEnabled() const
+{
+	return RuntimeCurve != NULL && RuntimeCurve->ExternalCurve != NULL;
+}
+
+FReply FCurveColorCustomization::OnCurvePreviewDoubleClick(const FGeometry& InMyGeometry, const FPointerEvent& InMouseEvent)
+{
+	if (InMouseEvent.GetEffectingButton() == EKeys::LeftMouseButton)
+	{
+		if (RuntimeCurve->ExternalCurve)
+		{
+			FAssetEditorManager::Get().OpenEditorForAsset(RuntimeCurve->ExternalCurve);
+		}
+		else
+		{
+			DestroyPopOutWindow();
+
+			// Determine the position of the window so that it will spawn near the mouse, but not go off the screen.
+			const FVector2D CursorPos = FSlateApplication::Get().GetCursorPos();
+			FSlateRect Anchor(CursorPos.X, CursorPos.Y, CursorPos.X, CursorPos.Y);
+
+			FVector2D AdjustedSummonLocation = FSlateApplication::Get().CalculatePopupWindowPosition( Anchor, FCurveColorCustomization::DEFAULT_WINDOW_SIZE, Orient_Horizontal );
+
+			TSharedPtr<SWindow> Window = SNew(SWindow)
+				.Title( FText::Format( LOCTEXT("WindowHeader", "{0} - Internal Color Curve Editor"), FText::FromString(StructPropertyHandle->GetPropertyDisplayName())) )
+				.ClientSize( FCurveColorCustomization::DEFAULT_WINDOW_SIZE )
+				.ScreenPosition(AdjustedSummonLocation)
+				.AutoCenter(EAutoCenter::None)
+				.SupportsMaximize(false)
+				.SupportsMinimize(false)
+				.SizingRule( ESizingRule::FixedSize );
+
+			// init the mini curve editor widget
+			TSharedRef<SMiniCurveEditor> MiniCurveEditor =
+				SNew(SMiniCurveEditor)
+				.CurveOwner(this)
+				.ParentWindow(Window);
+
+			Window->SetContent( MiniCurveEditor );
+
+			// Find the window of the parent widget
+			FWidgetPath WidgetPath;
+			FSlateApplication::Get().GeneratePathToWidgetChecked( CurveWidget.ToSharedRef(), WidgetPath );
+			Window = FSlateApplication::Get().AddWindowAsNativeChild( Window.ToSharedRef(), WidgetPath.GetWindow() );
+
+			//hold on to the window created for external use...
+			CurveEditorWindow = Window;
+		}
+	}
+	return FReply::Handled();
+}
+
+void FCurveColorCustomization::CopyCurveData( const FRichCurve* SrcCurve, FRichCurve* DestCurve )
+{
+	if( SrcCurve && DestCurve )
+	{
+		for (auto It(SrcCurve->GetKeyIterator()); It; ++It)
+		{
+			const FRichCurveKey& Key = *It;
+			FKeyHandle KeyHandle = DestCurve->AddKey(Key.Time, Key.Value);
+			DestCurve->GetKey(KeyHandle) = Key;
+		}
+	}
+}
+
+void FCurveColorCustomization::DestroyPopOutWindow()
+{
+	if (CurveEditorWindow.IsValid())
+	{
+		CurveEditorWindow.Pin()->RequestDestroyWindow();
+		CurveEditorWindow.Reset();
+	}
+}
+
+#undef LOCTEXT_NAMESPACE
\ No newline at end of file
diff --git a/Engine/Source/Editor/DetailCustomizations/Private/CurveColorCustomization.h b/Engine/Source/Editor/DetailCustomizations/Private/CurveColorCustomization.h
new file mode 100644
--- /dev/null
+++ b/Engine/Source/Editor/DetailCustomizations/Private/CurveColorCustomization.h
@@ -0,0 +1,134 @@
+// Copyright 1998-2014 Epic Games, Inc. All Rights Reserved.
+
+#pragma once
+
+class UObject;
+class SCurveEditor;
+class IPropertyHandle;
+class FDetailWidgetRow;
+class IDetailChildrenBuilder;
+
+struct FRuntimeCurveLinearColor;
+
+/**
+ * Customizes a RuntimeCurveLinearColor struct to display a Curve Editor
+ */
+class FCurveColorCustomization : public IPropertyTypeCustomization, public FCurveOwnerInterface
+{
+public:
+	static TSharedRef<IPropertyTypeCustomization> MakeInstance();
+
+	/**
+	 * Destructor
+	 */
+	virtual ~FCurveColorCustomization();
+
+	/** IPropertyTypeCustomization interface */
+	virtual void CustomizeHeader( TSharedRef<IPropertyHandle> InStructPropertyHandle, FDetailWidgetRow& HeaderRow, IPropertyTypeCustomizationUtils& StructCustomizationUtils ) override;
+	virtual void CustomizeChildren( TSharedRef<IPropertyHandle> InStructPropertyHandle, IDetailChildrenBuilder& StructBuilder, IPropertyTypeCustomizationUtils& StructCustomizationUtils ) override;
+
+	/** FCurveOwnerInterface interface */
+	virtual TArray<FRichCurveEditInfoConst> GetCurves() const override;
+	virtual TArray<FRichCurveEditInfo> GetCurves() override;
+	virtual UObject* GetOwner() override;
+	virtual void ModifyOwner() override;
+	virtual void MakeTransactional() override;
+	virtual void OnCurveChanged() override;
+	virtual bool IsLinearColorCurve() const override { return true; }
+	virtual FLinearColor GetLinearColorValue(float InTime) const override;
+	virtual bool HasAnyAlphaKeys() const override;
+
+private:
+	/**
+	 * Constructor
+	 */
+	FCurveColorCustomization();
+
+	/**
+	 * Get View Min Input for the Curve Editor
+	 */
+	float GetViewMinInput() const { return ViewMinInput; }
+
+	/**
+	 * Get View Min Input for the Curve Editor
+	 */
+	float GetViewMaxInput() const { return ViewMaxInput; }
+
+	/**
+	 * Get Timeline Length for the Curve Editor
+	 */
+	float GetTimelineLength() const;
+
+	/**
+	 * Set View Min and Max Inputs for the Curve Editor
+	 */
+	void SetInputViewRange(float InViewMinInput, float InViewMaxInput);
+
+	/**
+	 * Called when RuntimeCurveLinearColor's External Curve is changed
+	 */
+	void OnExternalCurveChanged(TSharedRef<IPropertyHandle> CurvePropertyHandle);
+
+	/**
+	 * Called when button clicked to create an External Colour Curve
+	 */
+	FReply OnCreateButtonClicked();
+
+	/**
+	 * Whether Create button is enabled
+	 */
+	bool IsCreateButtonEnabled() const;
+
+	/**
+	 * Called when button clicked to convert External Curve to internal curve
+	 */
+	FReply OnConvertButtonClicked();
+
+	/**
+	 * Whether Convert button is enabled
+	 */
+	bool IsConvertButtonEnabled() const;
+
+	/**
+	 * Called when user double clicks on the curve preview to open a full size editor
+	 */
+	FReply OnCurvePreviewDoubleClick( const FGeometry& InMyGeometry, const FPointerEvent& InMouseEvent );
+
+	/**
+	 * Copies data from one Rich Curve to another
+	 */
+	static void CopyCurveData( const FRichCurve* SrcCurve, FRichCurve* DestCurve );
+
+	/**
+	 * Destroys the pop out window used for editing internal curves.
+	 */
+	void DestroyPopOutWindow();
+
+private:
+	/** Cached RuntimeCurveLinearColor struct handle */
+	TSharedPtr<IPropertyHandle> StructPropertyHandle;
+
+	/** Cached External Curve handle */
+	TSharedPtr<IPropertyHandle> ExternalCurveHandle;
+
+	/** Small preview Curve Editor */
+	TSharedPtr<SCurveEditor> CurveWidget;
+
+	/** Window for pop out Curve Editor */
+	TWeakPtr<SWindow> CurveEditorWindow;
+
+	/** Cached pointer to the actual RuntimeCurveLinearColor struct */
+	FRuntimeCurveLinearColor* RuntimeCurve;
+
+	/** Object that owns the RuntimeFloatCurve */
+	UObject* Owner;
+
+	/** View Min Input for the Curve Editor */
+	float ViewMinInput;
+
+	/** View Max Input for the Curve Editor */
+	float ViewMaxInput;
+
+	/** Size of the pop out Curve Editor window */
+	static const FVector2D DEFAULT_WINDOW_SIZE;
+};
\ No newline at end of file
diff --git a/Engine/Source/Editor/DetailCustomizations/Private/CurveStructCustomization.cpp b/Engine/Source/Editor/DetailCustomizations/Private/CurveStructCustomization.cpp
--- a/Engine/Source/Editor/DetailCustomizations/Private/CurveStructCustomization.cpp
+++ b/Engine/Source/Editor/DetailCustomizations/Private/CurveStructCustomization.cpp
@@ -6,6 +6,7 @@
 #include "PackageTools.h"
 #include "MiniCurveEditor.h"
 #include "AssetRegistryModule.h"
+#include "SCurveEditor.h"
 
 #define LOCTEXT_NAMESPACE "CurveStructCustomization"
 
@@ -29,7 +30,7 @@ FCurveStructCustomization::FCurveStructCustomization()
 {
 }
 
-void FCurveStructCustomization::CustomizeHeader( TSharedRef<class IPropertyHandle> InStructPropertyHandle, class FDetailWidgetRow& HeaderRow, IPropertyTypeCustomizationUtils& StructCustomizationUtils )
+void FCurveStructCustomization::CustomizeHeader( TSharedRef<IPropertyHandle> InStructPropertyHandle, FDetailWidgetRow& HeaderRow, IPropertyTypeCustomizationUtils& StructCustomizationUtils )
 {
 	this->StructPropertyHandle = InStructPropertyHandle;
 
@@ -101,7 +102,7 @@ void FCurveStructCustomization::CustomizeHeader( TSharedRef<class IPropertyHandl
 	}
 }
 
-void FCurveStructCustomization::CustomizeChildren( TSharedRef<class IPropertyHandle> InStructPropertyHandle, class IDetailChildrenBuilder& StructBuilder, IPropertyTypeCustomizationUtils& StructCustomizationUtils )
+void FCurveStructCustomization::CustomizeChildren( TSharedRef<IPropertyHandle> InStructPropertyHandle, IDetailChildrenBuilder& StructBuilder, IPropertyTypeCustomizationUtils& StructCustomizationUtils )
 {
 	uint32 NumChildren = 0;
 	StructPropertyHandle->GetNumChildren(NumChildren);
@@ -141,8 +142,7 @@ void FCurveStructCustomization::CustomizeChildren( TSharedRef<class IPropertyHan
 							SNew(SButton)
 							.ButtonStyle( FEditorStyle::Get(), "NoBorder" )
 							.ContentPadding(1.f)
-							.ToolTipText(LOCTEXT( "ConvertInternalCurveTooltip", "Convert to Internal Curve"))
-							//.ToolTipText(LOCTEXT( "ConvertInternalCurveTooltip", "Copy the external CurveFloat asset to this curve") )
+							.ToolTipText(LOCTEXT("ConvertInternalCurveTooltip", "Convert to Internal Curve"))
 							.OnClicked(this, &FCurveStructCustomization::OnConvertButtonClicked)
 							.IsEnabled(this, &FCurveStructCustomization::IsConvertButtonEnabled)
 							[
@@ -154,25 +154,12 @@ void FCurveStructCustomization::CustomizeChildren( TSharedRef<class IPropertyHan
 					+SVerticalBox::Slot()
 					.AutoHeight()
 					[
-						SNew(SHorizontalBox)
-						+SHorizontalBox::Slot()
-						[
-							SNew(SButton)
-							.HAlign(HAlign_Center)
-							.Text( LOCTEXT( "CreateAssetButton", "Create External Curve" ) )
-							.ToolTipText(LOCTEXT( "CreateAssetTooltip", "Create a new CurveFloat asset from this curve") )
-							.OnClicked(this, &FCurveStructCustomization::OnCreateButtonClicked)
-							.IsEnabled(this, &FCurveStructCustomization::IsCreateButtonEnabled)
-						]
-						/*+SHorizontalBox::Slot()
-						[
-							SNew(SButton)
-							.HAlign(HAlign_Center)
-							.Text( LOCTEXT( "ConvertAssetButton", "Convert to Internal Curve" ) )
-							.ToolTipText(LOCTEXT( "CopyExternalCurveTooltip", "Copy the external CurveFloat asset to this curve") )
-							.OnClicked(this, &FCurveStructCustomization::OnConvertButtonClicked)
-							.IsEnabled(this, &FCurveStructCustomization::IsConvertButtonEnabled)
-						]*/
+						SNew(SButton)
+						.HAlign(HAlign_Center)
+						.Text( LOCTEXT( "CreateAssetButton", "Create External Curve" ) )
+						.ToolTipText(LOCTEXT( "CreateAssetTooltip", "Create a new CurveFloat asset from this curve") )
+						.OnClicked(this, &FCurveStructCustomization::OnCreateButtonClicked)
+						.IsEnabled(this, &FCurveStructCustomization::IsCreateButtonEnabled)
 					]
 				];
 		}
@@ -234,7 +221,7 @@ void FCurveStructCustomization::SetInputViewRange(float InViewMinInput, float In
 	ViewMinInput = InViewMinInput;
 }
 
-void FCurveStructCustomization::OnExternalCurveChanged(TSharedRef<class IPropertyHandle> CurvePropertyHandle)
+void FCurveStructCustomization::OnExternalCurveChanged(TSharedRef<IPropertyHandle> CurvePropertyHandle)
 {
 	if (RuntimeCurve)
 	{
@@ -257,7 +244,7 @@ FReply FCurveStructCustomization::OnCreateButtonClicked()
 	{
 		FString DefaultAsset = FPackageName::GetLongPackagePath(Owner->GetOutermost()->GetName()) + TEXT("/") + Owner->GetName() + TEXT("_ExternalCurve");
 
-		TSharedRef<SDlgPickAssetPath> NewCurveDlg = 
+		TSharedRef<SDlgPickAssetPath> NewCurveDlg =
 			SNew(SDlgPickAssetPath)
 			.Title(LOCTEXT("NewCurveDialogTitle", "Choose Location for External Curve Asset"))
 			.DefaultAssetPath(FText::FromString(DefaultAsset));
@@ -284,7 +271,7 @@ FReply FCurveStructCustomization::OnCreateButtonClicked()
 			{
 				return FReply::Handled();
 			}
-			
+
 			// PromptUserIfExistingObject may have GCed and recreated our outermost package - re-acquire it here.
 			OutermostPkg = Pkg->GetOutermost();
 
@@ -371,7 +358,7 @@ FReply FCurveStructCustomization::OnCurvePreviewDoubleClick( const FGeometry& In
 				.SizingRule( ESizingRule::FixedSize );
 
 			// init the mini curve editor widget
-			TSharedRef<SMiniCurveEditor> MiniCurveEditor = 
+			TSharedRef<SMiniCurveEditor> MiniCurveEditor =
 				SNew(SMiniCurveEditor)
 				.CurveOwner(this)
 				.ParentWindow(Window);
diff --git a/Engine/Source/Editor/DetailCustomizations/Private/CurveStructCustomization.h b/Engine/Source/Editor/DetailCustomizations/Private/CurveStructCustomization.h
--- a/Engine/Source/Editor/DetailCustomizations/Private/CurveStructCustomization.h
+++ b/Engine/Source/Editor/DetailCustomizations/Private/CurveStructCustomization.h
@@ -2,7 +2,13 @@
 
 #pragma once
 
-#include "SCurveEditor.h"
+class UObject;
+class SCurveEditor;
+class IPropertyHandle;
+class FDetailWidgetRow;
+class IDetailChildrenBuilder;
+
+struct FRuntimeFloatCurve;
 
 /**
  * Customizes a RuntimeFloatCurve struct to display a Curve Editor
@@ -18,8 +24,8 @@ public:
 	virtual ~FCurveStructCustomization();
 
 	/** IPropertyTypeCustomization interface */
-	virtual void CustomizeHeader( TSharedRef<class IPropertyHandle> InStructPropertyHandle, class FDetailWidgetRow& HeaderRow, IPropertyTypeCustomizationUtils& StructCustomizationUtils ) override;
-	virtual void CustomizeChildren( TSharedRef<class IPropertyHandle> InStructPropertyHandle, class IDetailChildrenBuilder& StructBuilder, IPropertyTypeCustomizationUtils& StructCustomizationUtils ) override;
+	virtual void CustomizeHeader( TSharedRef<IPropertyHandle> InStructPropertyHandle, FDetailWidgetRow& HeaderRow, IPropertyTypeCustomizationUtils& StructCustomizationUtils ) override;
+	virtual void CustomizeChildren( TSharedRef<IPropertyHandle> InStructPropertyHandle, IDetailChildrenBuilder& StructBuilder, IPropertyTypeCustomizationUtils& StructCustomizationUtils ) override;
 
 	/** FCurveOwnerInterface interface */
 	virtual TArray<FRichCurveEditInfoConst> GetCurves() const override;
@@ -34,7 +40,7 @@ private:
 	 * Constructor
 	 */
 	FCurveStructCustomization();
-	
+
 	/**
 	 * Get View Min Input for the Curve Editor
 	 */
@@ -58,7 +64,7 @@ private:
 	/**
 	 * Called when RuntimeFloatCurve's External Curve is changed
 	 */
-	void OnExternalCurveChanged(TSharedRef<class IPropertyHandle> CurvePropertyHandle);
+	void OnExternalCurveChanged(TSharedRef<IPropertyHandle> CurvePropertyHandle);
 
 	/**
 	 * Called when button clicked to create an External Curve
@@ -103,16 +109,16 @@ private:
 	TSharedPtr<IPropertyHandle> ExternalCurveHandle;
 
 	/** Small preview Curve Editor */
-	TSharedPtr<class SCurveEditor> CurveWidget;
+	TSharedPtr<SCurveEditor> CurveWidget;
 
 	/** Window for pop out Curve Editor */
 	TWeakPtr<SWindow> CurveEditorWindow;
 
 	/** Cached pointer to the actual RuntimeFloatCurve struct */
-	struct FRuntimeFloatCurve* RuntimeCurve;
+	FRuntimeFloatCurve* RuntimeCurve;
 
 	/** Object that owns the RuntimeFloatCurve */
-	class UObject* Owner;
+	UObject* Owner;
 
 	/** View Min Input for the Curve Editor */
 	float ViewMinInput;
diff --git a/Engine/Source/Editor/DetailCustomizations/Private/DetailCustomizations.cpp b/Engine/Source/Editor/DetailCustomizations/Private/DetailCustomizations.cpp
--- a/Engine/Source/Editor/DetailCustomizations/Private/DetailCustomizations.cpp
+++ b/Engine/Source/Editor/DetailCustomizations/Private/DetailCustomizations.cpp
@@ -87,6 +87,7 @@
 #include "TimespanStructCustomization.h"
 #include "FbxImportUIDetails.h"
 #include "RigDetails.h"
+#include "CurveColorCustomization.h"
 
 IMPLEMENT_MODULE( FDetailCustomizationsModule, DetailCustomizations );
 
@@ -177,6 +178,8 @@ void FDetailCustomizationsModule::RegisterPropertyTypeCustomizations()
 	RegisterCustomPropertyTypeLayout("InputAxisConfigEntry", FOnGetPropertyTypeCustomizationInstance::CreateStatic(&FInputAxisConfigCustomization::MakeInstance));
 	RegisterCustomPropertyTypeLayout("InputActionKeyMapping", FOnGetPropertyTypeCustomizationInstance::CreateStatic(&FInputActionMappingCustomization::MakeInstance));
 	RegisterCustomPropertyTypeLayout("InputAxisKeyMapping", FOnGetPropertyTypeCustomizationInstance::CreateStatic(&FInputAxisMappingCustomization::MakeInstance));
+
+	RegisterCustomPropertyTypeLayout("RuntimeCurveLinearColor", FOnGetPropertyTypeCustomizationInstance::CreateStatic(&FCurveColorCustomization::MakeInstance));
 }
 
 
@@ -286,4 +289,4 @@ void FDetailCustomizationsModule::RegisterCustomPropertyTypeLayout(FName Propert
 	static FName PropertyEditor("PropertyEditor");
 	FPropertyEditorModule& PropertyModule = FModuleManager::GetModuleChecked<FPropertyEditorModule>(PropertyEditor);
 	PropertyModule.RegisterCustomPropertyTypeLayout(PropertyTypeName, PropertyTypeLayoutDelegate);
-}
\ No newline at end of file
+}
diff --git a/Engine/Source/Editor/UnrealEd/Private/SColorGradientEditor.cpp b/Engine/Source/Editor/UnrealEd/Private/SColorGradientEditor.cpp
--- a/Engine/Source/Editor/UnrealEd/Private/SColorGradientEditor.cpp
+++ b/Engine/Source/Editor/UnrealEd/Private/SColorGradientEditor.cpp
@@ -57,7 +57,7 @@ bool FGradientStopMark::IsValidColorMark( const TArray<FRichCurveEditInfo>& Curv
 
 FLinearColor FGradientStopMark::GetColor( FCurveOwnerInterface& CurveOwner ) const
 {
-	return CastChecked<UCurveLinearColor>( CurveOwner.GetOwner() )->GetLinearColorValue( Time );
+	return CurveOwner.GetLinearColorValue(Time);
 }
 
 void FGradientStopMark::SetColor( const FLinearColor& InColor, FCurveOwnerInterface& CurveOwner )
@@ -139,8 +139,6 @@ int32 SColorGradientEditor::OnPaint( const FPaintArgs& Args, const FGeometry& Al
 		// Pixel to value input converter
 		FTrackScaleInfo ScaleInfo(ViewMinInput.Get(),  ViewMaxInput.Get(), 0.0f, 1.0f, GradientAreaGeometry.Size);
 
-		UCurveLinearColor* LinearColorCurve = CastChecked<UCurveLinearColor>( CurveOwner->GetOwner() );
-
 		// The start and end location in slate units of the area to draw
 		int32 Start = 0;
 		int32 Finish = FMath::TruncToInt( AllottedGeometry.Size.X );
@@ -148,7 +146,7 @@ int32 SColorGradientEditor::OnPaint( const FPaintArgs& Args, const FGeometry& Al
 		TArray<FSlateGradientStop> Stops;
 
 		// If no alpha keys are available, treat the curve as being completely opaque for drawing purposes
-		bool bHasAnyAlphaKeys = LinearColorCurve->HasAnyAlphaKeys(); 
+		bool bHasAnyAlphaKeys = CurveOwner->HasAnyAlphaKeys(); 
 
 		// If any transpareny (A < 1) is found, we'll draw a checkerboard to visualize the color with alpha
 		bool bHasTransparency = false;
@@ -160,7 +158,7 @@ int32 SColorGradientEditor::OnPaint( const FPaintArgs& Args, const FGeometry& Al
 			float Time = ScaleInfo.LocalXToInput(CurrentStep);
 
 			// Sample the curve
-			FLinearColor Color = LinearColorCurve->GetLinearColorValue( Time );
+			FLinearColor Color = CurveOwner->GetLinearColorValue( Time );
 			if( !bHasAnyAlphaKeys )
 			{
 				// Only show alpha if there is at least one key.  For some curves, alpha may not be important
@@ -220,7 +218,7 @@ int32 SColorGradientEditor::OnPaint( const FPaintArgs& Args, const FGeometry& Al
 			// Dont draw stops which are not visible
 			if( XVal >= 0 && XVal <= ColorMarkAreaGeometry.Size.X )
 			{
-				FLinearColor Color = LinearColorCurve->GetLinearColorValue( Mark.Time );
+				FLinearColor Color = CurveOwner->GetLinearColorValue( Mark.Time );
 				Color.A = 1.0f;
 				DrawGradientStopMark( Mark, ColorMarkAreaGeometry, XVal, Color, OutDrawElements, LayerId, MyClippingRect, DrawEffects, true, InWidgetStyle );
 			}
@@ -237,7 +235,7 @@ int32 SColorGradientEditor::OnPaint( const FPaintArgs& Args, const FGeometry& Al
 			// Dont draw stops which are not visible
 			if( XVal >= 0 && XVal <= AlphaMarkAreaGeometry.Size.X )
 			{
-				float Alpha = LinearColorCurve->GetLinearColorValue( Mark.Time ).A;
+				float Alpha = CurveOwner->GetLinearColorValue( Mark.Time ).A;
 				DrawGradientStopMark( Mark, AlphaMarkAreaGeometry, XVal, FLinearColor( Alpha, Alpha, Alpha, 1.0f ), OutDrawElements, LayerId, MyClippingRect, DrawEffects, false, InWidgetStyle );
 			}
 
@@ -845,9 +843,8 @@ void SColorGradientEditor::DeleteStop( const FGradientStopMark& InMark )
 FGradientStopMark SColorGradientEditor::AddStop( const FVector2D& Position, const FGeometry& MyGeometry, bool bColorStop )
 {
 	FScopedTransaction AddStopTrans( LOCTEXT("AddGradientStop", "Add Gradient Stop") );
-	UCurveLinearColor* ColorCurve = CastChecked<UCurveLinearColor>( CurveOwner->GetOwner() );
 
-	ColorCurve->Modify();
+	CurveOwner->ModifyOwner();
 
 	FTrackScaleInfo ScaleInfo(ViewMinInput.Get(),  ViewMaxInput.Get(), 0.0f, 1.0f, MyGeometry.Size);
 
diff --git a/Engine/Source/Editor/UnrealEd/Private/SCurveEditor.cpp b/Engine/Source/Editor/UnrealEd/Private/SCurveEditor.cpp
--- a/Engine/Source/Editor/UnrealEd/Private/SCurveEditor.cpp
+++ b/Engine/Source/Editor/UnrealEd/Private/SCurveEditor.cpp
@@ -2314,7 +2314,7 @@ void SCurveEditor::PostUndo(bool bSuccess)
 
 bool SCurveEditor::IsLinearColorCurve() const 
 {
-	return CurveOwner && CurveOwner->GetOwner() && CurveOwner->GetOwner()->IsA<UCurveLinearColor>();
+	return CurveOwner && CurveOwner->IsLinearColorCurve();
 }
 
 FVector2D SCurveEditor::SnapLocation(FVector2D InLocation)
diff --git a/Engine/Source/Runtime/Engine/Classes/Curves/CurveBase.h b/Engine/Source/Runtime/Engine/Classes/Curves/CurveBase.h
--- a/Engine/Source/Runtime/Engine/Classes/Curves/CurveBase.h
+++ b/Engine/Source/Runtime/Engine/Classes/Curves/CurveBase.h
@@ -430,6 +430,15 @@ public:
 
 	/** Called when the curve has been changed */
 	virtual void OnCurveChanged() = 0;
+
+	/** Whether the curve represents a linear color */
+	virtual bool IsLinearColorCurve() const { return false; }
+
+	/** Evaluate this color curve at the specified time */
+	virtual FLinearColor GetLinearColorValue(float InTime) const { return FLinearColor::Black; }
+
+	/** @return True if the curve has any alpha keys */
+	virtual bool HasAnyAlphaKeys() const { return false; }
 };
 
 
diff --git a/Engine/Source/Runtime/Engine/Classes/Curves/CurveLinearColor.h b/Engine/Source/Runtime/Engine/Classes/Curves/CurveLinearColor.h
--- a/Engine/Source/Runtime/Engine/Classes/Curves/CurveLinearColor.h
+++ b/Engine/Source/Runtime/Engine/Classes/Curves/CurveLinearColor.h
@@ -4,6 +4,20 @@
 #pragma once
 #include "CurveLinearColor.generated.h"
 
+USTRUCT()
+struct ENGINE_API FRuntimeCurveLinearColor
+{
+	GENERATED_USTRUCT_BODY()
+public:
+	UPROPERTY()
+	FRichCurve ColorCurves[4];
+
+	UPROPERTY(EditAnywhere, Category = RuntimeFloatCurve)
+	class UCurveLinearColor* ExternalCurve;
+
+	FLinearColor GetLinearColorValue(float InTime) const;
+};
+
 UCLASS(BlueprintType)
 class ENGINE_API UCurveLinearColor : public UCurveBase
 {
@@ -13,16 +27,16 @@ class ENGINE_API UCurveLinearColor : public UCurveBase
 	UPROPERTY()
 	FRichCurve FloatCurves[4];
 
-	/** Evaluate this float curve at the specified time */
-	UFUNCTION(BlueprintCallable, Category="Math|Curves")
-	virtual FLinearColor GetLinearColorValue(float InTime) const;
-
 	// Begin FCurveOwnerInterface
 	virtual TArray<FRichCurveEditInfoConst> GetCurves() const override;
 	virtual TArray<FRichCurveEditInfo> GetCurves() override;
+	virtual bool IsLinearColorCurve() const { return true; }
+
+	UFUNCTION(BlueprintCallable, Category="Math|Curves")
+	virtual FLinearColor GetLinearColorValue(float InTime) const override;
 
-	/** @return True if the curve as any alpha keys */
-	bool HasAnyAlphaKeys() const { return FloatCurves[3].GetNumKeys() > 0; }
+	bool HasAnyAlphaKeys() const override { return FloatCurves[3].GetNumKeys() > 0; }
+	// End FCurveOwnerInterface
 
 	/** Determine if Curve is the same */
 	bool operator == (const UCurveLinearColor& Curve) const;
diff --git a/Engine/Source/Runtime/Engine/Private/CurveLinearColor.cpp b/Engine/Source/Runtime/Engine/Private/CurveLinearColor.cpp
--- a/Engine/Source/Runtime/Engine/Private/CurveLinearColor.cpp
+++ b/Engine/Source/Runtime/Engine/Private/CurveLinearColor.cpp
@@ -6,6 +6,27 @@
 
 #include "EnginePrivate.h"
 
+FLinearColor FRuntimeCurveLinearColor::GetLinearColorValue(float InTime) const
+{
+	FLinearColor Result;
+
+	Result.R = ColorCurves[0].Eval(InTime);
+	Result.G = ColorCurves[1].Eval(InTime);
+	Result.B = ColorCurves[2].Eval(InTime);
+
+	// No alpha keys means alpha should be 1
+	if (ColorCurves[3].GetNumKeys() == 0)
+	{
+		Result.A = 1.0f;
+	}
+	else
+	{
+		Result.A = ColorCurves[3].Eval(InTime);
+	}
+
+	return Result;
+}
+
 UCurveLinearColor::UCurveLinearColor(const FObjectInitializer& ObjectInitializer)
 	: Super(ObjectInitializer)
 {
