Created struct customization for FRawDistributionVector allowing color-related particle modules to have their vector properties visualized as colors.  Created struct customization for FParticleSysParam in order to filter out irrelevant parameters from the list, depending on the parameter type.

diff --git a/Engine/Source/Editor/DetailCustomizations/DetailCustomizations.Build.cs b/Engine/Source/Editor/DetailCustomizations/DetailCustomizations.Build.cs
--- a/Engine/Source/Editor/DetailCustomizations/DetailCustomizations.Build.cs
+++ b/Engine/Source/Editor/DetailCustomizations/DetailCustomizations.Build.cs
@@ -48,6 +48,7 @@ public class DetailCustomizations : ModuleRules
 
 		PrivateIncludePathModuleNames.AddRange(
 			new string[] {
+                "Engine",
 				"PropertyEditor",
 				"Landscape",
 				"LandscapeEditor",
diff --git a/Engine/Source/Editor/DetailCustomizations/Private/DetailCustomizations.cpp b/Engine/Source/Editor/DetailCustomizations/Private/DetailCustomizations.cpp
--- a/Engine/Source/Editor/DetailCustomizations/Private/DetailCustomizations.cpp
+++ b/Engine/Source/Editor/DetailCustomizations/Private/DetailCustomizations.cpp
@@ -77,6 +77,8 @@
 #include "MoviePlayerSettingsDetails.h"
 #include "SourceCodeAccessSettingsDetails.h"
 #include "ParticleSystemComponentDetails.h"
+#include "ParticleSysParamStructCustomization.h"
+#include "RawDistributionVectorStructCustomization.h"
 #include "DocumentationActorDetails.h"
 #include "MediaPlayerCustomization.h"
 #include "MediaSoundWaveCustomization.h"
@@ -163,23 +165,20 @@ void FDetailCustomizationsModule::RegisterPropertyTypeCustomizations()
 	RegisterCustomPropertyTypeLayout("Int32Range", FOnGetPropertyTypeCustomizationInstance::CreateStatic(&FRangeStructCustomization<int32>::MakeInstance));
 	RegisterCustomPropertyTypeLayout("DateTime", FOnGetPropertyTypeCustomizationInstance::CreateStatic(&FDateTimeStructCustomization::MakeInstance));
 	RegisterCustomPropertyTypeLayout("Timespan", FOnGetPropertyTypeCustomizationInstance::CreateStatic(&FTimespanStructCustomization::MakeInstance));
-
 	RegisterCustomPropertyTypeLayout("BlackboardEntry", FOnGetPropertyTypeCustomizationInstance::CreateStatic(&FBlackboardEntryDetails::MakeInstance));
 	RegisterCustomPropertyTypeLayout("RuntimeFloatCurve", FOnGetPropertyTypeCustomizationInstance::CreateStatic(&FCurveStructCustomization::MakeInstance));
 	RegisterCustomPropertyTypeLayout("EnvNamedValue", FOnGetPropertyTypeCustomizationInstance::CreateStatic(&FEnvQueryParamInstanceCustomization::MakeInstance));
 	RegisterCustomPropertyTypeLayout("NavigationLink", FOnGetPropertyTypeCustomizationInstance::CreateStatic(&FNavLinkStructCustomization::MakeInstance));
 	RegisterCustomPropertyTypeLayout("NavigationSegmentLink", FOnGetPropertyTypeCustomizationInstance::CreateStatic(&FNavLinkStructCustomization::MakeInstance));
-
 	RegisterCustomPropertyTypeLayout("Margin", FOnGetPropertyTypeCustomizationInstance::CreateStatic(&FMarginStructCustomization::MakeInstance));
-
 	RegisterCustomPropertyTypeLayout("DirectoryPath", FOnGetPropertyTypeCustomizationInstance::CreateStatic(&FDirectoryPathStructCustomization::MakeInstance));
 	RegisterCustomPropertyTypeLayout("FilePath", FOnGetPropertyTypeCustomizationInstance::CreateStatic(&FFilePathStructCustomization::MakeInstance));
-
 	RegisterCustomPropertyTypeLayout("InputAxisConfigEntry", FOnGetPropertyTypeCustomizationInstance::CreateStatic(&FInputAxisConfigCustomization::MakeInstance));
 	RegisterCustomPropertyTypeLayout("InputActionKeyMapping", FOnGetPropertyTypeCustomizationInstance::CreateStatic(&FInputActionMappingCustomization::MakeInstance));
 	RegisterCustomPropertyTypeLayout("InputAxisKeyMapping", FOnGetPropertyTypeCustomizationInstance::CreateStatic(&FInputAxisMappingCustomization::MakeInstance));
-
 	RegisterCustomPropertyTypeLayout("RuntimeCurveLinearColor", FOnGetPropertyTypeCustomizationInstance::CreateStatic(&FCurveColorCustomization::MakeInstance));
+	RegisterCustomPropertyTypeLayout("ParticleSysParam", FOnGetPropertyTypeCustomizationInstance::CreateStatic(&FParticleSysParamStructCustomization::MakeInstance));
+	RegisterCustomPropertyTypeLayout("RawDistributionVector", FOnGetPropertyTypeCustomizationInstance::CreateStatic(&FRawDistributionVectorStructCustomization::MakeInstance));
 }
 
 
@@ -289,4 +288,4 @@ void FDetailCustomizationsModule::RegisterCustomPropertyTypeLayout(FName Propert
 	static FName PropertyEditor("PropertyEditor");
 	FPropertyEditorModule& PropertyModule = FModuleManager::GetModuleChecked<FPropertyEditorModule>(PropertyEditor);
 	PropertyModule.RegisterCustomPropertyTypeLayout(PropertyTypeName, PropertyTypeLayoutDelegate);
-}
+}
\ No newline at end of file
diff --git a/Engine/Source/Editor/DetailCustomizations/Private/ParticleSysParamStructCustomization.cpp b/Engine/Source/Editor/DetailCustomizations/Private/ParticleSysParamStructCustomization.cpp
new file mode 100644
--- /dev/null
+++ b/Engine/Source/Editor/DetailCustomizations/Private/ParticleSysParamStructCustomization.cpp
@@ -0,0 +1,202 @@
+// Copyright 1998-2014 Epic Games, Inc. All Rights Reserved.
+
+#include "DetailCustomizationsPrivatePCH.h"
+#include "ParticleSysParamStructCustomization.h"
+#include "IPropertyUtilities.h"
+
+
+#define LOCTEXT_NAMESPACE "ParticleSysParamStructCustomization"
+
+
+TSharedRef<IPropertyTypeCustomization> FParticleSysParamStructCustomization::MakeInstance()
+{
+	return MakeShareable(new FParticleSysParamStructCustomization);
+}
+
+
+FParticleSysParamStructCustomization::FParticleSysParamStructCustomization()
+{
+	ParameterType = PSPT_None;
+}
+
+
+void FParticleSysParamStructCustomization::CustomizeHeader(TSharedRef<IPropertyHandle> StructPropertyHandle, FDetailWidgetRow& HeaderRow, IPropertyTypeCustomizationUtils& StructCustomizationUtils)
+{
+	const bool bDisplayResetToDefault = false;
+	const FString DisplayNameOverride = TEXT("");
+	const FString DisplayToolTipOverride = TEXT("");
+
+	HeaderRow
+		.NameContent()
+		[
+			StructPropertyHandle->CreatePropertyNameWidget(DisplayNameOverride, DisplayToolTipOverride, bDisplayResetToDefault)
+		];
+}
+
+
+void FParticleSysParamStructCustomization::CustomizeChildren(TSharedRef<IPropertyHandle> StructPropertyHandle, IDetailChildrenBuilder& StructBuilder, IPropertyTypeCustomizationUtils& StructCustomizationUtils)
+{
+	// Cache struct property handle
+	PropertyHandle = StructPropertyHandle;
+
+	// Add name property
+	TSharedPtr<IPropertyHandle> NameHandle = StructPropertyHandle->GetChildHandle(GET_MEMBER_NAME_CHECKED(FParticleSysParam, Name));
+	StructBuilder.AddChildProperty(NameHandle.ToSharedRef());
+
+	// Add type property - this is a custom widget which remembers its type so that other widgets can alter their visibility accordingly
+	TSharedPtr<IPropertyHandle> ParamTypeHandle = StructPropertyHandle->GetChildHandle(GET_MEMBER_NAME_CHECKED(FParticleSysParam, ParamType));
+
+	uint8 Type;
+	FPropertyAccess::Result Result = ParamTypeHandle->GetValue(Type);
+	check(Result != FPropertyAccess::Fail);
+
+	if (Result == FPropertyAccess::Success)
+	{
+		ParameterType = Type;
+	}
+	else
+	{
+		// For multiple values
+		ParameterType = INDEX_NONE;
+	}
+
+	TArray<bool> RestrictedList;
+	ParamTypeHandle->GeneratePossibleValues(ParameterTypeNames, ParameterTypeToolTips, RestrictedList);
+
+	StructBuilder.AddChildContent(TEXT("ParamType"))
+		.NameContent()
+		[
+			ParamTypeHandle->CreatePropertyNameWidget()
+		]
+		.ValueContent()
+		.VAlign(VAlign_Center)
+		[
+			SNew(SComboBox<TSharedPtr<FString>>)
+			.OptionsSource(&ParameterTypeNames)
+			.OnGenerateWidget(this, &FParticleSysParamStructCustomization::OnGenerateComboWidget)
+			.OnSelectionChanged(this, &FParticleSysParamStructCustomization::OnComboSelectionChanged)
+			.InitiallySelectedItem(ParameterTypeNames[(ParameterType == INDEX_NONE) ? 0 : ParameterType])
+			[
+				SNew(STextBlock)
+				.Text(this, &FParticleSysParamStructCustomization::GetParameterTypeName)
+				.Font(IPropertyTypeCustomizationUtils::GetRegularFont())
+			]
+		];
+
+	// Add other properties
+	TSharedPtr<IPropertyHandle> ScalarHandle = StructPropertyHandle->GetChildHandle(GET_MEMBER_NAME_CHECKED(FParticleSysParam, Scalar));
+	StructBuilder.AddChildProperty(ScalarHandle.ToSharedRef())
+		.Visibility(TAttribute<EVisibility>(this, &FParticleSysParamStructCustomization::GetScalarVisibility));
+
+	TSharedPtr<IPropertyHandle> ScalarLowHandle = StructPropertyHandle->GetChildHandle(GET_MEMBER_NAME_CHECKED(FParticleSysParam, Scalar_Low));
+	StructBuilder.AddChildProperty(ScalarLowHandle.ToSharedRef())
+		.Visibility(TAttribute<EVisibility>(this, &FParticleSysParamStructCustomization::GetScalarLowVisibility));
+
+	TSharedPtr<IPropertyHandle> VectorHandle = StructPropertyHandle->GetChildHandle(GET_MEMBER_NAME_CHECKED(FParticleSysParam, Vector));
+	StructBuilder.AddChildProperty(VectorHandle.ToSharedRef())
+		.Visibility(TAttribute<EVisibility>(this, &FParticleSysParamStructCustomization::GetVectorVisibility));
+
+	TSharedPtr<IPropertyHandle> VectorLowHandle = StructPropertyHandle->GetChildHandle(GET_MEMBER_NAME_CHECKED(FParticleSysParam, Vector_Low));
+	StructBuilder.AddChildProperty(VectorLowHandle.ToSharedRef())
+		.Visibility(TAttribute<EVisibility>(this, &FParticleSysParamStructCustomization::GetVectorLowVisibility));
+
+	TSharedPtr<IPropertyHandle> ColorHandle = StructPropertyHandle->GetChildHandle(GET_MEMBER_NAME_CHECKED(FParticleSysParam, Color));
+	StructBuilder.AddChildProperty(ColorHandle.ToSharedRef())
+		.Visibility(TAttribute<EVisibility>(this, &FParticleSysParamStructCustomization::GetColorVisibility));
+
+	TSharedPtr<IPropertyHandle> ActorHandle = StructPropertyHandle->GetChildHandle(GET_MEMBER_NAME_CHECKED(FParticleSysParam, Actor));
+	StructBuilder.AddChildProperty(ActorHandle.ToSharedRef())
+		.Visibility(TAttribute<EVisibility>(this, &FParticleSysParamStructCustomization::GetActorVisibility));
+
+	TSharedPtr<IPropertyHandle> MaterialHandle = StructPropertyHandle->GetChildHandle(GET_MEMBER_NAME_CHECKED(FParticleSysParam, Material));
+	StructBuilder.AddChildProperty(MaterialHandle.ToSharedRef())
+		.Visibility(TAttribute<EVisibility>(this, &FParticleSysParamStructCustomization::GetMaterialVisibility));
+}
+
+
+EVisibility FParticleSysParamStructCustomization::GetScalarVisibility() const
+{
+	return (ParameterType == PSPT_Scalar || ParameterType == PSPT_ScalarRand) ? EVisibility::Visible : EVisibility::Collapsed;
+}
+
+EVisibility FParticleSysParamStructCustomization::GetScalarLowVisibility() const
+{
+	return (ParameterType == PSPT_ScalarRand) ? EVisibility::Visible : EVisibility::Collapsed;
+}
+
+EVisibility FParticleSysParamStructCustomization::GetVectorVisibility() const
+{
+	return (ParameterType == PSPT_Vector || ParameterType == PSPT_VectorRand) ? EVisibility::Visible : EVisibility::Collapsed;
+}
+
+EVisibility FParticleSysParamStructCustomization::GetVectorLowVisibility() const
+{
+	return (ParameterType == PSPT_VectorRand) ? EVisibility::Visible : EVisibility::Collapsed;
+}
+
+EVisibility FParticleSysParamStructCustomization::GetColorVisibility() const
+{
+	return (ParameterType == PSPT_Color) ? EVisibility::Visible : EVisibility::Collapsed;
+}
+
+EVisibility FParticleSysParamStructCustomization::GetActorVisibility() const
+{
+	return (ParameterType == PSPT_Actor) ? EVisibility::Visible : EVisibility::Collapsed;
+}
+
+EVisibility FParticleSysParamStructCustomization::GetMaterialVisibility() const
+{
+	return (ParameterType == PSPT_Material) ? EVisibility::Visible : EVisibility::Collapsed;
+}
+
+
+TSharedRef<SWidget> FParticleSysParamStructCustomization::OnGenerateComboWidget(TSharedPtr<FString> InComboString)
+{
+	// Find ToolTip which corresponds to string
+	FString ToolTip;
+	check(ParameterTypeNames.Num() == ParameterTypeToolTips.Num());
+	if (ParameterTypeToolTips.Num() > 0)
+	{
+		int32 Index = ParameterTypeNames.IndexOfByKey(InComboString);
+		if (ensure(Index >= 0))
+		{
+			ToolTip = *ParameterTypeToolTips[Index];
+		}
+	}
+
+	return
+		SNew(SBox)
+		[
+			SNew(STextBlock)
+			.Text(*InComboString)
+			.ToolTipText(ToolTip)
+			.Font(IPropertyTypeCustomizationUtils::GetRegularFont())
+		];
+}
+
+
+void FParticleSysParamStructCustomization::OnComboSelectionChanged(TSharedPtr<FString> InSelectedItem, ESelectInfo::Type SelectInfo)
+{
+	TSharedPtr<IPropertyHandle> ParamTypeHandle = PropertyHandle->GetChildHandle(GET_MEMBER_NAME_CHECKED(FParticleSysParam, ParamType));
+	if (ParamTypeHandle.IsValid())
+	{
+		ParameterType = static_cast<EParticleSysParamType>(ParameterTypeNames.IndexOfByKey(InSelectedItem));
+		if (ensure(ParameterType >= 0))
+		{
+			ensure(ParamTypeHandle->SetValue(static_cast<uint8>(ParameterType)) == FPropertyAccess::Success);
+		}
+	}
+}
+
+
+FText FParticleSysParamStructCustomization::GetParameterTypeName() const
+{
+	if (ParameterType == INDEX_NONE)
+	{
+		return LOCTEXT("MultipleValues", "Multiple Values");
+	}
+
+	return FText::FromString(*ParameterTypeNames[ParameterType]);
+}
+
+#undef LOCTEXT_NAMESPACE
diff --git a/Engine/Source/Editor/DetailCustomizations/Private/ParticleSysParamStructCustomization.h b/Engine/Source/Editor/DetailCustomizations/Private/ParticleSysParamStructCustomization.h
new file mode 100644
--- /dev/null
+++ b/Engine/Source/Editor/DetailCustomizations/Private/ParticleSysParamStructCustomization.h
@@ -0,0 +1,53 @@
+// Copyright 1998-2014 Epic Games, Inc. All Rights Reserved.
+
+#pragma once
+#include "Particles/ParticleSystemComponent.h"
+
+/**
+* Class for FParticleSysParam struct customization
+*/
+class FParticleSysParamStructCustomization : public IPropertyTypeCustomization
+{
+public:
+	static TSharedRef<IPropertyTypeCustomization> MakeInstance();
+
+	FParticleSysParamStructCustomization();
+
+	/** IPropertyTypeCustomization instance */
+	virtual void CustomizeHeader(TSharedRef<IPropertyHandle> StructPropertyHandle, FDetailWidgetRow& HeaderRow, IPropertyTypeCustomizationUtils& StructCustomizationUtils) override;
+	virtual void CustomizeChildren(TSharedRef<IPropertyHandle> StructPropertyHandle, IDetailChildrenBuilder& StructBuilder, IPropertyTypeCustomizationUtils& StructCustomizationUtils) override;
+
+protected:
+
+	/** Called by the combo box to generate each row of the widget */
+	TSharedRef<SWidget> OnGenerateComboWidget(TSharedPtr<FString> InComboString);
+
+	/** Called when the combo box selection changes, when a new parameter type is selected */
+	void OnComboSelectionChanged(TSharedPtr<FString> InSelectedItem, ESelectInfo::Type SelectInfo);
+
+	/** Called to obtain an FText representing the current combo box selection */
+	FText GetParameterTypeName() const;
+
+private:
+
+	/** Cached handle to the struct property */
+	TSharedPtr<IPropertyHandle> PropertyHandle;
+
+	/** Current parameter type */
+	int32 ParameterType;
+
+	/** A list of parameter type names */
+	TArray<TSharedPtr<FString>> ParameterTypeNames;
+
+	/** A list of parameter type tooltips */
+	TArray<TSharedPtr<FString>> ParameterTypeToolTips;
+
+	/** Delegates called by each property widget to determine their visibility */
+	EVisibility GetScalarVisibility() const;
+	EVisibility GetScalarLowVisibility() const;
+	EVisibility GetVectorVisibility() const;
+	EVisibility GetVectorLowVisibility() const;
+	EVisibility GetColorVisibility() const;
+	EVisibility GetActorVisibility() const;
+	EVisibility GetMaterialVisibility() const;
+};
diff --git a/Engine/Source/Editor/DetailCustomizations/Private/RawDistributionVectorStructCustomization.cpp b/Engine/Source/Editor/DetailCustomizations/Private/RawDistributionVectorStructCustomization.cpp
new file mode 100644
--- /dev/null
+++ b/Engine/Source/Editor/DetailCustomizations/Private/RawDistributionVectorStructCustomization.cpp
@@ -0,0 +1,539 @@
+// Copyright 1998-2014 Epic Games, Inc. All Rights Reserved.
+
+#include "DetailCustomizationsPrivatePCH.h"
+#include "RawDistributionVectorStructCustomization.h"
+#include "IPropertyUtilities.h"
+#include "SColorPicker.h"
+#include "SNumericEntryBox.h"
+
+#define LOCTEXT_NAMESPACE "RawDistributionVectorStructCustomization"
+
+
+class FReplaceVectorWithLinearColorBuilder : public IDetailCustomNodeBuilder, public TSharedFromThis<FReplaceVectorWithLinearColorBuilder>
+{
+public:
+
+	/**
+	 * Constructor
+	 *
+	 * @param The property handle representing this item
+	 */
+	explicit FReplaceVectorWithLinearColorBuilder(TSharedRef<IPropertyHandle> InPropertyHandle);
+
+	/**
+	 * Sets a delegate that should be used when the custom node needs to rebuild children
+	 *
+	 * @param A delegate to invoke when the tree should rebuild this nodes children
+	 */
+	virtual void SetOnRebuildChildren(FSimpleDelegate InOnRegenerateChildren) override {}
+
+
+	/**
+	 * Called to generate content in the header of this node ( the actual node content ).
+	 * Only called if HasHeaderRow is true
+	 *
+	 * @param NodeRow The row to put content in
+	 */
+	virtual void GenerateHeaderRowContent(FDetailWidgetRow& NodeRow) override;
+
+	/**
+	 * Called to generate child content of this node
+	 *
+	 * @param OutChildRows An array of rows to add children to
+	 */
+	virtual void GenerateChildContent(IDetailChildrenBuilder& ChildrenBuilder) override;
+
+	/**
+	 * Called each tick if ReqiresTick is true
+	 */
+	virtual void Tick(float DeltaTime) override {}
+
+	/**
+	 * @return true if this node requires tick, false otherwise
+	 */
+	virtual bool RequiresTick() const override { return false; }
+
+	/**
+	 * @return true if this node should be collapsed in the tree
+	 */
+	virtual bool InitiallyCollapsed() const override { return false; }
+
+	/**
+	 * @return The name of this custom builder.  This is used as an identifier to save expansion state if needed
+	 */
+	virtual FName GetName() const override;
+
+
+protected:
+
+	/**
+	 * Creates a widget representing the FVector being pointed to as a FLinearColor
+	 */
+	TSharedRef<SWidget> CreateColorWidget(const TSharedPtr<IPropertyHandle>& StructHandle);
+
+	/**
+	 * Adds a child widget representing a component of an FLinearColor
+	 *
+	 * @param	StructHandle	Handle of the FVector being represented as a FLinearColor
+	 * @param	Text			Label text
+	 * @param	ElementIndex	Element of the FVector it represents (0-2)
+	 * @param	ChildrenBuilder	The builder object
+	 */
+	void AddColorChildProperty(const TSharedPtr<IPropertyHandle>& StructHandle, const FText& Text, int32 ElementIndex, IDetailChildrenBuilder& ChildrenBuilder);
+
+	/**
+	 * Generates a child property, handling nested arrays and structs by creating a child property builder.
+	 *
+	 * @param	Handle			Property handle of the property being generated
+	 * @param	ChildrenBuilder	The current builder object
+	 */
+	void GeneratePropertyContent(const TSharedRef<IPropertyHandle>& Handle, IDetailChildrenBuilder& ChildrenBuilder);
+
+	/**
+	 * Called by the array builder when it needs to generate a new child widget
+	 *
+	 * @param	ElementProperty	Property handle of the array element
+	 * @param	ElementIndex	Index of the element in the array
+	 * @param	ChildrenBuilder	The current builder object
+	 */
+	void OnGenerateArrayElementWidget(TSharedRef<IPropertyHandle> ElementProperty, int32 ElementIndex, IDetailChildrenBuilder& ChildrenBuilder);
+
+	/**
+	 * Called by the color element widgets to determine the current color element value
+	 *
+	 * @param	StructHandle	Property handle of the FVector representing the color
+	 * @param	ElementIndex	Element of the FVector it represents (0-2)
+	 */
+	TOptional<float> OnGetColorElementValue(TSharedPtr<IPropertyHandle> StructHandle, int32 ElementIndex) const;
+
+	/**
+	 * Called when a numeric component is set.
+	 *
+	 * @param	NewValue		New element value
+	 * @param	CommitType		Specifies how the result was entered
+	 * @param	StructHandle	Property handle of the FVector representing the color
+	 * @param	ElementIndex	Element of the FVector it represents (0-2)
+	 */
+	void OnColorElementValueCommitted(float NewValue, ETextCommit::Type CommitType, TSharedPtr<IPropertyHandle> StructHandle, int32 ElementIndex);
+
+	/**
+	 * Called when a numeric component is changed.
+	 *
+	 * @param	NewValue		New element value
+	 * @param	StructHandle	Property handle of the FVector representing the color
+	 * @param	ElementIndex	Element of the FVector it represents (0-2)
+	 */
+	void OnColorElementValueChanged(float NewValue, TSharedPtr<IPropertyHandle> StructHandle, int32 ElementIndex);
+
+	/** Called when the slider is moved on a numeric entry box */
+	void OnBeginSliderMovement();
+
+	/** Called when the slider is released on a numeric entry box */
+	void OnEndSliderMovement(float NewValue);
+
+	/** Returns the FLinearColor represented by the FVector pointed to by the property handle */
+	FLinearColor OnGetColorForColorBlock(TSharedPtr<IPropertyHandle> StructHandle) const;
+
+	/**
+	 * Called when mouse is clicked on the color widget
+	 */
+	FReply OnMouseButtonDownColorBlock(const FGeometry& MyGeometry, const FPointerEvent& MouseEvent, TSharedPtr<IPropertyHandle> StructHandle);
+
+	/** Creates a color picker window */
+	void CreateColorPicker(const TSharedPtr<IPropertyHandle>& StructHandle);
+
+	/**
+	 * Called when a color is chosen from the color picker
+	 *
+	 * @param	NewColor		Color that was chosen
+	 * @param	StructHandle	Property handle of the FVector representing the color
+	 */
+	void OnSetColorFromColorPicker(FLinearColor NewColor, TSharedPtr<IPropertyHandle> StructHandle);
+
+	/**
+	 * Called when the color picker is cancelled
+	 *
+	 * @param	OriginalColor	Original value of the property
+	 * @param	StructHandle	Property handle of the FVector representing the color
+	 */
+	void OnColorPickerCancelled(FLinearColor OriginalColor, TSharedPtr<IPropertyHandle> StructHandle);
+
+	/** Called when a drag starts in the color picker */
+	void OnColorPickerInteractiveBegin();
+
+	/** Called when a drag ends in the color picker */
+	void OnColorPickerInteractiveEnd();
+
+private:
+
+	/** Holds the property handle being referenced by this builder object */
+	TSharedRef<IPropertyHandle> PropertyHandle;
+
+	/** true if the property is an FVector, and hence needs customization */
+	bool bIsVectorProperty;
+
+	/** true if the slider is being dragged in a numeric entry box */
+	bool bIsUsingSlider;
+
+	/** Original value of the property, prior to using the color picker */
+	FLinearColor OldColorValue;
+
+	/** Widget the Color Picker is parented to */
+	TSharedPtr<SWidget> ColorPickerParentWidget;
+};
+
+
+// Some helper functions for obtaining values from the FVector property, and converting them to FLinearColor
+namespace
+{
+	float GetColorElementValue(const TSharedPtr<IPropertyHandle>& StructHandle, int32 ElementIndex)
+	{
+		FVector Result;
+		ensure(StructHandle->GetValue(Result) == FPropertyAccess::Success);
+		return Result[ElementIndex];
+	}
+
+	void SetColorElementValue(const TSharedPtr<IPropertyHandle>& StructHandle, int32 ElementIndex, float Value)
+	{
+		FVector Result;
+		ensure(StructHandle->GetValue(Result) == FPropertyAccess::Success);
+		Result[ElementIndex] = Value;
+		ensure(StructHandle->SetValue(Result) == FPropertyAccess::Success);
+	}
+
+	FLinearColor GetColorValue(const TSharedPtr<IPropertyHandle>& StructHandle)
+	{
+		FVector Result;
+		ensure(StructHandle->GetValue(Result) == FPropertyAccess::Success);
+		return FLinearColor(Result.X, Result.Y, Result.Z, 1.0f);
+	}
+
+	void SetColorValue(const TSharedPtr<IPropertyHandle>& StructHandle, FLinearColor Value)
+	{
+		FVector ValueAsVector(Value.R, Value.G, Value.B);
+		ensure(StructHandle->SetValue(ValueAsVector) == FPropertyAccess::Success);
+	}
+}
+
+
+FReplaceVectorWithLinearColorBuilder::FReplaceVectorWithLinearColorBuilder(TSharedRef<IPropertyHandle> InPropertyHandle)
+	: PropertyHandle(InPropertyHandle)
+	, bIsVectorProperty(false)
+	, bIsUsingSlider(false)
+{
+	// Determine if this is an FVector - if so it will be specialized
+	UProperty* Property = InPropertyHandle->GetProperty();
+
+	if (UStructProperty* StructProperty = Cast<UStructProperty>(Property))
+	{
+		FName StructType = StructProperty->Struct->GetFName();
+		bIsVectorProperty = (StructType == NAME_Vector);
+	}
+}
+
+
+void FReplaceVectorWithLinearColorBuilder::GenerateHeaderRowContent(FDetailWidgetRow& NodeRow)
+{
+	// Only generate a header row if the handle has a valid UProperty.
+	// Note that it's possible for the Property to be NULL if the property node is an FObjectPropertyNode - however we still want to create children in this case.
+	if (PropertyHandle->GetProperty() != nullptr)
+	{
+		NodeRow.NameContent()
+			[
+				PropertyHandle->CreatePropertyNameWidget()
+			];
+
+		if (bIsVectorProperty)
+		{
+			// Customization - make FVector look like an FLinearColor
+			NodeRow.ValueContent()
+				.MinDesiredWidth(250.0f)
+				.MaxDesiredWidth(250.0f)
+				[
+					CreateColorWidget(PropertyHandle)
+				];
+		}
+		else
+		{
+			// Otherwise, use the default property widget
+			NodeRow.ValueContent()
+				.MinDesiredWidth(1)
+				.MaxDesiredWidth(4096)
+				[
+					PropertyHandle->CreatePropertyValueWidget()
+				];
+		}
+	}
+}
+
+
+void FReplaceVectorWithLinearColorBuilder::GenerateChildContent(IDetailChildrenBuilder& ChildrenBuilder)
+{
+	if (bIsVectorProperty)
+	{
+		// Customization - children of an FVector are made to look like color components of an FLinearColor
+		AddColorChildProperty(PropertyHandle, LOCTEXT("RedComponent", "R"), 0, ChildrenBuilder);
+		AddColorChildProperty(PropertyHandle, LOCTEXT("GreenComponent", "G"), 1, ChildrenBuilder);
+		AddColorChildProperty(PropertyHandle, LOCTEXT("BlueComponent", "B"), 2, ChildrenBuilder);
+	}
+	else
+	{
+		// Otherwise, go through the child properties and render them as normal
+		uint32 NumChildren;
+		ensure(PropertyHandle->GetNumChildren(NumChildren) == FPropertyAccess::Success);
+
+		for (uint32 ChildIndex = 0; ChildIndex < NumChildren; ChildIndex++)
+		{
+			TSharedPtr<IPropertyHandle> ChildHandle = PropertyHandle->GetChildHandle(ChildIndex);
+			check(ChildHandle.IsValid());
+			GeneratePropertyContent(ChildHandle.ToSharedRef(), ChildrenBuilder);
+		}
+	}
+}
+
+
+FName FReplaceVectorWithLinearColorBuilder::GetName() const
+{
+	if (PropertyHandle->GetProperty() != nullptr)
+	{
+		return PropertyHandle->GetProperty()->GetFName();
+	}
+
+	return NAME_None;
+}
+
+
+
+TSharedRef<SWidget> FReplaceVectorWithLinearColorBuilder::CreateColorWidget(const TSharedPtr<IPropertyHandle>& StructHandle)
+{
+	FSlateFontInfo NormalText = IDetailLayoutBuilder::GetDetailFont();
+
+	return SNew(SBox)
+		.VAlign(VAlign_Center)
+		[
+			SAssignNew(ColorPickerParentWidget, SColorBlock)
+			.Color(this, &FReplaceVectorWithLinearColorBuilder::OnGetColorForColorBlock, StructHandle)
+			.ShowBackgroundForAlpha(false)
+			.IgnoreAlpha(true)
+			.OnMouseButtonDown(this, &FReplaceVectorWithLinearColorBuilder::OnMouseButtonDownColorBlock, StructHandle)
+			.Size(FVector2D(70.0f, 12.0f))
+		];
+}
+
+
+void FReplaceVectorWithLinearColorBuilder::AddColorChildProperty(const TSharedPtr<IPropertyHandle>& StructHandle, const FText& Text, int32 ElementIndex, IDetailChildrenBuilder& ChildrenBuilder)
+{
+	ChildrenBuilder.AddChildContent(TEXT("Color"))
+		.NameContent()
+		[
+			SNew(STextBlock)
+			.Text(Text)
+			.Font(IPropertyTypeCustomizationUtils::GetRegularFont())
+		]
+	.ValueContent()
+		.MinDesiredWidth(100.0f)
+		.MaxDesiredWidth(100.0f)
+		[
+			SNew(SNumericEntryBox<float>)
+			.Font(IPropertyTypeCustomizationUtils::GetRegularFont())
+			.Value(this, &FReplaceVectorWithLinearColorBuilder::OnGetColorElementValue, StructHandle, ElementIndex)
+			.OnValueCommitted(this, &FReplaceVectorWithLinearColorBuilder::OnColorElementValueCommitted, StructHandle, ElementIndex)
+			.OnValueChanged(this, &FReplaceVectorWithLinearColorBuilder::OnColorElementValueChanged, StructHandle, ElementIndex)
+			.OnBeginSliderMovement(this, &FReplaceVectorWithLinearColorBuilder::OnBeginSliderMovement)
+			.OnEndSliderMovement(this, &FReplaceVectorWithLinearColorBuilder::OnEndSliderMovement)
+			.AllowSpin(true)
+			.MinValue(0.0f)
+			.MaxValue(1.0f)
+			.MinSliderValue(0.0f)
+			.MaxSliderValue(1.0f)
+		];
+}
+
+
+void FReplaceVectorWithLinearColorBuilder::GeneratePropertyContent(const TSharedRef<IPropertyHandle>& Handle, IDetailChildrenBuilder& ChildrenBuilder)
+{
+	// Add to the current builder, depending on the property type.
+	uint32 NumChildren = 0;
+	ensure(Handle->GetNumChildren(NumChildren) == FPropertyAccess::Success);
+	bool bHasChildren = (NumChildren > 0);
+	bool bIsArray = Handle->AsArray().IsValid();
+
+	if (bIsArray)
+	{
+		// Arrays need special handling and will create an array builder
+		TSharedRef<FDetailArrayBuilder> ArrayBuilder = MakeShareable(new FDetailArrayBuilder(Handle));
+		ArrayBuilder->OnGenerateArrayElementWidget(FOnGenerateArrayElementWidget::CreateSP(this, &FReplaceVectorWithLinearColorBuilder::OnGenerateArrayElementWidget));
+		ChildrenBuilder.AddChildCustomBuilder(ArrayBuilder);
+	}
+	else if (bHasChildren)
+	{
+		// If there are children, we invoke a new instance of our custom builder for recursive handling
+		// Note, if this is an FVector, it will be handled specially by the implementation of the IDetailCustomNodeBuilder interface.
+		TSharedRef<FReplaceVectorWithLinearColorBuilder> StructBuilder = MakeShareable(new FReplaceVectorWithLinearColorBuilder(Handle));
+		ChildrenBuilder.AddChildCustomBuilder(StructBuilder);
+	}
+	else
+	{
+		// No children - just add the property.
+		ChildrenBuilder.AddChildProperty(Handle);
+	}
+}
+
+
+void FReplaceVectorWithLinearColorBuilder::OnGenerateArrayElementWidget(TSharedRef<IPropertyHandle> ElementProperty, int32 ElementIndex, IDetailChildrenBuilder& ChildrenBuilder)
+{
+	GeneratePropertyContent(ElementProperty, ChildrenBuilder);
+}
+
+
+TOptional<float> FReplaceVectorWithLinearColorBuilder::OnGetColorElementValue(TSharedPtr<IPropertyHandle> StructHandle, int32 ElementIndex) const
+{
+	return GetColorElementValue(StructHandle, ElementIndex);
+}
+
+
+void FReplaceVectorWithLinearColorBuilder::OnColorElementValueCommitted(float NewValue, ETextCommit::Type CommitType, TSharedPtr<IPropertyHandle> StructHandle, int32 ElementIndex)
+{
+	SetColorElementValue(StructHandle, ElementIndex, NewValue);
+}
+
+
+void FReplaceVectorWithLinearColorBuilder::OnColorElementValueChanged(float NewValue, TSharedPtr<IPropertyHandle> StructHandle, int32 ElementIndex)
+{
+	if (bIsUsingSlider)
+	{
+		SetColorElementValue(StructHandle, ElementIndex, NewValue);
+	}
+}
+
+
+FLinearColor FReplaceVectorWithLinearColorBuilder::OnGetColorForColorBlock(TSharedPtr<IPropertyHandle> StructHandle) const
+{
+	return GetColorValue(StructHandle);
+}
+
+
+void FReplaceVectorWithLinearColorBuilder::OnBeginSliderMovement()
+{
+	GEditor->BeginTransaction(LOCTEXT("SetColorProperty", "Set Color Property"));
+	bIsUsingSlider = true;
+}
+
+
+void FReplaceVectorWithLinearColorBuilder::OnEndSliderMovement(float NewValue)
+{
+	bIsUsingSlider = false;
+	GEditor->EndTransaction();
+}
+
+
+FReply FReplaceVectorWithLinearColorBuilder::OnMouseButtonDownColorBlock(const FGeometry& MyGeometry, const FPointerEvent& MouseEvent, TSharedPtr<IPropertyHandle> StructHandle)
+{
+	if (MouseEvent.GetEffectingButton() != EKeys::LeftMouseButton)
+	{
+		return FReply::Unhandled();
+	}
+
+	CreateColorPicker(StructHandle);
+	return FReply::Handled();
+}
+
+
+void FReplaceVectorWithLinearColorBuilder::CreateColorPicker(const TSharedPtr<IPropertyHandle>& StructHandle)
+{
+	OldColorValue = GetColorValue(StructHandle);
+
+	FColorPickerArgs PickerArgs;
+	PickerArgs.bUseAlpha = false;
+	PickerArgs.bOnlyRefreshOnMouseUp = false;
+	PickerArgs.bOnlyRefreshOnOk = false;
+	PickerArgs.DisplayGamma = TAttribute<float>::Create(TAttribute<float>::FGetter::CreateUObject(GEngine, &UEngine::GetDisplayGamma));
+	PickerArgs.OnColorCommitted = FOnLinearColorValueChanged::CreateSP(this, &FReplaceVectorWithLinearColorBuilder::OnSetColorFromColorPicker, StructHandle);
+	PickerArgs.OnColorPickerCancelled = FOnColorPickerCancelled::CreateSP(this, &FReplaceVectorWithLinearColorBuilder::OnColorPickerCancelled, StructHandle);
+	PickerArgs.OnInteractivePickBegin = FSimpleDelegate::CreateSP(this, &FReplaceVectorWithLinearColorBuilder::OnColorPickerInteractiveBegin);
+	PickerArgs.OnInteractivePickEnd = FSimpleDelegate::CreateSP(this, &FReplaceVectorWithLinearColorBuilder::OnColorPickerInteractiveEnd);
+	PickerArgs.InitialColorOverride = OldColorValue;
+	PickerArgs.ParentWidget = ColorPickerParentWidget;
+
+	OpenColorPicker(PickerArgs);
+}
+
+
+void FReplaceVectorWithLinearColorBuilder::OnSetColorFromColorPicker(FLinearColor NewColor, TSharedPtr<IPropertyHandle> StructHandle)
+{
+	SetColorValue(StructHandle, NewColor);
+}
+
+
+void FReplaceVectorWithLinearColorBuilder::OnColorPickerCancelled(FLinearColor OriginalColor, TSharedPtr<IPropertyHandle> StructHandle)
+{
+	SetColorValue(StructHandle, OldColorValue);
+}
+
+
+void FReplaceVectorWithLinearColorBuilder::OnColorPickerInteractiveBegin()
+{
+	GEditor->BeginTransaction(LOCTEXT("SetColorProperty", "Set Color Property"));
+}
+
+
+void FReplaceVectorWithLinearColorBuilder::OnColorPickerInteractiveEnd()
+{
+	GEditor->EndTransaction();
+}
+
+
+TSharedRef<IPropertyTypeCustomization> FRawDistributionVectorStructCustomization::MakeInstance()
+{
+	return MakeShareable(new FRawDistributionVectorStructCustomization);
+}
+
+
+void FRawDistributionVectorStructCustomization::CustomizeHeader(TSharedRef<IPropertyHandle> StructPropertyHandle, FDetailWidgetRow& HeaderRow, IPropertyTypeCustomizationUtils& StructCustomizationUtils)
+{
+	const bool bDisplayResetToDefault = false;
+	const FString DisplayNameOverride = TEXT("");
+	const FString DisplayToolTipOverride = TEXT("");
+
+	HeaderRow
+	.NameContent()
+	[
+		StructPropertyHandle->CreatePropertyNameWidget(DisplayNameOverride, DisplayToolTipOverride, bDisplayResetToDefault)
+	]
+	.ValueContent()
+	.MinDesiredWidth(1)
+	.MaxDesiredWidth(4096)
+	[
+		StructPropertyHandle->CreatePropertyValueWidget()
+	];
+}
+
+
+void FRawDistributionVectorStructCustomization::CustomizeChildren(TSharedRef<IPropertyHandle> StructPropertyHandle, IDetailChildrenBuilder& StructBuilder, IPropertyTypeCustomizationUtils& StructCustomizationUtils)
+{
+	// Determine from the metadata whether we should treat vectors as FLinearColors or not
+	bool bTreatAsColor = StructPropertyHandle->HasMetaData("TreatAsColor");
+
+	uint32 NumChildren;
+	ensure(StructPropertyHandle->GetNumChildren(NumChildren) == FPropertyAccess::Success);
+
+	// Now recurse through all children, creating a custom builder for each which will either add the default property row, or
+	// a property row exposing a FLinearColor type customization which maps directly to the elements of the original FVector.
+	for (uint32 ChildIndex = 0; ChildIndex < NumChildren; ChildIndex++)
+	{
+		TSharedPtr<IPropertyHandle> ChildHandle = StructPropertyHandle->GetChildHandle(ChildIndex);
+		check(ChildHandle.IsValid());
+		if (bTreatAsColor)
+		{
+			TSharedRef<FReplaceVectorWithLinearColorBuilder> CustomBuilder = MakeShareable(new FReplaceVectorWithLinearColorBuilder(ChildHandle.ToSharedRef()));
+			StructBuilder.AddChildCustomBuilder(CustomBuilder);
+		}
+		else
+		{
+			StructBuilder.AddChildProperty(ChildHandle.ToSharedRef());
+		}
+	}
+}
+
+
+#undef LOCTEXT_NAMESPACE
diff --git a/Engine/Source/Editor/DetailCustomizations/Private/RawDistributionVectorStructCustomization.h b/Engine/Source/Editor/DetailCustomizations/Private/RawDistributionVectorStructCustomization.h
new file mode 100644
--- /dev/null
+++ b/Engine/Source/Editor/DetailCustomizations/Private/RawDistributionVectorStructCustomization.h
@@ -0,0 +1,16 @@
+// Copyright 1998-2014 Epic Games, Inc. All Rights Reserved.
+
+#pragma once
+
+/**
+ * Class for FRawDistributionVector struct customization
+ */
+class FRawDistributionVectorStructCustomization : public IPropertyTypeCustomization
+{
+public:
+	static TSharedRef<IPropertyTypeCustomization> MakeInstance();
+
+	/** IPropertyTypeCustomization instance */
+	virtual void CustomizeHeader(TSharedRef<IPropertyHandle> StructPropertyHandle, FDetailWidgetRow& HeaderRow, IPropertyTypeCustomizationUtils& StructCustomizationUtils) override;
+	virtual void CustomizeChildren(TSharedRef<IPropertyHandle> StructPropertyHandle, IDetailChildrenBuilder& StructBuilder, IPropertyTypeCustomizationUtils& StructCustomizationUtils) override;
+};
diff --git a/Engine/Source/Runtime/Engine/Classes/Particles/Color/ParticleModuleColor.h b/Engine/Source/Runtime/Engine/Classes/Particles/Color/ParticleModuleColor.h
--- a/Engine/Source/Runtime/Engine/Classes/Particles/Color/ParticleModuleColor.h
+++ b/Engine/Source/Runtime/Engine/Classes/Particles/Color/ParticleModuleColor.h
@@ -11,7 +11,7 @@ class UParticleModuleColor : public UParticleModuleColorBase
 	GENERATED_UCLASS_BODY()
 
 	/** Initial color for a particle as a function of Emitter time. */
-	UPROPERTY(EditAnywhere, Category=Color)
+	UPROPERTY(EditAnywhere, Category = Color, meta = (TreatAsColor))
 	struct FRawDistributionVector StartColor;
 
 	/** Initial alpha for a particle as a function of Emitter time. */
diff --git a/Engine/Source/Runtime/Engine/Classes/Particles/Color/ParticleModuleColorOverLife.h b/Engine/Source/Runtime/Engine/Classes/Particles/Color/ParticleModuleColorOverLife.h
--- a/Engine/Source/Runtime/Engine/Classes/Particles/Color/ParticleModuleColorOverLife.h
+++ b/Engine/Source/Runtime/Engine/Classes/Particles/Color/ParticleModuleColorOverLife.h
@@ -11,7 +11,7 @@ class UParticleModuleColorOverLife : public UParticleModuleColorBase
 	GENERATED_UCLASS_BODY()
 
 	/** The color to apply to the particle, as a function of the particle RelativeTime. */
-	UPROPERTY(EditAnywhere, Category=Color)
+	UPROPERTY(EditAnywhere, Category=Color, meta=(TreatAsColor))
 	struct FRawDistributionVector ColorOverLife;
 
 	/** The alpha to apply to the particle, as a function of the particle RelativeTime. */
diff --git a/Engine/Source/Runtime/Engine/Classes/Particles/Color/ParticleModuleColorScaleOverLife.h b/Engine/Source/Runtime/Engine/Classes/Particles/Color/ParticleModuleColorScaleOverLife.h
--- a/Engine/Source/Runtime/Engine/Classes/Particles/Color/ParticleModuleColorScaleOverLife.h
+++ b/Engine/Source/Runtime/Engine/Classes/Particles/Color/ParticleModuleColorScaleOverLife.h
@@ -17,7 +17,7 @@ class UParticleModuleColorScaleOverLife : public UParticleModuleColorBase
 	GENERATED_UCLASS_BODY()
 
 	/** The scale factor for the color.													*/
-	UPROPERTY(EditAnywhere, Category=Color)
+	UPROPERTY(EditAnywhere, Category = Color, meta = (TreatAsColor))
 	struct FRawDistributionVector ColorScaleOverLife;
 
 	/** The scale factor for the alpha.													*/
