Media: Tweaked some build rule files

diff --git a/Engine/Source/Editor/DetailCustomizations/DetailCustomizations.Build.cs b/Engine/Source/Editor/DetailCustomizations/DetailCustomizations.Build.cs
--- a/Engine/Source/Editor/DetailCustomizations/DetailCustomizations.Build.cs
+++ b/Engine/Source/Editor/DetailCustomizations/DetailCustomizations.Build.cs
@@ -39,7 +39,6 @@ public class DetailCustomizations : ModuleRules
                 "InternationalizationSettings",
 				"SourceCodeAccess",
 				"RHI",
-                "Media",
                 "MediaAssets",
                 "HardwareTargeting",
 				"SharedSettingsWidgets"
@@ -49,16 +48,17 @@ public class DetailCustomizations : ModuleRules
 		PrivateIncludePathModuleNames.AddRange(
 			new string[] {
                 "Engine",
-				"PropertyEditor",
+                "Media",
 				"Landscape",
 				"LandscapeEditor",
+				"PropertyEditor",
 			}
 		);
 
 		DynamicallyLoadedModuleNames.AddRange(
 			new string[] {
+				"Layers",
 				"PropertyEditor",
-				"Layers"
 			}
 		);
 	}
diff --git a/Engine/Source/Runtime/MediaAssets/MediaAssets.Build.cs b/Engine/Source/Runtime/MediaAssets/MediaAssets.Build.cs
--- a/Engine/Source/Runtime/MediaAssets/MediaAssets.Build.cs
+++ b/Engine/Source/Runtime/MediaAssets/MediaAssets.Build.cs
@@ -6,6 +6,12 @@ namespace UnrealBuildTool.Rules
 	{
 		public MediaAssets(TargetInfo Target)
 		{
+			DynamicallyLoadedModuleNames.AddRange(
+				new string[] {
+					"Media",
+				}
+			);
+
             PublicDependencyModuleNames.AddRange(
                 new string[] {
 					"Core",
@@ -16,11 +22,16 @@ namespace UnrealBuildTool.Rules
 
 			PrivateDependencyModuleNames.AddRange(
 				new string[] {
-					"Media",
                     "RenderCore",
                     "RHI",
 				}
-			); 
+			);
+
+			PrivateIncludePathModuleNames.AddRange(
+				new string[] {
+					"Media",
+				}
+			);
 
 			PrivateIncludePaths.AddRange(
 				new string[] {
