DetailCustomizations: Include what you use

diff --git a/Engine/Source/Editor/DetailCustomizations/Private/MediaPlayerCustomization.cpp b/Engine/Source/Editor/DetailCustomizations/Private/MediaPlayerCustomization.cpp
--- a/Engine/Source/Editor/DetailCustomizations/Private/MediaPlayerCustomization.cpp
+++ b/Engine/Source/Editor/DetailCustomizations/Private/MediaPlayerCustomization.cpp
@@ -1,7 +1,9 @@
 // Copyright 1998-2014 Epic Games, Inc. All Rights Reserved.
 
 #include "DetailCustomizationsPrivatePCH.h"
-#include "Media.h"
+#include "IMediaModule.h"
+#include "IMediaPlayer.h"
+#include "IMediaPlayerFactory.h"
 #include "MediaPlayer.h"
 #include "MediaPlayerCustomization.h"
 #include "SFilePathPicker.h"
@@ -198,7 +200,7 @@ void FMediaPlayerCustomization::CustomizeDetails( IDetailLayoutBuilder& DetailBu
 /* FMediaPlayerCustomization callbacks
  *****************************************************************************/
 
-FText FMediaPlayerCustomization::HandleDurationTextBlockText( ) const
+FText FMediaPlayerCustomization::HandleDurationTextBlockText() const
 {
 	TOptional<FTimespan> Duration;
 
@@ -274,7 +276,7 @@ FText FMediaPlayerCustomization::HandleSupportedRatesTextBlockText( EMediaPlayba
 }
 
 
-FText FMediaPlayerCustomization::HandleSupportsScrubbingTextBlockText( ) const
+FText FMediaPlayerCustomization::HandleSupportsScrubbingTextBlockText() const
 {
 	TOptional<bool> SupportsScrubbing;
 
@@ -305,7 +307,7 @@ FText FMediaPlayerCustomization::HandleSupportsScrubbingTextBlockText( ) const
 }
 
 
-FText FMediaPlayerCustomization::HandleSupportsSeekingTextBlockText( ) const
+FText FMediaPlayerCustomization::HandleSupportsSeekingTextBlockText() const
 {
 	TOptional<bool> SupportsSeeking;
 
@@ -336,7 +338,7 @@ FText FMediaPlayerCustomization::HandleSupportsSeekingTextBlockText( ) const
 }
 
 
-FString FMediaPlayerCustomization::HandleUrlPickerFilePath( ) const
+FString FMediaPlayerCustomization::HandleUrlPickerFilePath() const
 {
 	FString Url;
 	UrlProperty->GetValue(Url);
@@ -345,7 +347,7 @@ FString FMediaPlayerCustomization::HandleUrlPickerFilePath( ) const
 }
 
 
-FString FMediaPlayerCustomization::HandleUrlPickerFileTypeFilter( ) const
+FString FMediaPlayerCustomization::HandleUrlPickerFileTypeFilter() const
 {
 	FString Filter = TEXT("All files (*.*)|*.*");
 
@@ -399,7 +401,7 @@ void FMediaPlayerCustomization::HandleUrlPickerPathPicked( const FString& Picked
 }
 
 
-EVisibility FMediaPlayerCustomization::HandleUrlWarningIconVisibility( ) const
+EVisibility FMediaPlayerCustomization::HandleUrlWarningIconVisibility() const
 {
 	FString Url;
 
diff --git a/Engine/Source/Editor/DetailCustomizations/Private/MediaPlayerCustomization.h b/Engine/Source/Editor/DetailCustomizations/Private/MediaPlayerCustomization.h
--- a/Engine/Source/Editor/DetailCustomizations/Private/MediaPlayerCustomization.h
+++ b/Engine/Source/Editor/DetailCustomizations/Private/MediaPlayerCustomization.h
@@ -3,6 +3,7 @@
 #pragma once
 
 
+// forward declarations
 enum class EMediaPlaybackDirections;
 
 
@@ -25,7 +26,7 @@ public:
 	 *
 	 * @return The new instance.
 	 */
-	static TSharedRef<IDetailCustomization> MakeInstance( )
+	static TSharedRef<IDetailCustomization> MakeInstance()
 	{
 		return MakeShareable(new FMediaPlayerCustomization());
 	}
@@ -33,31 +34,31 @@ public:
 private:
 
 	/** Callback for getting the text of the Duration text block. */
-	FText HandleDurationTextBlockText( ) const;
+	FText HandleDurationTextBlockText() const;
 
 	/** Callback for getting the text of the ForwardRates text block. */
-	FText HandleForwardRatesTextBlockText( ) const;
+	FText HandleForwardRatesTextBlockText() const;
 
 	/** Callback for getting the text of a supported playback rate text block. */
 	FText HandleSupportedRatesTextBlockText( EMediaPlaybackDirections Direction, bool Unthinned ) const;
 
 	/** Callback for getting the text of the SupportsScrubbing text block. */
-	FText HandleSupportsScrubbingTextBlockText( ) const;
+	FText HandleSupportsScrubbingTextBlockText() const;
 
 	/** Callback for getting the text of the SupportsSeeking text block. */
-	FText HandleSupportsSeekingTextBlockText( ) const;
+	FText HandleSupportsSeekingTextBlockText() const;
 
 	/** Callback for getting the selected path in the URL picker widget. */
-	FString HandleUrlPickerFilePath( ) const;
+	FString HandleUrlPickerFilePath() const;
 
 	/** Callback for getting the file type filter for the URL picker. */
-	FString HandleUrlPickerFileTypeFilter( ) const;
+	FString HandleUrlPickerFileTypeFilter() const;
 
 	/** Callback for picking a path in the URL picker. */
 	void HandleUrlPickerPathPicked( const FString& PickedPath );
 
 	/** Callback for getting the visibility of warning icon for invalid URLs. */
-	EVisibility HandleUrlWarningIconVisibility( ) const;
+	EVisibility HandleUrlWarningIconVisibility() const;
 
 private:
 
diff --git a/Engine/Source/Editor/DetailCustomizations/Private/MediaSoundWaveCustomization.cpp b/Engine/Source/Editor/DetailCustomizations/Private/MediaSoundWaveCustomization.cpp
--- a/Engine/Source/Editor/DetailCustomizations/Private/MediaSoundWaveCustomization.cpp
+++ b/Engine/Source/Editor/DetailCustomizations/Private/MediaSoundWaveCustomization.cpp
@@ -1,6 +1,9 @@
 // Copyright 1998-2014 Epic Games, Inc. All Rights Reserved.
 
 #include "DetailCustomizationsPrivatePCH.h"
+#include "IMediaPlayer.h"
+#include "IMediaTrack.h"
+#include "MediaPlayer.h"
 #include "MediaSoundWaveCustomization.h"
 
 
diff --git a/Engine/Source/Editor/DetailCustomizations/Private/MediaSoundWaveCustomization.h b/Engine/Source/Editor/DetailCustomizations/Private/MediaSoundWaveCustomization.h
--- a/Engine/Source/Editor/DetailCustomizations/Private/MediaSoundWaveCustomization.h
+++ b/Engine/Source/Editor/DetailCustomizations/Private/MediaSoundWaveCustomization.h
@@ -2,9 +2,6 @@
 
 #pragma once
 
-#include "Media.h"
-#include "MediaSoundWave.h"
-
 
 /**
  * Implements a details view customization for the UMediaSoundWave class.
@@ -25,7 +22,7 @@ public:
 	 *
 	 * @return The new instance.
 	 */
-	static TSharedRef<IDetailCustomization> MakeInstance( )
+	static TSharedRef<IDetailCustomization> MakeInstance()
 	{
 		return MakeShareable(new FMediaSoundWaveCustomization());
 	}
diff --git a/Engine/Source/Editor/DetailCustomizations/Private/MediaTextureCustomization.cpp b/Engine/Source/Editor/DetailCustomizations/Private/MediaTextureCustomization.cpp
--- a/Engine/Source/Editor/DetailCustomizations/Private/MediaTextureCustomization.cpp
+++ b/Engine/Source/Editor/DetailCustomizations/Private/MediaTextureCustomization.cpp
@@ -1,6 +1,8 @@
 // Copyright 1998-2014 Epic Games, Inc. All Rights Reserved.
 
 #include "DetailCustomizationsPrivatePCH.h"
+#include "IMediaPlayer.h"
+#include "IMediaTrack.h"
 #include "MediaPlayer.h"
 #include "MediaTextureCustomization.h"
 
@@ -49,7 +51,7 @@ void FMediaTextureCustomization::CustomizeDetails( IDetailLayoutBuilder& DetailB
 /* FMediaTextureCustomization callbacks
  *****************************************************************************/
 
-TSharedRef<SWidget> FMediaTextureCustomization::HandleVideoTrackComboButtonMenuContent( ) const
+TSharedRef<SWidget> FMediaTextureCustomization::HandleVideoTrackComboButtonMenuContent() const
 {
 	// get assigned media player asset
 	UObject* MediaPlayerObj = nullptr;
@@ -105,7 +107,7 @@ void FMediaTextureCustomization::HandleVideoTrackComboButtonMenuEntryExecute( ui
 }
 
 
-FText FMediaTextureCustomization::HandleVideoTrackComboButtonText( ) const
+FText FMediaTextureCustomization::HandleVideoTrackComboButtonText() const
 {
 	// get assigned media player asset
 	UObject* MediaPlayerObj = nullptr;
diff --git a/Engine/Source/Editor/DetailCustomizations/Private/MediaTextureCustomization.h b/Engine/Source/Editor/DetailCustomizations/Private/MediaTextureCustomization.h
--- a/Engine/Source/Editor/DetailCustomizations/Private/MediaTextureCustomization.h
+++ b/Engine/Source/Editor/DetailCustomizations/Private/MediaTextureCustomization.h
@@ -2,9 +2,6 @@
 
 #pragma once
 
-#include "Media.h"
-#include "MediaTexture.h"
-
 
 /**
  * Implements a details view customization for the UMediaTexture class.
@@ -25,21 +22,21 @@ public:
 	 *
 	 * @return The new instance.
 	 */
-	static TSharedRef<IDetailCustomization> MakeInstance( )
+	static TSharedRef<IDetailCustomization> MakeInstance()
 	{
 		return MakeShareable(new FMediaTextureCustomization());
 	}
 
 private:
 
-	// Callback for generating the menu content of the VideoTrack combo box.
-	TSharedRef<SWidget> HandleVideoTrackComboButtonMenuContent( ) const;
+	/** Callback for generating the menu content of the VideoTrack combo box. */
+	TSharedRef<SWidget> HandleVideoTrackComboButtonMenuContent() const;
 
-	// Callback for selecting a track in the VideoTrack combo box.
+	/** Callback for selecting a track in the VideoTrack combo box. */
 	void HandleVideoTrackComboButtonMenuEntryExecute( uint32 TrackIndex );
 
-	// Callback for getting the text of the VideoTrack combo box.
-	FText HandleVideoTrackComboButtonText( ) const;
+	/** Callback for getting the text of the VideoTrack combo box. */
+	FText HandleVideoTrackComboButtonText() const;
 
 private:
 
