SceneCaptures: Support for setting show flags for the view of a scene capture actors via the details panel.

diff --git a/Engine/Source/Editor/DetailCustomizations/Private/DetailCustomizations.cpp b/Engine/Source/Editor/DetailCustomizations/Private/DetailCustomizations.cpp
--- a/Engine/Source/Editor/DetailCustomizations/Private/DetailCustomizations.cpp
+++ b/Engine/Source/Editor/DetailCustomizations/Private/DetailCustomizations.cpp
@@ -89,6 +89,7 @@
 #include "TimespanStructCustomization.h"
 #include "FbxImportUIDetails.h"
 #include "RigDetails.h"
+#include "SceneCaptureDetails.h"
 #include "CurveColorCustomization.h"
 
 IMPLEMENT_MODULE( FDetailCustomizationsModule, DetailCustomizations );
@@ -202,6 +203,8 @@ void FDetailCustomizationsModule::RegisterObjectCustomizations()
 	RegisterCustomClassLayout("MeshComponent", FOnGetDetailCustomizationInstance::CreateStatic(&FMeshComponentDetails::MakeInstance));
 	RegisterCustomClassLayout("MatineeActor", FOnGetDetailCustomizationInstance::CreateStatic(&FMatineeActorDetails::MakeInstance));
 	RegisterCustomClassLayout("ReflectionCapture", FOnGetDetailCustomizationInstance::CreateStatic(&FReflectionCaptureDetails::MakeInstance));
+	RegisterCustomClassLayout("SceneCapture2D", FOnGetDetailCustomizationInstance::CreateStatic(&FSceneCaptureDetails::MakeInstance));
+	RegisterCustomClassLayout("SceneCaptureCube", FOnGetDetailCustomizationInstance::CreateStatic(&FSceneCaptureDetails::MakeInstance));
 	RegisterCustomClassLayout("SkyLight", FOnGetDetailCustomizationInstance::CreateStatic(&FSkyLightComponentDetails::MakeInstance));
 	RegisterCustomClassLayout("Brush", FOnGetDetailCustomizationInstance::CreateStatic(&FBrushDetails::MakeInstance));
 	RegisterCustomClassLayout("AmbientSound", FOnGetDetailCustomizationInstance::CreateStatic(&FAmbientSoundDetails::MakeInstance));
diff --git a/Engine/Source/Editor/DetailCustomizations/Private/SceneCaptureDetails.cpp b/Engine/Source/Editor/DetailCustomizations/Private/SceneCaptureDetails.cpp
new file mode 100644
--- /dev/null
+++ b/Engine/Source/Editor/DetailCustomizations/Private/SceneCaptureDetails.cpp
@@ -0,0 +1,288 @@
+// Copyright 1998-2014 Epic Games, Inc. All Rights Reserved.
+
+#include "DetailCustomizationsPrivatePCH.h"
+#include "Runtime/Engine/Public/ShowFlags.h"
+#include "SceneCaptureDetails.h"
+
+#define LOCTEXT_NAMESPACE "SceneCaptureDetails"
+
+TSharedRef<IDetailCustomization> FSceneCaptureDetails::MakeInstance()
+{
+	return MakeShareable( new FSceneCaptureDetails );
+}
+
+// Used to sort the show flags inside of their categories in the order of the text that is actually displayed
+inline static bool SortAlphabeticallyByLocalizedText(const FString& ip1, const FString& ip2)
+{
+	FText LocalizedText1;
+	FEngineShowFlags::FindShowFlagDisplayName(ip1, LocalizedText1);
+
+	FText LocalizedText2;
+	FEngineShowFlags::FindShowFlagDisplayName(ip2, LocalizedText2);
+
+	return LocalizedText1.ToString() < LocalizedText2.ToString();
+}
+
+void FSceneCaptureDetails::CustomizeDetails( IDetailLayoutBuilder& DetailLayout )
+{
+	// Add all the properties that are there by default
+	// (These would get added by default anyway, but we want to add them first so what we add next comes later in the list)
+	TArray<TSharedRef<IPropertyHandle>> SceneCaptureCategoryDefaultProperties;
+	IDetailCategoryBuilder& SceneCaptureCategoryBuilder = DetailLayout.EditCategory("SceneCapture");
+	SceneCaptureCategoryBuilder.GetDefaultProperties(SceneCaptureCategoryDefaultProperties);
+
+	for (TSharedRef<IPropertyHandle> Handle : SceneCaptureCategoryDefaultProperties)
+	{
+		SceneCaptureCategoryBuilder.AddProperty(Handle);
+	}
+
+	const TArray< TWeakObjectPtr<UObject> >& SelectedObjects = DetailLayout.GetDetailsView().GetSelectedObjects();
+
+	for( int32 ObjectIndex = 0; ObjectIndex < SelectedObjects.Num(); ++ObjectIndex )
+	{
+		const TWeakObjectPtr<UObject>& CurrentObject = SelectedObjects[ObjectIndex];
+		if ( CurrentObject.IsValid() )
+		{
+			ASceneCapture2D* CurrentCaptureActor2D = Cast<ASceneCapture2D>(CurrentObject.Get());
+			if (CurrentCaptureActor2D != nullptr)
+			{
+				SceneCaptureComponent = Cast<USceneCaptureComponent>(CurrentCaptureActor2D->GetCaptureComponent2D());
+				break;
+			}
+			ASceneCaptureCube* CurrentCaptureActorCube = Cast<ASceneCaptureCube>(CurrentObject.Get());
+			if (CurrentCaptureActorCube != nullptr)
+			{
+				SceneCaptureComponent = Cast<USceneCaptureComponent>(CurrentCaptureActorCube->GetCaptureComponentCube());
+				break;
+			}
+		}
+	}
+
+	// Hide certain categories
+	TArray<EShowFlagGroup> ShowFlagGroupsToHideForCaptures;
+	ShowFlagGroupsToHideForCaptures.Add(SFG_Hidden);
+	ShowFlagGroupsToHideForCaptures.Add(SFG_CollisionModes);
+	ShowFlagGroupsToHideForCaptures.Add(SFG_Visualize);
+	ShowFlagGroupsToHideForCaptures.Add(SFG_Developer);
+	ShowFlagGroupsToHideForCaptures.Add(SFG_PostProcess);
+
+	// Hide certain flags
+	TArray<FEngineShowFlags::EShowFlag> ShowFlagsToHideForCaptures;
+
+	// General
+	ShowFlagsToHideForCaptures.Add(FEngineShowFlags::EShowFlag::SF_AntiAliasing);
+	ShowFlagsToHideForCaptures.Add(FEngineShowFlags::EShowFlag::SF_Collision);
+	ShowFlagsToHideForCaptures.Add(FEngineShowFlags::EShowFlag::SF_Grid);
+	ShowFlagsToHideForCaptures.Add(FEngineShowFlags::EShowFlag::SF_Navigation);
+
+	// Advanced
+	ShowFlagsToHideForCaptures.Add(FEngineShowFlags::EShowFlag::SF_AudioRadius);
+	ShowFlagsToHideForCaptures.Add(FEngineShowFlags::EShowFlag::SF_StreamingBounds);
+	ShowFlagsToHideForCaptures.Add(FEngineShowFlags::EShowFlag::SF_Bounds);
+	ShowFlagsToHideForCaptures.Add(FEngineShowFlags::EShowFlag::SF_BSPSplit);
+	ShowFlagsToHideForCaptures.Add(FEngineShowFlags::EShowFlag::SF_CameraAspectRatioBars);
+	ShowFlagsToHideForCaptures.Add(FEngineShowFlags::EShowFlag::SF_CameraFrustums);
+	ShowFlagsToHideForCaptures.Add(FEngineShowFlags::EShowFlag::SF_CameraSafeFrames);
+	ShowFlagsToHideForCaptures.Add(FEngineShowFlags::EShowFlag::SF_Constraints);
+	ShowFlagsToHideForCaptures.Add(FEngineShowFlags::EShowFlag::SF_HighResScreenshotMask);
+	ShowFlagsToHideForCaptures.Add(FEngineShowFlags::EShowFlag::SF_LevelColoration);
+	ShowFlagsToHideForCaptures.Add(FEngineShowFlags::EShowFlag::SF_LOD);
+	ShowFlagsToHideForCaptures.Add(FEngineShowFlags::EShowFlag::SF_MeshEdges);
+	ShowFlagsToHideForCaptures.Add(FEngineShowFlags::EShowFlag::SF_ModeWidgets);
+	ShowFlagsToHideForCaptures.Add(FEngineShowFlags::EShowFlag::SF_SeparateTranslucency);
+	ShowFlagsToHideForCaptures.Add(FEngineShowFlags::EShowFlag::SF_ShadowFrustums);
+	ShowFlagsToHideForCaptures.Add(FEngineShowFlags::EShowFlag::SF_Splines);
+	ShowFlagsToHideForCaptures.Add(FEngineShowFlags::EShowFlag::SF_TemporalAA);
+	ShowFlagsToHideForCaptures.Add(FEngineShowFlags::EShowFlag::SF_VertexColors);
+	ShowFlagsToHideForCaptures.Add(FEngineShowFlags::EShowFlag::SF_VisualizeSenses);
+	ShowFlagsToHideForCaptures.Add(FEngineShowFlags::EShowFlag::SF_Volumes);
+	
+	// Post process
+	//ShowFlagsToHideForCaptures.Add(FEngineShowFlags::EShowFlag::SF_HMDDistortion);
+	//ShowFlagsToHideForCaptures.Add(FEngineShowFlags::EShowFlag::SF_MotionBlur);
+
+	// Lighting features
+	ShowFlagsToHideForCaptures.Add(FEngineShowFlags::EShowFlag::SF_SubsurfaceScattering);
+
+	// Lighting Components
+	ShowFlagsToHideForCaptures.Add(FEngineShowFlags::EShowFlag::SF_Diffuse);
+	ShowFlagsToHideForCaptures.Add(FEngineShowFlags::EShowFlag::SF_DirectLighting);
+	ShowFlagsToHideForCaptures.Add(FEngineShowFlags::EShowFlag::SF_DirectionalLights);
+	ShowFlagsToHideForCaptures.Add(FEngineShowFlags::EShowFlag::SF_PointLights);
+	ShowFlagsToHideForCaptures.Add(FEngineShowFlags::EShowFlag::SF_Specular);
+	ShowFlagsToHideForCaptures.Add(FEngineShowFlags::EShowFlag::SF_SpotLights);
+	
+	// Create array of flag name strings for each group
+	TArray< TArray<FString> > ShowFlagsByGroup;
+	for (int32 GroupIndex = 0; GroupIndex < SFG_Max; ++GroupIndex)
+	{
+		ShowFlagsByGroup.Add(TArray<FString>());
+	}
+
+	// Enumerate all the flags and put them in the correct group
+	bool FlagExists;
+	int32 FlagIndex = 0;
+	do 
+	{
+		TArray<FString> GroupStrings;
+		FString FlagName;
+		FlagName = FEngineShowFlags::FindNameByIndex(FlagIndex);
+		FlagExists = !FlagName.IsEmpty();
+		if (FlagExists)
+		{
+			EShowFlagGroup Group = FEngineShowFlags::FindShowFlagGroup(*FlagName);
+			ShowFlagsByGroup[Group].Add(FlagName);
+		}
+		++FlagIndex;
+	} while (FlagExists);
+
+	// Sort the flags in their respective group alphabetically
+	for (TArray<FString>& ShowFlagGroup : ShowFlagsByGroup)
+	{
+		ShowFlagGroup.Sort(SortAlphabeticallyByLocalizedText);
+	}
+
+	// Add each group
+	for (int32 GroupIndex = 0; GroupIndex < SFG_Max; ++GroupIndex)
+	{
+		bool bGroupHidden = false;
+
+		// Don't add show flag groups that we've specified to hide for Scene Captures
+		for (EShowFlagGroup HiddenGroup : ShowFlagGroupsToHideForCaptures)
+		{
+			if (EShowFlagGroup(GroupIndex) == HiddenGroup)
+			{
+				bGroupHidden = true;
+				break;
+			}
+		}
+
+		if (!bGroupHidden)
+		{
+			FText GroupName;
+			FText GroupTooltip;
+			switch (GroupIndex)
+			{
+			case SFG_Normal:
+				GroupName = LOCTEXT("CommonShowFlagHeader", "General Show Flags");
+				break;
+			case SFG_Advanced:
+				GroupName = LOCTEXT("AdvancedShowFlagsMenu", "Advanced Show Flags");
+				break;
+			case SFG_PostProcess:
+				GroupName = LOCTEXT("PostProcessShowFlagsMenu", "Post Processing Show Flags");
+				break;
+			case SFG_Developer:
+				GroupName = LOCTEXT("DeveloperShowFlagsMenu", "Developer Show Flags");
+				break;
+			case SFG_Visualize:
+				GroupName = LOCTEXT("VisualizeShowFlagsMenu", "Visualize Show Flags");
+				break;
+			case SFG_LightingComponents:
+				GroupName = LOCTEXT("LightingComponentsShowFlagsMenu", "Lighting Components Show Flags");
+				break;
+			case SFG_LightingFeatures:
+				GroupName = LOCTEXT("LightingFeaturesShowFlagsMenu", "Lighting Features Show Flags");
+				break;
+			case SFG_CollisionModes:
+				GroupName = LOCTEXT("CollisionModesShowFlagsMenu", "Collision Modes Show Flags");
+				break;
+			case SFG_Hidden:
+				GroupName = LOCTEXT("HiddenShowFlagsMenu", "Hidden Show Flags");
+				break;
+			default:
+				// Should not get here unless a new group is added without being updated here
+				GroupName = LOCTEXT("MiscFlagsMenu", "Misc Show Flags");
+				break;
+			}
+
+			FName GroupFName = FName(*(GroupName.ToString()));
+			IDetailGroup& Group = SceneCaptureCategoryBuilder.AddGroup(GroupFName, GroupName.ToString());
+
+			// Add each show flag for this group
+			for (FString& FlagName : ShowFlagsByGroup[GroupIndex])
+			{
+				bool bFlagHidden = false;
+				FText LocalizedText;
+				FEngineShowFlags::FindShowFlagDisplayName(FlagName, LocalizedText);
+
+				for (FEngineShowFlags::EShowFlag HiddenFlag : ShowFlagsToHideForCaptures)
+				{
+					// Don't add show flag that we've specified to hide for Scene Captures
+					// or that cannot be toggled in the editor
+					if (FEngineShowFlags::EShowFlag(FEngineShowFlags::FindIndexByName(*FlagName)) == HiddenFlag 
+						|| !FEngineShowFlags::CanBeToggledInEditor(*FlagName)
+						)
+					{
+						bFlagHidden = true;
+						break;
+					}
+				}
+
+				if (!bFlagHidden)
+				{
+					Group.AddWidgetRow()
+						.IsEnabled(true)
+						.NameContent()
+						[
+							SNew(STextBlock)
+							.Text(LocalizedText)
+						]
+					.ValueContent()
+						[
+							SNew(SCheckBox)
+							.OnCheckStateChanged(this, &FSceneCaptureDetails::OnShowFlagCheckStateChanged, FlagName)
+							.IsChecked(this, &FSceneCaptureDetails::OnGetDisplayCheckState, FlagName)
+						]
+					.FilterString(FlagName);
+				}
+			}
+		}
+	}
+}
+
+ESlateCheckBoxState::Type FSceneCaptureDetails::OnGetDisplayCheckState(FString ShowFlagName) const
+{
+	bool IsChecked = false;
+	FEngineShowFlagsSetting* FlagSetting;
+	if (SceneCaptureComponent->GetSettingForShowFlag(ShowFlagName, &FlagSetting))
+	{
+		IsChecked = FlagSetting->Enabled;
+	}
+	else
+	{
+		int32 SettingIndex = SceneCaptureComponent->ShowFlags.FindIndexByName(*(ShowFlagName));
+		if (SettingIndex != INDEX_NONE)
+		{
+			IsChecked = SceneCaptureComponent->ShowFlags.GetSingleFlag(SettingIndex);
+		}
+	}
+
+	return IsChecked ? ESlateCheckBoxState::Checked : ESlateCheckBoxState::Unchecked;
+}
+
+void FSceneCaptureDetails::OnShowFlagCheckStateChanged(ESlateCheckBoxState::Type InNewRadioState, FString FlagName)
+{
+	SceneCaptureComponent->Modify();
+	FEngineShowFlagsSetting* FlagSetting;
+	bool bNewCheckState = (InNewRadioState == ESlateCheckBoxState::Checked);
+
+	// If setting exists, update it
+	if (SceneCaptureComponent->GetSettingForShowFlag(FlagName, &FlagSetting))
+	{
+		FlagSetting->Enabled = bNewCheckState;
+	}
+	// Otherwise create a new setting
+	else
+	{
+		FEngineShowFlagsSetting NewFlagSetting;
+		NewFlagSetting.ShowFlagName = FlagName;
+		NewFlagSetting.Enabled = bNewCheckState;
+		SceneCaptureComponent->ShowFlagSettings.Add(NewFlagSetting);
+	}
+
+	// Ensure PostEditChangeProperty is called, which in turn calls update
+	SceneCaptureComponent->PostEditChange();
+}
+
+#undef LOCTEXT_NAMESPACE
diff --git a/Engine/Source/Editor/DetailCustomizations/Private/SceneCaptureDetails.h b/Engine/Source/Editor/DetailCustomizations/Private/SceneCaptureDetails.h
new file mode 100644
--- /dev/null
+++ b/Engine/Source/Editor/DetailCustomizations/Private/SceneCaptureDetails.h
@@ -0,0 +1,28 @@
+// Copyright 1998-2014 Epic Games, Inc. All Rights Reserved.
+
+#pragma once
+
+class FSceneCaptureDetails : public IDetailCustomization
+{
+public:
+	/** Makes a new instance of this detail layout class for a specific detail view requesting it */
+	static TSharedRef<IDetailCustomization> MakeInstance();
+private:
+	/** IDetailCustomization interface */
+	virtual void CustomizeDetails( IDetailLayoutBuilder& DetailLayout ) override;
+
+private:
+	/** The selected reflection capture */
+	TWeakObjectPtr<USceneCaptureComponent> SceneCaptureComponent;
+
+	/**
+	* Gets the display state to send to a display filter check box
+	*
+	* @param The type of checkbox.
+	* @return The desired checkbox state.
+	*/
+	ESlateCheckBoxState::Type OnGetDisplayCheckState(FString ShowFlagName) const;
+
+	/** Show flag settings changed, so update the scene capture */
+	void OnShowFlagCheckStateChanged(ESlateCheckBoxState::Type InNewRadioState, FString FlagName);
+};
diff --git a/Engine/Source/Runtime/Engine/Classes/Components/SceneCaptureComponent.h b/Engine/Source/Runtime/Engine/Classes/Components/SceneCaptureComponent.h
--- a/Engine/Source/Runtime/Engine/Classes/Components/SceneCaptureComponent.h
+++ b/Engine/Source/Runtime/Engine/Classes/Components/SceneCaptureComponent.h
@@ -2,8 +2,21 @@
 
 
 #pragma once
+#include "Runtime/Engine/Public/ShowFlags.h"
 #include "SceneCaptureComponent.generated.h"
 
+USTRUCT()
+struct FEngineShowFlagsSetting
+{
+	GENERATED_USTRUCT_BODY()
+
+	UPROPERTY(EditAnywhere, BlueprintReadWrite, Category = SceneCapture)
+	FString ShowFlagName;
+
+	UPROPERTY(EditAnywhere, BlueprintReadWrite, Category = SceneCapture)
+	bool Enabled;
+};
+
 	// -> will be exported to EngineDecalClasses.h
 UCLASS(hidecategories=(abstract, Collision, Object, Physics, SceneComponent, Mobility), MinimalAPI)
 class USceneCaptureComponent : public USceneComponent
@@ -22,6 +35,14 @@ class USceneCaptureComponent : public USceneComponent
 	UPROPERTY(EditAnywhere, BlueprintReadWrite, Category=SceneCapture, meta=(UIMin = "100", UIMax = "10000"))
 	float MaxViewDistanceOverride;
 
+	/** ShowFlags for the SceneCapture's ViewFamily, to control rendering settings for this view. Hidden but accessible through details customization */
+	UPROPERTY(EditAnywhere, BlueprintReadWrite, interp, Category=SceneComponent)
+	TArray<struct FEngineShowFlagsSetting> ShowFlagSettings;
+
+	// TODO: Make this a UStruct to set directly?
+	/** Settings stored here read from the strings and int values in the ShowFlagSettings array */
+	FEngineShowFlags ShowFlags;
+
 public:
 
 	/** Adds the component to our list of hidden components. */
@@ -34,8 +55,25 @@ public:
 
 	/** Returns the view state, if any, and allocates one if needed. This function can return NULL, e.g. when bCaptureEveryFrame is false. */
 	ENGINE_API FSceneViewStateInterface* GetViewState();
-	
+
+	/** Return a boolean for whether this flag exists in the ShowFlagSettings array, and a pointer to the flag if it does exist  */
+	ENGINE_API bool GetSettingForShowFlag(FString FlagName, FEngineShowFlagsSetting** ShowFlagSettingOut);
+
+#if WITH_EDITOR
+	/**
+	* Called when a property on this object has been modified externally
+	*
+	* @param PropertyThatChanged the property that was modified
+	*/
+	virtual void PostEditChangeProperty(struct FPropertyChangedEvent& PropertyChangedEvent) override;
+
+	/** Called after done loading to update show flags from saved data */
+	virtual void PostLoad() override;
+#endif
+
 private:
+	/** Update the show flags from our show flags settings (ideally, you'd be able to set this more directly, but currently unable to make FEngineShowFlags a UStruct to use it as a UProperty...) */
+	void UpdateShowFlags();
 
 	/**
 	 * The view state holds persistent scene rendering state and enables occlusion culling in scene captures.
diff --git a/Engine/Source/Runtime/Engine/Private/SceneCaptureComponent.cpp b/Engine/Source/Runtime/Engine/Private/SceneCaptureComponent.cpp
--- a/Engine/Source/Runtime/Engine/Private/SceneCaptureComponent.cpp
+++ b/Engine/Source/Runtime/Engine/Private/SceneCaptureComponent.cpp
@@ -144,10 +144,23 @@ void ASceneCaptureCube::PostEditMove(bool bFinished)
 #endif
 // -----------------------------------------------
 USceneCaptureComponent::USceneCaptureComponent(const FObjectInitializer& ObjectInitializer)
-	: Super(ObjectInitializer)
+	: Super(ObjectInitializer), ShowFlags(FEngineShowFlags(ESFIM_Game))
 {
 	bCaptureEveryFrame = true;
 	MaxViewDistanceOverride = -1;
+
+	// Disable features that are not desired when capturing the scene
+	ShowFlags.MotionBlur = 0; // motion blur doesn't work correctly with scene captures.
+	ShowFlags.SeparateTranslucency = 0;
+	ShowFlags.HMDDistortion = 0;
+}
+
+void USceneCaptureComponent::PostLoad()
+{
+	Super::PostLoad();
+
+	// Make sure any loaded saved flag settings are reflected in our FEngineShowFlags
+	UpdateShowFlags();
 }
 
 
@@ -188,6 +201,49 @@ FSceneViewStateInterface* USceneCaptureComponent::GetViewState()
 	return ViewStateInterface;
 }
 
+void USceneCaptureComponent::UpdateShowFlags()
+{
+	for (FEngineShowFlagsSetting ShowFlagSetting : ShowFlagSettings)
+	{
+		int32 SettingIndex = ShowFlags.FindIndexByName(*(ShowFlagSetting.ShowFlagName));
+		if (SettingIndex != INDEX_NONE)
+		{ 
+			ShowFlags.SetSingleFlag(SettingIndex, ShowFlagSetting.Enabled);
+		}
+	}
+}
+
+#if WITH_EDITOR
+void USceneCaptureComponent::PostEditChangeProperty(struct FPropertyChangedEvent& PropertyChangedEvent)
+{
+	Super::PostEditChangeProperty(PropertyChangedEvent);
+
+	const FName Name = (PropertyChangedEvent.Property != nullptr) ? PropertyChangedEvent.Property->GetFName() : NAME_None;
+	const FName MemberPropertyName = (PropertyChangedEvent.MemberProperty != NULL) ? PropertyChangedEvent.MemberProperty->GetFName() : NAME_None;
+
+	// If our ShowFlagSetting UStruct changed, (or if PostEditChange was called without specifying a property) update the actual show flags
+	if (MemberPropertyName.IsEqual("ShowFlagSettings") || MemberPropertyName.IsNone())
+	{
+		UpdateShowFlags();
+	}
+}
+#endif
+
+bool USceneCaptureComponent::GetSettingForShowFlag(FString FlagName, FEngineShowFlagsSetting** ShowFlagSettingOut)
+{
+	bool HasSetting = false;
+	for (int32 ShowFlagSettingsIndex = 0; ShowFlagSettingsIndex < ShowFlagSettings.Num(); ++ShowFlagSettingsIndex)
+	{
+		if (ShowFlagSettings[ShowFlagSettingsIndex].ShowFlagName.Equals(FlagName))
+		{
+			HasSetting = true;
+			*ShowFlagSettingOut = &(ShowFlagSettings[ShowFlagSettingsIndex]);
+			break;
+		}
+	}	
+	return HasSetting;
+}
+
 // -----------------------------------------------
 
 
diff --git a/Engine/Source/Runtime/Renderer/Private/SceneCaptureRendering.cpp b/Engine/Source/Runtime/Renderer/Private/SceneCaptureRendering.cpp
--- a/Engine/Source/Runtime/Renderer/Private/SceneCaptureRendering.cpp
+++ b/Engine/Source/Runtime/Renderer/Private/SceneCaptureRendering.cpp
@@ -166,13 +166,9 @@ FSceneRenderer* FScene::CreateSceneRenderer( USceneCaptureComponent* SceneCaptur
 	FSceneViewFamilyContext ViewFamily(FSceneViewFamily::ConstructionValues(
 		Resource,
 		this,
-		FEngineShowFlags(ESFIM_Game))
+		SceneCaptureComponent->ShowFlags)
 		.SetResolveScene(!bCaptureSceneColour));
 
-	// Disable features that are not desired when capturing the scene
-	ViewFamily.EngineShowFlags.MotionBlur = 0; // motion blur doesn't work correctly with scene captures.
-	ViewFamily.EngineShowFlags.SeparateTranslucency = 0;
-	ViewFamily.EngineShowFlags.HMDDistortion = 0;
 	FSceneViewInitOptions ViewInitOptions;
 	ViewInitOptions.SetViewRectangle(FIntRect(0, 0, CaptureSize.X, CaptureSize.Y));
 	ViewInitOptions.ViewFamily = &ViewFamily;
