Scene Capture details panel: Remove a few show flags that don't make sense for scene captures. Reconfigured the show flags to be a whitelist model instead of a blacklist, to ensure everything shown here is supported.

diff --git a/Engine/Source/Editor/DetailCustomizations/Private/SceneCaptureDetails.cpp b/Engine/Source/Editor/DetailCustomizations/Private/SceneCaptureDetails.cpp
--- a/Engine/Source/Editor/DetailCustomizations/Private/SceneCaptureDetails.cpp
+++ b/Engine/Source/Editor/DetailCustomizations/Private/SceneCaptureDetails.cpp
@@ -58,59 +58,39 @@ void FSceneCaptureDetails::CustomizeDetails( IDetailLayoutBuilder& DetailLayout
 		}
 	}
 
-	// Hide certain categories
-	TArray<EShowFlagGroup> ShowFlagGroupsToHideForCaptures;
-	ShowFlagGroupsToHideForCaptures.Add(SFG_Hidden);
-	ShowFlagGroupsToHideForCaptures.Add(SFG_CollisionModes);
-	ShowFlagGroupsToHideForCaptures.Add(SFG_Visualize);
-	ShowFlagGroupsToHideForCaptures.Add(SFG_Developer);
-	ShowFlagGroupsToHideForCaptures.Add(SFG_PostProcess);
-
-	// Hide certain flags
-	TArray<FEngineShowFlags::EShowFlag> ShowFlagsToHideForCaptures;
+	// Show flags that should be exposed for Scene Captures
+	TArray<FEngineShowFlags::EShowFlag> ShowFlagsToAllowForCaptures;
 
 	// General
-	ShowFlagsToHideForCaptures.Add(FEngineShowFlags::EShowFlag::SF_AntiAliasing);
-	ShowFlagsToHideForCaptures.Add(FEngineShowFlags::EShowFlag::SF_Collision);
-	ShowFlagsToHideForCaptures.Add(FEngineShowFlags::EShowFlag::SF_Grid);
-	ShowFlagsToHideForCaptures.Add(FEngineShowFlags::EShowFlag::SF_Navigation);
+	ShowFlagsToAllowForCaptures.Add(FEngineShowFlags::EShowFlag::SF_Atmosphere);
+	ShowFlagsToAllowForCaptures.Add(FEngineShowFlags::EShowFlag::SF_BSP);
+	ShowFlagsToAllowForCaptures.Add(FEngineShowFlags::EShowFlag::SF_Decals);
+	ShowFlagsToAllowForCaptures.Add(FEngineShowFlags::EShowFlag::SF_Fog);
+	ShowFlagsToAllowForCaptures.Add(FEngineShowFlags::EShowFlag::SF_Landscape);
+	ShowFlagsToAllowForCaptures.Add(FEngineShowFlags::EShowFlag::SF_Particles);
+	ShowFlagsToAllowForCaptures.Add(FEngineShowFlags::EShowFlag::SF_SkeletalMeshes);
+	ShowFlagsToAllowForCaptures.Add(FEngineShowFlags::EShowFlag::SF_StaticMeshes);
+	ShowFlagsToAllowForCaptures.Add(FEngineShowFlags::EShowFlag::SF_Translucency);
 
 	// Advanced
-	ShowFlagsToHideForCaptures.Add(FEngineShowFlags::EShowFlag::SF_AudioRadius);
-	ShowFlagsToHideForCaptures.Add(FEngineShowFlags::EShowFlag::SF_StreamingBounds);
-	ShowFlagsToHideForCaptures.Add(FEngineShowFlags::EShowFlag::SF_Bounds);
-	ShowFlagsToHideForCaptures.Add(FEngineShowFlags::EShowFlag::SF_BSPSplit);
-	ShowFlagsToHideForCaptures.Add(FEngineShowFlags::EShowFlag::SF_CameraAspectRatioBars);
-	ShowFlagsToHideForCaptures.Add(FEngineShowFlags::EShowFlag::SF_CameraFrustums);
-	ShowFlagsToHideForCaptures.Add(FEngineShowFlags::EShowFlag::SF_CameraSafeFrames);
-	ShowFlagsToHideForCaptures.Add(FEngineShowFlags::EShowFlag::SF_Constraints);
-	ShowFlagsToHideForCaptures.Add(FEngineShowFlags::EShowFlag::SF_HighResScreenshotMask);
-	ShowFlagsToHideForCaptures.Add(FEngineShowFlags::EShowFlag::SF_LevelColoration);
-	ShowFlagsToHideForCaptures.Add(FEngineShowFlags::EShowFlag::SF_LOD);
-	ShowFlagsToHideForCaptures.Add(FEngineShowFlags::EShowFlag::SF_MeshEdges);
-	ShowFlagsToHideForCaptures.Add(FEngineShowFlags::EShowFlag::SF_ModeWidgets);
-	ShowFlagsToHideForCaptures.Add(FEngineShowFlags::EShowFlag::SF_SeparateTranslucency);
-	ShowFlagsToHideForCaptures.Add(FEngineShowFlags::EShowFlag::SF_ShadowFrustums);
-	ShowFlagsToHideForCaptures.Add(FEngineShowFlags::EShowFlag::SF_Splines);
-	ShowFlagsToHideForCaptures.Add(FEngineShowFlags::EShowFlag::SF_TemporalAA);
-	ShowFlagsToHideForCaptures.Add(FEngineShowFlags::EShowFlag::SF_VertexColors);
-	ShowFlagsToHideForCaptures.Add(FEngineShowFlags::EShowFlag::SF_VisualizeSenses);
-	ShowFlagsToHideForCaptures.Add(FEngineShowFlags::EShowFlag::SF_Volumes);
-	
-	// Post process
-	//ShowFlagsToHideForCaptures.Add(FEngineShowFlags::EShowFlag::SF_HMDDistortion);
-	//ShowFlagsToHideForCaptures.Add(FEngineShowFlags::EShowFlag::SF_MotionBlur);
-
-	// Lighting features
-	ShowFlagsToHideForCaptures.Add(FEngineShowFlags::EShowFlag::SF_SubsurfaceScattering);
+	ShowFlagsToAllowForCaptures.Add(FEngineShowFlags::EShowFlag::SF_DeferredLighting);
+	ShowFlagsToAllowForCaptures.Add(FEngineShowFlags::EShowFlag::SF_InstancedStaticMeshes);
+	ShowFlagsToAllowForCaptures.Add(FEngineShowFlags::EShowFlag::SF_Paper2DSprites);
+	ShowFlagsToAllowForCaptures.Add(FEngineShowFlags::EShowFlag::SF_TextRender);
 
 	// Lighting Components
-	ShowFlagsToHideForCaptures.Add(FEngineShowFlags::EShowFlag::SF_Diffuse);
-	ShowFlagsToHideForCaptures.Add(FEngineShowFlags::EShowFlag::SF_DirectLighting);
-	ShowFlagsToHideForCaptures.Add(FEngineShowFlags::EShowFlag::SF_DirectionalLights);
-	ShowFlagsToHideForCaptures.Add(FEngineShowFlags::EShowFlag::SF_PointLights);
-	ShowFlagsToHideForCaptures.Add(FEngineShowFlags::EShowFlag::SF_Specular);
-	ShowFlagsToHideForCaptures.Add(FEngineShowFlags::EShowFlag::SF_SpotLights);
+	ShowFlagsToAllowForCaptures.Add(FEngineShowFlags::EShowFlag::SF_AmbientOcclusion);
+	ShowFlagsToAllowForCaptures.Add(FEngineShowFlags::EShowFlag::SF_DynamicShadows);
+	ShowFlagsToAllowForCaptures.Add(FEngineShowFlags::EShowFlag::SF_SkyLighting);
+
+	// Lighting Features
+	ShowFlagsToAllowForCaptures.Add(FEngineShowFlags::EShowFlag::SF_AmbientCubemap);
+	ShowFlagsToAllowForCaptures.Add(FEngineShowFlags::EShowFlag::SF_DistanceFieldAO);
+	ShowFlagsToAllowForCaptures.Add(FEngineShowFlags::EShowFlag::SF_LightFunctions);
+	ShowFlagsToAllowForCaptures.Add(FEngineShowFlags::EShowFlag::SF_LightShafts);
+	ShowFlagsToAllowForCaptures.Add(FEngineShowFlags::EShowFlag::SF_ReflectionEnvironment);
+	ShowFlagsToAllowForCaptures.Add(FEngineShowFlags::EShowFlag::SF_ScreenSpaceReflections);
+	ShowFlagsToAllowForCaptures.Add(FEngineShowFlags::EShowFlag::SF_TexturedLightProfiles);
 	
 	// Create array of flag name strings for each group
 	TArray< TArray<FString> > ShowFlagsByGroup;
@@ -119,22 +99,17 @@ void FSceneCaptureDetails::CustomizeDetails( IDetailLayoutBuilder& DetailLayout
 		ShowFlagsByGroup.Add(TArray<FString>());
 	}
 
-	// Enumerate all the flags and put them in the correct group
-	bool FlagExists;
-	int32 FlagIndex = 0;
-	do 
+	// Add the show flags we want to expose to their group's array
+	for (FEngineShowFlags::EShowFlag AllowedFlag : ShowFlagsToAllowForCaptures)
 	{
-		TArray<FString> GroupStrings;
 		FString FlagName;
-		FlagName = FEngineShowFlags::FindNameByIndex(FlagIndex);
-		FlagExists = !FlagName.IsEmpty();
-		if (FlagExists)
+		FlagName = FEngineShowFlags::FindNameByIndex(AllowedFlag);
+		if (!FlagName.IsEmpty())
 		{
 			EShowFlagGroup Group = FEngineShowFlags::FindShowFlagGroup(*FlagName);
 			ShowFlagsByGroup[Group].Add(FlagName);
 		}
-		++FlagIndex;
-	} while (FlagExists);
+	}
 
 	// Sort the flags in their respective group alphabetically
 	for (TArray<FString>& ShowFlagGroup : ShowFlagsByGroup)
@@ -145,19 +120,8 @@ void FSceneCaptureDetails::CustomizeDetails( IDetailLayoutBuilder& DetailLayout
 	// Add each group
 	for (int32 GroupIndex = 0; GroupIndex < SFG_Max; ++GroupIndex)
 	{
-		bool bGroupHidden = false;
-
-		// Don't add show flag groups that we've specified to hide for Scene Captures
-		for (EShowFlagGroup HiddenGroup : ShowFlagGroupsToHideForCaptures)
-		{
-			if (EShowFlagGroup(GroupIndex) == HiddenGroup)
-			{
-				bGroupHidden = true;
-				break;
-			}
-		}
-
-		if (!bGroupHidden)
+		// Don't add a group if there are no flags allowed for it
+		if (ShowFlagsByGroup[GroupIndex].Num() >= 1)
 		{
 			FText GroupName;
 			FText GroupTooltip;
@@ -206,36 +170,20 @@ void FSceneCaptureDetails::CustomizeDetails( IDetailLayoutBuilder& DetailLayout
 				FText LocalizedText;
 				FEngineShowFlags::FindShowFlagDisplayName(FlagName, LocalizedText);
 
-				for (FEngineShowFlags::EShowFlag HiddenFlag : ShowFlagsToHideForCaptures)
-				{
-					// Don't add show flag that we've specified to hide for Scene Captures
-					// or that cannot be toggled in the editor
-					if (FEngineShowFlags::EShowFlag(FEngineShowFlags::FindIndexByName(*FlagName)) == HiddenFlag 
-						|| !FEngineShowFlags::CanBeToggledInEditor(*FlagName)
-						)
-					{
-						bFlagHidden = true;
-						break;
-					}
-				}
-
-				if (!bFlagHidden)
-				{
-					Group.AddWidgetRow()
-						.IsEnabled(true)
-						.NameContent()
-						[
-							SNew(STextBlock)
-							.Text(LocalizedText)
-						]
-					.ValueContent()
-						[
-							SNew(SCheckBox)
-							.OnCheckStateChanged(this, &FSceneCaptureDetails::OnShowFlagCheckStateChanged, FlagName)
-							.IsChecked(this, &FSceneCaptureDetails::OnGetDisplayCheckState, FlagName)
-						]
-					.FilterString(FlagName);
-				}
+				Group.AddWidgetRow()
+					.IsEnabled(true)
+					.NameContent()
+					[
+						SNew(STextBlock)
+						.Text(LocalizedText)
+					]
+				.ValueContent()
+					[
+						SNew(SCheckBox)
+						.OnCheckStateChanged(this, &FSceneCaptureDetails::OnShowFlagCheckStateChanged, FlagName)
+						.IsChecked(this, &FSceneCaptureDetails::OnGetDisplayCheckState, FlagName)
+					]
+				.FilterString(FlagName);
 			}
 		}
 	}
