Initial Mac-specific project settings (splash screen and icon for now), including minimal implementation of ImageWrapper for .icns files

diff --git a/Engine/Source/Developer/ExternalImagePicker/Private/SExternalImagePicker.cpp b/Engine/Source/Developer/ExternalImagePicker/Private/SExternalImagePicker.cpp
--- a/Engine/Source/Developer/ExternalImagePicker/Private/SExternalImagePicker.cpp
+++ b/Engine/Source/Developer/ExternalImagePicker/Private/SExternalImagePicker.cpp
@@ -217,11 +217,12 @@ TSharedPtr< FSlateDynamicImageBrush > SExternalImagePicker::LoadImageAsBrush( co
 	if( FFileHelper::LoadFileToArray( RawFileData, *ImagePath ) )
 	{
 		IImageWrapperModule& ImageWrapperModule = FModuleManager::LoadModuleChecked<IImageWrapperModule>( FName("ImageWrapper") );
-		IImageWrapperPtr ImageWrappers[3] =
+		IImageWrapperPtr ImageWrappers[4] =
 		{ 
 			ImageWrapperModule.CreateImageWrapper( EImageFormat::PNG ),
 			ImageWrapperModule.CreateImageWrapper( EImageFormat::BMP ),
 			ImageWrapperModule.CreateImageWrapper( EImageFormat::ICO ),
+			ImageWrapperModule.CreateImageWrapper( EImageFormat::ICNS ),
 		};
 
 		for( auto ImageWrapper : ImageWrappers )
diff --git a/Engine/Source/Developer/ImageWrapper/Private/IcnsImageWrapper.cpp b/Engine/Source/Developer/ImageWrapper/Private/IcnsImageWrapper.cpp
new file mode 100644
--- /dev/null
+++ b/Engine/Source/Developer/ImageWrapper/Private/IcnsImageWrapper.cpp
@@ -0,0 +1,49 @@
+// Copyright 1998-2014 Epic Games, Inc. All Rights Reserved.
+
+#include "ImageWrapperPrivatePCH.h"
+#include "BmpImageSupport.h"
+
+/**
+ * ICNS image wrapper class.
+ */
+FIcnsImageWrapper::FIcnsImageWrapper()
+	: FImageWrapperBase()
+{
+}
+
+void FIcnsImageWrapper::Compress(int32 Quality)
+{
+	checkf(false, TEXT("ICNS compression not supported"));
+}
+
+void FIcnsImageWrapper::Uncompress(const ERGBFormat::Type InFormat, const int32 InBitDepth)
+{
+#if PLATFORM_MAC
+	SCOPED_AUTORELEASE_POOL;
+
+	NSData* ImageData = [NSData dataWithBytesNoCopy:CompressedData.GetData() length:CompressedData.Num() freeWhenDone:NO];
+	NSImage* Image = [[NSImage alloc] initWithData:ImageData];
+	if (Image)
+	{
+		NSBitmapImageRep* Bitmap = [NSBitmapImageRep imageRepWithData:[Image TIFFRepresentation]];
+		if (Bitmap)
+		{
+			// @todo: Only 8-bit BGRA supported currently
+			check(InFormat == ERGBFormat::BGRA);
+			check(InBitDepth == 8);
+
+			RawData.Empty();
+			RawData.Append([Bitmap bitmapData], [Bitmap bytesPerPlane]);
+
+			RawFormat = Format = InFormat;
+			RawBitDepth = BitDepth = InBitDepth;
+
+			Width = [Bitmap pixelsWide];
+			Height = [Bitmap pixelsHigh];
+		}
+		[Image release];
+	}
+#else
+	checkf(false, TEXT("ICNS uncompressing not supported on this platform"));
+#endif
+}
diff --git a/Engine/Source/Developer/ImageWrapper/Private/IcnsImageWrapper.h b/Engine/Source/Developer/ImageWrapper/Private/IcnsImageWrapper.h
new file mode 100644
--- /dev/null
+++ b/Engine/Source/Developer/ImageWrapper/Private/IcnsImageWrapper.h
@@ -0,0 +1,24 @@
+// Copyright 1998-2014 Epic Games, Inc. All Rights Reserved.
+
+#pragma once
+
+/**
+ * ICNS implementation of the helper class
+ */
+class FIcnsImageWrapper
+	: public FImageWrapperBase
+{
+public:
+
+	/**
+	 * Default Constructor.
+	 */
+	FIcnsImageWrapper();
+
+public:
+
+	// FImageWrapper Interface
+
+	virtual void Compress( int32 Quality ) override;
+	virtual void Uncompress( const ERGBFormat::Type InFormat, int32 InBitDepth ) override;
+};
diff --git a/Engine/Source/Developer/ImageWrapper/Private/ImageWrapperModule.cpp b/Engine/Source/Developer/ImageWrapper/Private/ImageWrapperModule.cpp
--- a/Engine/Source/Developer/ImageWrapper/Private/ImageWrapperModule.cpp
+++ b/Engine/Source/Developer/ImageWrapper/Private/ImageWrapperModule.cpp
@@ -49,6 +49,9 @@ public:
 			ImageWrapper = new FExrImageWrapper();
 			break;
 #endif
+		case EImageFormat::ICNS:
+			ImageWrapper = new FIcnsImageWrapper();
+			break;
 		default:
 			break;
 		}
diff --git a/Engine/Source/Developer/ImageWrapper/Private/ImageWrapperPrivatePCH.h b/Engine/Source/Developer/ImageWrapper/Private/ImageWrapperPrivatePCH.h
--- a/Engine/Source/Developer/ImageWrapper/Private/ImageWrapperPrivatePCH.h
+++ b/Engine/Source/Developer/ImageWrapper/Private/ImageWrapperPrivatePCH.h
@@ -34,3 +34,4 @@ DECLARE_LOG_CATEGORY_EXTERN(LogImageWrapper, Log, All);
 
 #include "BmpImageWrapper.h"
 #include "IcoImageWrapper.h"
+#include "IcnsImageWrapper.h"
diff --git a/Engine/Source/Developer/ImageWrapper/Public/Interfaces/IImageWrapper.h b/Engine/Source/Developer/ImageWrapper/Public/Interfaces/IImageWrapper.h
--- a/Engine/Source/Developer/ImageWrapper/Public/Interfaces/IImageWrapper.h
+++ b/Engine/Source/Developer/ImageWrapper/Public/Interfaces/IImageWrapper.h
@@ -27,6 +27,9 @@ namespace EImageFormat
 
 		/** OpenEXR (HDR) image file format */
 		EXR,
+
+		/** Mac icon */
+		ICNS
 	};
 };
 
diff --git a/Engine/Source/Developer/Mac/MacTargetPlatform/Private/MacTargetPlatformModule.cpp b/Engine/Source/Developer/Mac/MacTargetPlatform/Private/MacTargetPlatformModule.cpp
--- a/Engine/Source/Developer/Mac/MacTargetPlatform/Private/MacTargetPlatformModule.cpp
+++ b/Engine/Source/Developer/Mac/MacTargetPlatform/Private/MacTargetPlatformModule.cpp
@@ -63,7 +63,7 @@ public:
 		{
 			SettingsModule->RegisterSettings("Project", "Platforms", "Mac",
 				LOCTEXT("TargetSettingsName", "Mac"),
-				LOCTEXT("TargetSettingsDescription", "Mac platform settings description text here"),
+				LOCTEXT("TargetSettingsDescription", "Settings and resources for Mac platform"),
 				TargetSettings
 			);
 		}
diff --git a/Engine/Source/Editor/DetailCustomizations/Private/DetailCustomizations.cpp b/Engine/Source/Editor/DetailCustomizations/Private/DetailCustomizations.cpp
--- a/Engine/Source/Editor/DetailCustomizations/Private/DetailCustomizations.cpp
+++ b/Engine/Source/Editor/DetailCustomizations/Private/DetailCustomizations.cpp
@@ -75,6 +75,7 @@
 #include "GeneralProjectSettingsDetails.h"
 #include "HardwareTargetingSettingsDetails.h"
 #include "WindowsTargetSettingsDetails.h"
+#include "MacTargetSettingsDetails.h"
 #include "MoviePlayerSettingsDetails.h"
 #include "SourceCodeAccessSettingsDetails.h"
 #include "ParticleSystemComponentDetails.h"
@@ -259,6 +260,7 @@ void FDetailCustomizationsModule::RegisterObjectCustomizations()
 	RegisterCustomClassLayout("InternationalizationSettingsModel", FOnGetDetailCustomizationInstance::CreateStatic(&FInternationalizationSettingsModelDetails::MakeInstance));
 	RegisterCustomClassLayout("InputSettings", FOnGetDetailCustomizationInstance::CreateStatic(&FInputSettingsDetails::MakeInstance));
 	RegisterCustomClassLayout("WindowsTargetSettings", FOnGetDetailCustomizationInstance::CreateStatic(&FWindowsTargetSettingsDetails::MakeInstance));
+	RegisterCustomClassLayout("MacTargetSettings", FOnGetDetailCustomizationInstance::CreateStatic(&FMacTargetSettingsDetails::MakeInstance));
 	RegisterCustomClassLayout("MoviePlayerSettings", FOnGetDetailCustomizationInstance::CreateStatic(&FMoviePlayerSettingsDetails::MakeInstance));
 
 	RegisterCustomClassLayout("SourceCodeAccessSettings", FOnGetDetailCustomizationInstance::CreateStatic(&FSourceCodeAccessSettingsDetails::MakeInstance));
diff --git a/Engine/Source/Editor/DetailCustomizations/Private/MacTargetSettingsDetails.cpp b/Engine/Source/Editor/DetailCustomizations/Private/MacTargetSettingsDetails.cpp
new file mode 100644
--- /dev/null
+++ b/Engine/Source/Editor/DetailCustomizations/Private/MacTargetSettingsDetails.cpp
@@ -0,0 +1,215 @@
+// Copyright 1998-2014 Epic Games, Inc. All Rights Reserved.
+
+#include "DetailCustomizationsPrivatePCH.h"
+#include "MacTargetSettingsDetails.h"
+#include "TargetPlatform.h"
+#include "SExternalImageReference.h"
+#include "IExternalImagePickerModule.h"
+#include "GameProjectGenerationModule.h"
+#include "ISourceControlModule.h"
+
+namespace MacTargetSettingsDetailsConstants
+{
+	/** The filename for the game splash screen */
+	const FString GameSplashFileName(TEXT("Splash/Splash.bmp"));
+
+	/** The filename for the editor splash screen */
+	const FString EditorSplashFileName(TEXT("Splash/EdSplash.bmp"));
+}
+
+#define LOCTEXT_NAMESPACE "MacTargetSettingsDetails"
+
+TSharedRef<IDetailCustomization> FMacTargetSettingsDetails::MakeInstance()
+{
+	return MakeShareable(new FMacTargetSettingsDetails);
+}
+
+namespace EImageScope
+{
+	enum Type
+	{
+		Engine,
+		GameOverride
+	};
+}
+
+/* Helper function used to generate filenames for splash screens */
+static FString GetSplashFilename(EImageScope::Type Scope, bool bIsEditorSplash)
+{
+	FString Filename;
+
+	if (Scope == EImageScope::Engine)
+	{
+		Filename = FPaths::EngineContentDir();
+	}
+	else
+	{
+		Filename = FPaths::GameContentDir();
+	}
+
+	if(bIsEditorSplash)
+	{
+		Filename /= MacTargetSettingsDetailsConstants::EditorSplashFileName;
+	}
+	else
+	{
+		Filename /= MacTargetSettingsDetailsConstants::GameSplashFileName;
+	}
+
+	Filename = FPaths::ConvertRelativePathToFull(Filename);
+
+	return Filename;
+}
+
+/* Helper function used to generate filenames for icons */
+static FString GetIconFilename(EImageScope::Type Scope)
+{
+	const FString& PlatformName = FModuleManager::GetModuleChecked<ITargetPlatformModule>("MacTargetPlatform").GetTargetPlatform()->PlatformName();
+
+	if (Scope == EImageScope::Engine)
+	{
+		FString Filename = FPaths::EngineDir() / FString(TEXT("Source/Runtime/Launch/Resources")) / PlatformName / FString("UE4.icns");
+		return FPaths::ConvertRelativePathToFull(Filename);
+	}
+	else
+	{
+		FString Filename = FPaths::GameDir() / TEXT("Build/Mac/Application.icns");
+		if(!FPaths::FileExists(Filename))
+		{
+			FString LegacyFilename = FPaths::GameSourceDir() / FString(FApp::GetGameName()) / FString(TEXT("Resources")) / PlatformName / FString(FApp::GetGameName()) + TEXT(".icns");
+			if(FPaths::FileExists(LegacyFilename))
+			{
+				Filename = LegacyFilename;
+			}
+		}
+		return FPaths::ConvertRelativePathToFull(Filename);
+	}
+}
+
+void FMacTargetSettingsDetails::CustomizeDetails( IDetailLayoutBuilder& DetailBuilder )
+{
+	// Add the splash image customization
+	IDetailCategoryBuilder& SplashCategoryBuilder = DetailBuilder.EditCategory(TEXT("Splash"));
+	FDetailWidgetRow& EditorSplashWidgetRow = SplashCategoryBuilder.AddCustomRow(TEXT("Editor Splash"));
+
+	const FText EditorSplashDesc(LOCTEXT("EditorSplashLabel", "Editor Splash"));
+	const FString EditorSplash_TargetImagePath = GetSplashFilename(EImageScope::GameOverride, true);
+	const FString EditorSplash_DefaultImagePath = GetSplashFilename(EImageScope::Engine, true);
+
+	EditorSplashWidgetRow
+	.NameContent()
+	[
+		SNew(SHorizontalBox)
+		+SHorizontalBox::Slot()
+		.Padding( FMargin( 0, 1, 0, 1 ) )
+		.FillWidth(1.0f)
+		[
+			SNew(STextBlock)
+			.Text(EditorSplashDesc)
+			.Font(DetailBuilder.GetDetailFont())
+		]
+	]
+	.ValueContent()
+	.MaxDesiredWidth(500.0f)
+	.MinDesiredWidth(100.0f)
+	[
+		SNew(SHorizontalBox)
+		+SHorizontalBox::Slot()
+		.FillWidth(1.0f)
+		.VAlign(VAlign_Center)
+		[
+			SNew(SExternalImageReference, EditorSplash_DefaultImagePath, EditorSplash_TargetImagePath)
+			.FileDescription(EditorSplashDesc)
+			.OnGetPickerPath(FOnGetPickerPath::CreateSP(this, &FMacTargetSettingsDetails::GetPickerPath))
+			.OnPostExternalImageCopy(FOnPostExternalImageCopy::CreateSP(this, &FMacTargetSettingsDetails::HandlePostExternalIconCopy))
+		]
+	];
+
+	FDetailWidgetRow& GameSplashWidgetRow = SplashCategoryBuilder.AddCustomRow(TEXT("Game Splash"));
+
+	const FText GameSplashDesc(LOCTEXT("GameSplashLabel", "Game Splash"));
+	const FString GameSplash_TargetImagePath = GetSplashFilename(EImageScope::GameOverride, false);
+	const FString GameSplash_DefaultImagePath = GetSplashFilename(EImageScope::Engine, false);
+
+	GameSplashWidgetRow
+	.NameContent()
+	[
+		SNew(SHorizontalBox)
+		+SHorizontalBox::Slot()
+		.Padding( FMargin( 0, 1, 0, 1 ) )
+		.FillWidth(1.0f)
+		[
+			SNew(STextBlock)
+			.Text(GameSplashDesc)
+			.Font(DetailBuilder.GetDetailFont())
+		]
+	]
+	.ValueContent()
+	.MaxDesiredWidth(500.0f)
+	.MinDesiredWidth(100.0f)
+	[
+		SNew(SHorizontalBox)
+		+SHorizontalBox::Slot()
+		.FillWidth(1.0f)
+		.VAlign(VAlign_Center)
+		[
+			SNew(SExternalImageReference, GameSplash_DefaultImagePath, GameSplash_TargetImagePath)
+			.FileDescription(GameSplashDesc)
+			.OnGetPickerPath(FOnGetPickerPath::CreateSP(this, &FMacTargetSettingsDetails::GetPickerPath))
+			.OnPostExternalImageCopy(FOnPostExternalImageCopy::CreateSP(this, &FMacTargetSettingsDetails::HandlePostExternalIconCopy))
+		]
+	];
+
+	IDetailCategoryBuilder& IconsCategoryBuilder = DetailBuilder.EditCategory(TEXT("Icon"));	
+	FDetailWidgetRow& GameIconWidgetRow = IconsCategoryBuilder.AddCustomRow(TEXT("Game Icon"));
+	GameIconWidgetRow
+	.NameContent()
+	[
+		SNew(SHorizontalBox)
+		+SHorizontalBox::Slot()
+		.Padding( FMargin( 0, 1, 0, 1 ) )
+		.FillWidth(1.0f)
+		[
+			SNew(STextBlock)
+			.Text(LOCTEXT("GameIconLabel", "Game Icon"))
+			.Font(DetailBuilder.GetDetailFont())
+		]
+	]
+	.ValueContent()
+	.MaxDesiredWidth(500.0f)
+	.MinDesiredWidth(100.0f)
+	[
+		SNew(SHorizontalBox)
+		+SHorizontalBox::Slot()
+		.FillWidth(1.0f)
+		.VAlign(VAlign_Center)
+		[
+			SNew(SExternalImageReference, GetIconFilename(EImageScope::Engine), GetIconFilename(EImageScope::GameOverride))
+			.FileDescription(GameSplashDesc)
+			.OnPreExternalImageCopy(FOnPreExternalImageCopy::CreateSP(this, &FMacTargetSettingsDetails::HandlePreExternalIconCopy))
+			.OnGetPickerPath(FOnGetPickerPath::CreateSP(this, &FMacTargetSettingsDetails::GetPickerPath))
+			.OnPostExternalImageCopy(FOnPostExternalImageCopy::CreateSP(this, &FMacTargetSettingsDetails::HandlePostExternalIconCopy))
+		]
+	];
+}
+
+
+bool FMacTargetSettingsDetails::HandlePreExternalIconCopy(const FString& InChosenImage)
+{
+	return true;
+}
+
+
+FString FMacTargetSettingsDetails::GetPickerPath()
+{
+	return FEditorDirectories::Get().GetLastDirectory(ELastDirectory::GENERIC_OPEN);
+}
+
+
+bool FMacTargetSettingsDetails::HandlePostExternalIconCopy(const FString& InChosenImage)
+{
+	FEditorDirectories::Get().SetLastDirectory(ELastDirectory::GENERIC_OPEN, FPaths::GetPath(InChosenImage));
+	return true;
+}
+
+#undef LOCTEXT_NAMESPACE
diff --git a/Engine/Source/Editor/DetailCustomizations/Private/MacTargetSettingsDetails.h b/Engine/Source/Editor/DetailCustomizations/Private/MacTargetSettingsDetails.h
new file mode 100644
--- /dev/null
+++ b/Engine/Source/Editor/DetailCustomizations/Private/MacTargetSettingsDetails.h
@@ -0,0 +1,26 @@
+// Copyright 1998-2014 Epic Games, Inc. All Rights Reserved.
+
+#pragma once
+
+/**
+ * Manages the Transform section of a details view                    
+ */
+class FMacTargetSettingsDetails : public IDetailCustomization
+{
+public:
+	/** Makes a new instance of this detail layout class for a specific detail view requesting it */
+	static TSharedRef<IDetailCustomization> MakeInstance();
+
+	/** IDetailCustomization interface */
+	virtual void CustomizeDetails( IDetailLayoutBuilder& DetailBuilder ) override;
+
+private:
+	/** Delegate handler for before an icon is copied */
+	bool HandlePreExternalIconCopy(const FString& InChosenImage);
+
+	/** Delegate handler to get the path to start picking from */
+	FString GetPickerPath();
+
+	/** Delegate handler to set the path to start picking from */
+	bool HandlePostExternalIconCopy(const FString& InChosenImage);
+};
\ No newline at end of file
diff --git a/Engine/Source/Programs/AutomationTool/Mac/MacPlatform.Automation.cs b/Engine/Source/Programs/AutomationTool/Mac/MacPlatform.Automation.cs
--- a/Engine/Source/Programs/AutomationTool/Mac/MacPlatform.Automation.cs
+++ b/Engine/Source/Programs/AutomationTool/Mac/MacPlatform.Automation.cs
@@ -36,13 +36,21 @@ public class MacPlatform : Platform
 		List<string> Exes = GetExecutableNames(SC);
 		foreach (var Exe in Exes)
 		{
+			string AppBundlePath = "";
 			if (Exe.StartsWith(CombinePaths(SC.RuntimeProjectRootDir, "Binaries", SC.PlatformDir)))
 			{
-				StageAppBundle(SC, CombinePaths(SC.ProjectRoot, "Binaries", SC.PlatformDir, Path.GetFileNameWithoutExtension(Exe) + ".app"), CombinePaths(SC.ShortProjectName, "Binaries", SC.PlatformDir, Path.GetFileNameWithoutExtension(Exe) + ".app"));
+				AppBundlePath = CombinePaths(SC.ShortProjectName, "Binaries", SC.PlatformDir, Path.GetFileNameWithoutExtension(Exe) + ".app");
+				StageAppBundle(SC, CombinePaths(SC.ProjectRoot, "Binaries", SC.PlatformDir, Path.GetFileNameWithoutExtension(Exe) + ".app"), AppBundlePath);
 			}
 			else if (Exe.StartsWith(CombinePaths(SC.RuntimeRootDir, "Engine/Binaries", SC.PlatformDir)))
 			{
-				StageAppBundle(SC, CombinePaths(SC.LocalRoot, "Engine/Binaries", SC.PlatformDir, Path.GetFileNameWithoutExtension(Exe) + ".app"), CombinePaths("Engine/Binaries", SC.PlatformDir, Path.GetFileNameWithoutExtension(Exe) + ".app"));
+				AppBundlePath = CombinePaths("Engine/Binaries", SC.PlatformDir, Path.GetFileNameWithoutExtension(Exe) + ".app");
+				StageAppBundle(SC, CombinePaths(SC.LocalRoot, "Engine/Binaries", SC.PlatformDir, Path.GetFileNameWithoutExtension(Exe) + ".app"), AppBundlePath);
+			}
+
+			if (!string.IsNullOrEmpty(AppBundlePath))
+			{
+				SC.StageFiles(StagedFileType.NonUFS, CombinePaths(SC.ProjectRoot, "Build/Mac"), "Application.icns", false, null, CombinePaths(AppBundlePath, "Contents/Resources"), true);
 			}
 		}
 
@@ -97,7 +105,7 @@ public class MacPlatform : Platform
 			}
 		}
 
-		// Update executable name and entry in Info.plist
+		// Update executable name, icon and entry in Info.plist
 		string UE4GamePath = CombinePaths(SC.ArchiveDirectory, SC.ShortProjectName + ".app", "Contents", "MacOS", ExeName);
 		if (ExeName != SC.ShortProjectName && File.Exists(UE4GamePath))
 		{
@@ -105,9 +113,23 @@ public class MacPlatform : Platform
 			File.Delete(GameExePath);
 			File.Move(UE4GamePath, GameExePath);
 
+			string DefaultIconPath = CombinePaths(SC.ArchiveDirectory, SC.ShortProjectName + ".app", "Contents", "Resources", "UE4.icns");
+			string CustomIconSrcPath = CombinePaths(SC.ArchiveDirectory, SC.ShortProjectName + ".app", "Contents", "Resources", "Application.icns");
+			string CustomIconDestPath = CombinePaths(SC.ArchiveDirectory, SC.ShortProjectName + ".app", "Contents", "Resources", SC.ShortProjectName + ".icns");
+			if (File.Exists(CustomIconSrcPath))
+			{
+				File.Delete(DefaultIconPath);
+				File.Move(CustomIconSrcPath, CustomIconDestPath);
+			}
+			else
+			{
+				File.Move(DefaultIconPath, CustomIconDestPath);
+			}
+
 			string InfoPlistPath = CombinePaths(SC.ArchiveDirectory, SC.ShortProjectName + ".app", "Contents", "Info.plist");
 			string InfoPlistContents = File.ReadAllText(InfoPlistPath);
 			InfoPlistContents = InfoPlistContents.Replace(ExeName, SC.ShortProjectName);
+			InfoPlistContents = InfoPlistContents.Replace("<string>UE4</string>", "<string>" + SC.ShortProjectName + "</string>");
 			File.Delete(InfoPlistPath);
 			File.WriteAllText(InfoPlistPath, InfoPlistContents);
 		}
diff --git a/Engine/Source/Programs/UnrealBuildTool/Mac/MacToolChain.cs b/Engine/Source/Programs/UnrealBuildTool/Mac/MacToolChain.cs
--- a/Engine/Source/Programs/UnrealBuildTool/Mac/MacToolChain.cs
+++ b/Engine/Source/Programs/UnrealBuildTool/Mac/MacToolChain.cs
@@ -942,12 +942,14 @@ namespace UnrealBuildTool
 
 					string UProjectFilePath = UProjectInfo.GetProjectFilePath(GameName);
 					string CustomResourcesPath = "";
+					string CustomBuildPath = "";
 					if (string.IsNullOrEmpty(UProjectFilePath))
 					{
 						string[] TargetFiles = Directory.GetFiles(Directory.GetCurrentDirectory(), GameName + ".Target.cs", SearchOption.AllDirectories);
 						if (TargetFiles.Length == 1)
 						{
 							CustomResourcesPath = Path.GetDirectoryName(TargetFiles[0]) + "/Resources/Mac";
+							CustomBuildPath = Path.GetDirectoryName(TargetFiles[0]) + "../Build/Mac";
 						}
 						else
 						{
@@ -957,15 +959,32 @@ namespace UnrealBuildTool
 					else
 					{
 						CustomResourcesPath = Path.GetDirectoryName(UProjectFilePath) + "/Source/" + GameName + "/Resources/Mac";
+						CustomBuildPath = Path.GetDirectoryName(UProjectFilePath) + "/Build/Mac";
 					}
 
-					CustomResourcesPath = ConvertPath(CustomResourcesPath);
+					CustomResourcesPath = ConvertPath(Path.GetFullPath(CustomResourcesPath));
+					CustomBuildPath = ConvertPath(Path.GetFullPath(CustomBuildPath));
+
+					bool bBuildingEditor = GameName.EndsWith("Editor");
 
 					// Copy resources
-					string CustomIcon = CustomResourcesPath + "/" + GameName + ".icns";
-					if (!File.Exists(CustomIcon))
+					string DefaultIcon = EngineSourcePath + "/Runtime/Launch/Resources/Mac/" + IconName + ".icns";
+					string CustomIcon = "";
+					if (bBuildingEditor)
+					{
+						CustomIcon = DefaultIcon;
+					}
+					else
 					{
-						CustomIcon = EngineSourcePath + "/Runtime/Launch/Resources/Mac/" + IconName + ".icns";
+						CustomIcon = CustomBuildPath + "/Application.icns";
+						if (!File.Exists(CustomIcon))
+						{
+							CustomIcon = CustomResourcesPath + "/" + GameName + ".icns";
+							if (!File.Exists(CustomIcon))
+							{
+								CustomIcon = DefaultIcon;
+							}
+						}
 					}
 					AppendMacLine(FinalizeAppBundleScript, "cp -f \"{0}\" \"{2}.app/Contents/Resources/{1}.icns\"", CustomIcon, IconName, ExeName);
 
@@ -977,7 +996,7 @@ namespace UnrealBuildTool
 					string InfoPlistFile = CustomResourcesPath + "/Info.plist";
 					if (!File.Exists(InfoPlistFile))
 					{
-						InfoPlistFile = EngineSourcePath + "/Runtime/Launch/Resources/Mac/" + (GameName.EndsWith("Editor") ? "Info-Editor.plist" : "Info.plist");
+						InfoPlistFile = EngineSourcePath + "/Runtime/Launch/Resources/Mac/" + (bBuildingEditor ? "Info-Editor.plist" : "Info.plist");
 					}
 					AppendMacLine(FinalizeAppBundleScript, "cp -f \"{0}\" \"{1}.app/Contents/Info.plist\"", InfoPlistFile, ExeName);
 
