Media: Fixed incorrect handling of relative movie file paths

diff --git a/Engine/Source/Editor/DetailCustomizations/Private/MediaPlayerCustomization.cpp b/Engine/Source/Editor/DetailCustomizations/Private/MediaPlayerCustomization.cpp
--- a/Engine/Source/Editor/DetailCustomizations/Private/MediaPlayerCustomization.cpp
+++ b/Engine/Source/Editor/DetailCustomizations/Private/MediaPlayerCustomization.cpp
@@ -388,16 +388,23 @@ FString FMediaPlayerCustomization::HandleUrlPickerFileTypeFilter() const
 
 void FMediaPlayerCustomization::HandleUrlPickerPathPicked( const FString& PickedPath )
 {
-	const FString FullGameContentDir = FPaths::ConvertRelativePathToFull(FPaths::GameContentDir());
-	FString FullUrl = FPaths::ConvertRelativePathToFull(FPaths::IsRelative(PickedPath) ? FPaths::GameContentDir() / PickedPath : PickedPath);
-
-	if (FullUrl.StartsWith(FullGameContentDir))
+	if (PickedPath.StartsWith(TEXT("./")))
 	{
-		FPaths::MakePathRelativeTo(FullUrl, *FullGameContentDir);
-		FullUrl = FString(TEXT("./")) + FullUrl;
+		UrlProperty->SetValue(PickedPath);
 	}
+	else
+	{	
+		FString FullUrl = FPaths::ConvertRelativePathToFull(PickedPath);
+		const FString FullGameContentDir = FPaths::ConvertRelativePathToFull(FPaths::GameContentDir());
 
-	UrlProperty->SetValue(FullUrl);
+		if (FullUrl.StartsWith(FullGameContentDir))
+		{
+			FPaths::MakePathRelativeTo(FullUrl, *FullGameContentDir);
+			FullUrl = FString(TEXT("./")) + FullUrl;
+		}
+
+		UrlProperty->SetValue(FullUrl);
+	}
 }
 
 
