- Fix Unity compile error

diff --git a/Engine/Source/Editor/DetailCustomizations/Private/MacTargetSettingsDetails.cpp b/Engine/Source/Editor/DetailCustomizations/Private/MacTargetSettingsDetails.cpp
--- a/Engine/Source/Editor/DetailCustomizations/Private/MacTargetSettingsDetails.cpp
+++ b/Engine/Source/Editor/DetailCustomizations/Private/MacTargetSettingsDetails.cpp
@@ -24,7 +24,7 @@ TSharedRef<IDetailCustomization> FMacTargetSettingsDetails::MakeInstance()
 	return MakeShareable(new FMacTargetSettingsDetails);
 }
 
-namespace EImageScope
+namespace EMacImageScope
 {
 	enum Type
 	{
@@ -34,11 +34,11 @@ namespace EImageScope
 }
 
 /* Helper function used to generate filenames for splash screens */
-static FString GetSplashFilename(EImageScope::Type Scope, bool bIsEditorSplash)
+static FString GetSplashFilename(EMacImageScope::Type Scope, bool bIsEditorSplash)
 {
 	FString Filename;
 
-	if (Scope == EImageScope::Engine)
+	if (Scope == EMacImageScope::Engine)
 	{
 		Filename = FPaths::EngineContentDir();
 	}
@@ -62,11 +62,11 @@ static FString GetSplashFilename(EImageScope::Type Scope, bool bIsEditorSplash)
 }
 
 /* Helper function used to generate filenames for icons */
-static FString GetIconFilename(EImageScope::Type Scope)
+static FString GetIconFilename(EMacImageScope::Type Scope)
 {
 	const FString& PlatformName = FModuleManager::GetModuleChecked<ITargetPlatformModule>("MacTargetPlatform").GetTargetPlatform()->PlatformName();
 
-	if (Scope == EImageScope::Engine)
+	if (Scope == EMacImageScope::Engine)
 	{
 		FString Filename = FPaths::EngineDir() / FString(TEXT("Source/Runtime/Launch/Resources")) / PlatformName / FString("UE4.icns");
 		return FPaths::ConvertRelativePathToFull(Filename);
@@ -93,8 +93,8 @@ void FMacTargetSettingsDetails::CustomizeDetails( IDetailLayoutBuilder& DetailBu
 	FDetailWidgetRow& EditorSplashWidgetRow = SplashCategoryBuilder.AddCustomRow(TEXT("Editor Splash"));
 
 	const FText EditorSplashDesc(LOCTEXT("EditorSplashLabel", "Editor Splash"));
-	const FString EditorSplash_TargetImagePath = GetSplashFilename(EImageScope::GameOverride, true);
-	const FString EditorSplash_DefaultImagePath = GetSplashFilename(EImageScope::Engine, true);
+	const FString EditorSplash_TargetImagePath = GetSplashFilename(EMacImageScope::GameOverride, true);
+	const FString EditorSplash_DefaultImagePath = GetSplashFilename(EMacImageScope::Engine, true);
 
 	EditorSplashWidgetRow
 	.NameContent()
@@ -128,8 +128,8 @@ void FMacTargetSettingsDetails::CustomizeDetails( IDetailLayoutBuilder& DetailBu
 	FDetailWidgetRow& GameSplashWidgetRow = SplashCategoryBuilder.AddCustomRow(TEXT("Game Splash"));
 
 	const FText GameSplashDesc(LOCTEXT("GameSplashLabel", "Game Splash"));
-	const FString GameSplash_TargetImagePath = GetSplashFilename(EImageScope::GameOverride, false);
-	const FString GameSplash_DefaultImagePath = GetSplashFilename(EImageScope::Engine, false);
+	const FString GameSplash_TargetImagePath = GetSplashFilename(EMacImageScope::GameOverride, false);
+	const FString GameSplash_DefaultImagePath = GetSplashFilename(EMacImageScope::Engine, false);
 
 	GameSplashWidgetRow
 	.NameContent()
@@ -184,7 +184,7 @@ void FMacTargetSettingsDetails::CustomizeDetails( IDetailLayoutBuilder& DetailBu
 		.FillWidth(1.0f)
 		.VAlign(VAlign_Center)
 		[
-			SNew(SExternalImageReference, GetIconFilename(EImageScope::Engine), GetIconFilename(EImageScope::GameOverride))
+			SNew(SExternalImageReference, GetIconFilename(EMacImageScope::Engine), GetIconFilename(EMacImageScope::GameOverride))
 			.FileDescription(GameSplashDesc)
 			.OnPreExternalImageCopy(FOnPreExternalImageCopy::CreateSP(this, &FMacTargetSettingsDetails::HandlePreExternalIconCopy))
 			.OnGetPickerPath(FOnGetPickerPath::CreateSP(this, &FMacTargetSettingsDetails::GetPickerPath))
