Scene Capture details panel: put engine show flags in the advanced section

diff --git a/Engine/Source/Editor/DetailCustomizations/Private/SceneCaptureDetails.cpp b/Engine/Source/Editor/DetailCustomizations/Private/SceneCaptureDetails.cpp
--- a/Engine/Source/Editor/DetailCustomizations/Private/SceneCaptureDetails.cpp
+++ b/Engine/Source/Editor/DetailCustomizations/Private/SceneCaptureDetails.cpp
@@ -161,7 +161,7 @@ void FSceneCaptureDetails::CustomizeDetails( IDetailLayoutBuilder& DetailLayout
 			}
 
 			FName GroupFName = FName(*(GroupName.ToString()));
-			IDetailGroup& Group = SceneCaptureCategoryBuilder.AddGroup(GroupFName, GroupName.ToString());
+			IDetailGroup& Group = SceneCaptureCategoryBuilder.AddGroup(GroupFName, GroupName.ToString(), true);
 
 			// Add each show flag for this group
 			for (FString& FlagName : ShowFlagsByGroup[GroupIndex])
