Fixed bad state in FBX import window that would show incorrect mesh UI when moving between import modes.

diff --git a/Engine/Source/Editor/DetailCustomizations/Private/FbxImportUIDetails.cpp b/Engine/Source/Editor/DetailCustomizations/Private/FbxImportUIDetails.cpp
--- a/Engine/Source/Editor/DetailCustomizations/Private/FbxImportUIDetails.cpp
+++ b/Engine/Source/Editor/DetailCustomizations/Private/FbxImportUIDetails.cpp
@@ -309,7 +309,15 @@ void FFbxImportUIDetails::ImportMeshToggleChanged()
 {
 	if(CachedDetailBuilder)
 	{
-		ImportUI->MeshTypeToImport = ImportUI->bImportMesh ? FBXIT_SkeletalMesh: FBXIT_Animation;
+		if(ImportUI->bImportMesh)
+		{
+			ImportUI->MeshTypeToImport = ImportUI->bImportAsSkeletal ? FBXIT_SkeletalMesh : FBXIT_StaticMesh;
+		}
+		else
+		{
+			ImportUI->MeshTypeToImport = FBXIT_Animation;
+		}
+
 		CachedDetailBuilder->ForceRefreshDetails();
 	}
 }
