[Jira UETOOL-102] The curve editor now allows metadata to control labeling of each axis.  THe curve editor for user interface DPI scaling now labels the axises Resolution and Scale.

diff --git a/Engine/Source/Editor/DetailCustomizations/Private/CurveStructCustomization.cpp b/Engine/Source/Editor/DetailCustomizations/Private/CurveStructCustomization.cpp
--- a/Engine/Source/Editor/DetailCustomizations/Private/CurveStructCustomization.cpp
+++ b/Engine/Source/Editor/DetailCustomizations/Private/CurveStructCustomization.cpp
@@ -43,6 +43,21 @@ void FCurveStructCustomization::CustomizeHeader( TSharedRef<IPropertyHandle> InS
 
 	if (StructPtrs.Num() == 1)
 	{
+		static const FName XAxisName(TEXT("XAxisName"));
+		static const FName YAxisName(TEXT("YAxisName"));
+
+		TOptional<FString> XAxisString;
+		if ( InStructPropertyHandle->HasMetaData(XAxisName) )
+		{
+			XAxisString = InStructPropertyHandle->GetMetaData(XAxisName);
+		}
+
+		TOptional<FString> YAxisString;
+		if ( InStructPropertyHandle->HasMetaData(YAxisName) )
+		{
+			YAxisString = InStructPropertyHandle->GetMetaData(YAxisName);
+		}
+
 		RuntimeCurve = reinterpret_cast<FRuntimeFloatCurve*>(StructPtrs[0]);
 
 		if (OuterObjects.Num() == 1)
@@ -68,6 +83,8 @@ void FCurveStructCustomization::CustomizeHeader( TSharedRef<IPropertyHandle> InS
 					.ViewMaxInput(this, &FCurveStructCustomization::GetViewMaxInput)
 					.TimelineLength(this, &FCurveStructCustomization::GetTimelineLength)
 					.OnSetInputViewRange(this, &FCurveStructCustomization::SetInputViewRange)
+					.XAxisName(XAxisString)
+					.YAxisName(YAxisString)
 					.HideUI(false)
 					.DesiredSize(FVector2D(300, 150))
 				]
diff --git a/Engine/Source/Editor/UnrealEd/Private/SCurveEditor.cpp b/Engine/Source/Editor/UnrealEd/Private/SCurveEditor.cpp
--- a/Engine/Source/Editor/UnrealEd/Private/SCurveEditor.cpp
+++ b/Engine/Source/Editor/UnrealEd/Private/SCurveEditor.cpp
@@ -229,7 +229,7 @@ void SCurveEditor::Construct(const FArguments& InArgs)
 							.Label()
 							[
 								SNew(STextBlock)
-								.Text(LOCTEXT("Time", "Time"))
+								.Text(InArgs._XAxisName.IsSet() ? FText::FromString(InArgs._XAxisName.GetValue()) : LOCTEXT("Time", "Time"))
 							]
 						]
 					]
@@ -261,7 +261,7 @@ void SCurveEditor::Construct(const FArguments& InArgs)
 							.Label()
 							[
 								SNew(STextBlock)
-								.Text(LOCTEXT("Value", "Value"))
+								.Text(InArgs._YAxisName.IsSet() ? FText::FromString(InArgs._YAxisName.GetValue()) : LOCTEXT("Value", "Value"))
 							]
 						]
 					]
diff --git a/Engine/Source/Editor/UnrealEd/Public/SCurveEditor.h b/Engine/Source/Editor/UnrealEd/Public/SCurveEditor.h
--- a/Engine/Source/Editor/UnrealEd/Public/SCurveEditor.h
+++ b/Engine/Source/Editor/UnrealEd/Public/SCurveEditor.h
@@ -141,6 +141,8 @@ public:
 		, _ZoomToFitVertical(true)
 		, _ZoomToFitHorizontal(true)
 		, _ShowZoomButtons(true)
+		, _XAxisName()
+		, _YAxisName()
 		{}
 
 		SLATE_ATTRIBUTE( float, ViewMinInput )
@@ -161,6 +163,8 @@ public:
 		SLATE_ARGUMENT( bool, ZoomToFitVertical )
 		SLATE_ARGUMENT( bool, ZoomToFitHorizontal )
 		SLATE_ARGUMENT( bool, ShowZoomButtons )
+		SLATE_ARGUMENT( TOptional<FString>, XAxisName )
+		SLATE_ARGUMENT( TOptional<FString>, YAxisName )
 		SLATE_EVENT( FOnSetInputViewRange, OnSetInputViewRange )
 		SLATE_EVENT( FOnSetOutputViewRange, OnSetOutputViewRange )
 		SLATE_EVENT( FSimpleDelegate, OnCreateAsset )
diff --git a/Engine/Source/Runtime/Engine/Classes/Engine/UserInterfaceSettings.h b/Engine/Source/Runtime/Engine/Classes/Engine/UserInterfaceSettings.h
--- a/Engine/Source/Runtime/Engine/Classes/Engine/UserInterfaceSettings.h
+++ b/Engine/Source/Runtime/Engine/Classes/Engine/UserInterfaceSettings.h
@@ -43,7 +43,9 @@ public:
 
 	UPROPERTY(config, EditAnywhere, Category="DPI Scaling", meta=(
 		DisplayName="DPI Curve",
-		ToolTip="Controls how the UI is scaled at different resolutions based on the DPI Scale Rule" ))
+		ToolTip="Controls how the UI is scaled at different resolutions based on the DPI Scale Rule",
+		XAxisName="Resolution",
+		YAxisName="Scale"))
 	FRuntimeFloatCurve UIScaleCurve;
 
 public:
