Removed clamping of RGB to 1.0 in the Raw Distribution Vector customization, to allow specification of super saturated colors.

diff --git a/Engine/Source/Editor/DetailCustomizations/Private/RawDistributionVectorStructCustomization.cpp b/Engine/Source/Editor/DetailCustomizations/Private/RawDistributionVectorStructCustomization.cpp
--- a/Engine/Source/Editor/DetailCustomizations/Private/RawDistributionVectorStructCustomization.cpp
+++ b/Engine/Source/Editor/DetailCustomizations/Private/RawDistributionVectorStructCustomization.cpp
@@ -342,8 +342,6 @@ void FReplaceVectorWithLinearColorBuilder::AddColorChildProperty(const TSharedPt
 			.OnBeginSliderMovement(this, &FReplaceVectorWithLinearColorBuilder::OnBeginSliderMovement)
 			.OnEndSliderMovement(this, &FReplaceVectorWithLinearColorBuilder::OnEndSliderMovement)
 			.AllowSpin(true)
-			.MinValue(0.0f)
-			.MaxValue(1.0f)
 			.MinSliderValue(0.0f)
 			.MaxSliderValue(1.0f)
 		];
