added showflags to scenecaptureactor

diff --git a/Engine/Source/Editor/DetailCustomizations/Private/SceneCaptureDetails.cpp b/Engine/Source/Editor/DetailCustomizations/Private/SceneCaptureDetails.cpp
--- a/Engine/Source/Editor/DetailCustomizations/Private/SceneCaptureDetails.cpp
+++ b/Engine/Source/Editor/DetailCustomizations/Private/SceneCaptureDetails.cpp
@@ -71,19 +71,13 @@ void FSceneCaptureDetails::CustomizeDetails( IDetailLayoutBuilder& DetailLayout
 	ShowFlagsToAllowForCaptures.Add(FEngineShowFlags::EShowFlag::SF_SkeletalMeshes);
 	ShowFlagsToAllowForCaptures.Add(FEngineShowFlags::EShowFlag::SF_StaticMeshes);
 	ShowFlagsToAllowForCaptures.Add(FEngineShowFlags::EShowFlag::SF_Translucency);
-
-	// Advanced
 	ShowFlagsToAllowForCaptures.Add(FEngineShowFlags::EShowFlag::SF_DeferredLighting);
 	ShowFlagsToAllowForCaptures.Add(FEngineShowFlags::EShowFlag::SF_InstancedStaticMeshes);
 	ShowFlagsToAllowForCaptures.Add(FEngineShowFlags::EShowFlag::SF_Paper2DSprites);
 	ShowFlagsToAllowForCaptures.Add(FEngineShowFlags::EShowFlag::SF_TextRender);
-
-	// Lighting Components
 	ShowFlagsToAllowForCaptures.Add(FEngineShowFlags::EShowFlag::SF_AmbientOcclusion);
 	ShowFlagsToAllowForCaptures.Add(FEngineShowFlags::EShowFlag::SF_DynamicShadows);
 	ShowFlagsToAllowForCaptures.Add(FEngineShowFlags::EShowFlag::SF_SkyLighting);
-
-	// Lighting Features
 	ShowFlagsToAllowForCaptures.Add(FEngineShowFlags::EShowFlag::SF_AmbientCubemap);
 	ShowFlagsToAllowForCaptures.Add(FEngineShowFlags::EShowFlag::SF_DistanceFieldAO);
 	ShowFlagsToAllowForCaptures.Add(FEngineShowFlags::EShowFlag::SF_LightFunctions);
@@ -91,6 +85,11 @@ void FSceneCaptureDetails::CustomizeDetails( IDetailLayoutBuilder& DetailLayout
 	ShowFlagsToAllowForCaptures.Add(FEngineShowFlags::EShowFlag::SF_ReflectionEnvironment);
 	ShowFlagsToAllowForCaptures.Add(FEngineShowFlags::EShowFlag::SF_ScreenSpaceReflections);
 	ShowFlagsToAllowForCaptures.Add(FEngineShowFlags::EShowFlag::SF_TexturedLightProfiles);
+	ShowFlagsToAllowForCaptures.Add(FEngineShowFlags::EShowFlag::SF_AntiAliasing);
+	ShowFlagsToAllowForCaptures.Add(FEngineShowFlags::EShowFlag::SF_TemporalAA);
+	ShowFlagsToAllowForCaptures.Add(FEngineShowFlags::EShowFlag::SF_MotionBlur);
+	ShowFlagsToAllowForCaptures.Add(FEngineShowFlags::EShowFlag::SF_Bloom);
+	ShowFlagsToAllowForCaptures.Add(FEngineShowFlags::EShowFlag::SF_EyeAdaptation);
 	
 	// Create array of flag name strings for each group
 	TArray< TArray<FString> > ShowFlagsByGroup;
