UE-4195 Make more use of GET_MEMBER_NAME_CHECKED when getting/checking properties by name
  - Remove disabled/legacy 'bake and prune' functions/vars
UE-2968 Make FBodyInstance::bUseAsyncScene protected (add SetUseAsyncScene)

diff --git a/Engine/Plugins/Editor/MacGraphicsSwitching/Source/MacGraphicsSwitching/Private/MacGraphicsSwitchingSettingsDetails.cpp b/Engine/Plugins/Editor/MacGraphicsSwitching/Source/MacGraphicsSwitching/Private/MacGraphicsSwitchingSettingsDetails.cpp
--- a/Engine/Plugins/Editor/MacGraphicsSwitching/Source/MacGraphicsSwitching/Private/MacGraphicsSwitchingSettingsDetails.cpp
+++ b/Engine/Plugins/Editor/MacGraphicsSwitching/Source/MacGraphicsSwitching/Private/MacGraphicsSwitchingSettingsDetails.cpp
@@ -13,24 +13,24 @@ TSharedRef<IDetailCustomization> FMacGraphicsSwitchingSettingsDetails::MakeInsta
 
 void FMacGraphicsSwitchingSettingsDetails::CustomizeDetails( IDetailLayoutBuilder& DetailLayout )
 {
-	TSharedRef<IPropertyHandle> PreferredRendererPropertyHandle = DetailLayout.GetProperty("RendererID");
+	TSharedRef<IPropertyHandle> PreferredRendererPropertyHandle = DetailLayout.GetProperty(GET_MEMBER_NAME_CHECKED(UMacGraphicsSwitchingSettings, RendererID));
 	DetailLayout.HideProperty("RendererID");
 
 	bool bAllowMultiGPUs = IMacGraphicsSwitchingModule::Get().AllowMultipleGPUs();
 	bool bAllowAutomaticGraphicsSwitching = IMacGraphicsSwitchingModule::Get().AllowAutomaticGraphicsSwitching();
 	
-	TSharedRef<IPropertyHandle> MultiGPUPropertyHandle = DetailLayout.GetProperty("bUseMultipleRenderers");
+	TSharedRef<IPropertyHandle> MultiGPUPropertyHandle = DetailLayout.GetProperty(GET_MEMBER_NAME_CHECKED(UMacGraphicsSwitchingSettings, bUseMultipleRenderers));
 	if (!bAllowMultiGPUs)
 	{
 		MultiGPUPropertyHandle->SetValue(false);
-		DetailLayout.HideProperty("bUseMultipleRenderers");
+		DetailLayout.HideProperty(GET_MEMBER_NAME_CHECKED(UMacGraphicsSwitchingSettings, bUseMultipleRenderers));
 	}
 	
-	TSharedRef<IPropertyHandle> SwitchingPropertyHandle = DetailLayout.GetProperty("bAllowAutomaticGraphicsSwitching");
+	TSharedRef<IPropertyHandle> SwitchingPropertyHandle = DetailLayout.GetProperty(GET_MEMBER_NAME_CHECKED(UMacGraphicsSwitchingSettings, bAllowAutomaticGraphicsSwitching));
 	if (!bAllowAutomaticGraphicsSwitching)
 	{
 		SwitchingPropertyHandle->SetValue(false);
-		DetailLayout.HideProperty("bAllowAutomaticGraphicsSwitching");
+		DetailLayout.HideProperty(GET_MEMBER_NAME_CHECKED(UMacGraphicsSwitchingSettings, bAllowAutomaticGraphicsSwitching));
 	}
 	
 	IDetailCategoryBuilder& AccessorCategory = DetailLayout.EditCategory( "OpenGL" );
diff --git a/Engine/Source/Editor/DetailCustomizations/Private/AnimTransitionNodeDetails.cpp b/Engine/Source/Editor/DetailCustomizations/Private/AnimTransitionNodeDetails.cpp
--- a/Engine/Source/Editor/DetailCustomizations/Private/AnimTransitionNodeDetails.cpp
+++ b/Engine/Source/Editor/DetailCustomizations/Private/AnimTransitionNodeDetails.cpp
@@ -52,11 +52,11 @@ void FAnimTransitionNodeDetails::CustomizeDetails( IDetailLayoutBuilder& DetailB
 	{
 		// Transitions to conduits are just shorthand for some other real transition;
 		// All of the blend related settings are ignored, so hide them.
-		DetailBuilder.HideProperty("Bidirectional");
-		DetailBuilder.HideProperty("CrossfadeDuration");
-		DetailBuilder.HideProperty("CrossfadeMode");
-		DetailBuilder.HideProperty("LogicType");
-		DetailBuilder.HideProperty("Priority Order");
+		DetailBuilder.HideProperty(GET_MEMBER_NAME_CHECKED(UAnimStateTransitionNode, Bidirectional));
+		DetailBuilder.HideProperty(GET_MEMBER_NAME_CHECKED(UAnimStateTransitionNode, CrossfadeDuration));
+		DetailBuilder.HideProperty(GET_MEMBER_NAME_CHECKED(UAnimStateTransitionNode, CrossfadeMode));
+		DetailBuilder.HideProperty(GET_MEMBER_NAME_CHECKED(UAnimStateTransitionNode, LogicType));
+		DetailBuilder.HideProperty(GET_MEMBER_NAME_CHECKED(UAnimStateTransitionNode, PriorityOrder));
 	}
 	else
 	{
@@ -67,9 +67,9 @@ void FAnimTransitionNodeDetails::CustomizeDetails( IDetailLayoutBuilder& DetailB
 			.Font( IDetailLayoutBuilder::GetDetailFontBold() )
 		];
 
-		TransitionCategory.AddProperty("PriorityOrder").DisplayName( TEXT("Priority Order") );
-		TransitionCategory.AddProperty("Bidirectional").DisplayName( TEXT("Bidirectional"));
-		TransitionCategory.AddProperty("LogicType").DisplayName( TEXT("Blend Logic") );
+		TransitionCategory.AddProperty(GET_MEMBER_NAME_CHECKED(UAnimStateTransitionNode, PriorityOrder)).DisplayName(TEXT("Priority Order"));
+		TransitionCategory.AddProperty(GET_MEMBER_NAME_CHECKED(UAnimStateTransitionNode, Bidirectional)).DisplayName( TEXT("Bidirectional"));
+		TransitionCategory.AddProperty(GET_MEMBER_NAME_CHECKED(UAnimStateTransitionNode, LogicType)).DisplayName( TEXT("Blend Logic") );
 
 		UAnimStateTransitionNode* TransNode = TransitionNode.Get();
 		if (TransitionNode != NULL)
@@ -124,8 +124,8 @@ void FAnimTransitionNodeDetails::CustomizeDetails( IDetailLayoutBuilder& DetailB
 		}
 
 		//@TODO: Gate editing these on shared non-authorative ones
-		CrossfadeCategory.AddProperty("CrossfadeDuration").DisplayName( TEXT("Duration") );
-		CrossfadeCategory.AddProperty("CrossfadeMode").DisplayName( TEXT("Mode") );
+		CrossfadeCategory.AddProperty(GET_MEMBER_NAME_CHECKED(UAnimStateTransitionNode, CrossfadeDuration)).DisplayName( TEXT("Duration") );
+		CrossfadeCategory.AddProperty(GET_MEMBER_NAME_CHECKED(UAnimStateTransitionNode, CrossfadeMode)).DisplayName( TEXT("Mode") );
 
 		// Add a button that is only visible when blend logic type is custom
 		CrossfadeCategory.AddCustomRow( LOCTEXT("EditBlendGraph", "Edit Blend Graph").ToString() )
@@ -176,8 +176,8 @@ void FAnimTransitionNodeDetails::CustomizeDetails( IDetailLayoutBuilder& DetailB
 		CreateTransitionEventPropertyWidgets(NotificationCategory, TEXT("TransitionInterrupt"));
 	}
 
-	DetailBuilder.HideProperty("TransitionStart");
-	DetailBuilder.HideProperty("TransitionEnd");
+	DetailBuilder.HideProperty(GET_MEMBER_NAME_CHECKED(UAnimStateTransitionNode, TransitionStart));
+	DetailBuilder.HideProperty(GET_MEMBER_NAME_CHECKED(UAnimStateTransitionNode, TransitionEnd));
 }
 END_SLATE_FUNCTION_BUILD_OPTIMIZATION
 
diff --git a/Engine/Source/Editor/DetailCustomizations/Private/BodySetupDetails.cpp b/Engine/Source/Editor/DetailCustomizations/Private/BodySetupDetails.cpp
--- a/Engine/Source/Editor/DetailCustomizations/Private/BodySetupDetails.cpp
+++ b/Engine/Source/Editor/DetailCustomizations/Private/BodySetupDetails.cpp
@@ -18,15 +18,15 @@ void FBodySetupDetails::CustomizeDetails( IDetailLayoutBuilder& DetailBuilder )
 {
 	// Customize collision section
 	{
-		if ( DetailBuilder.GetProperty("DefaultInstance")->IsValidHandle() )
+		if ( DetailBuilder.GetProperty(GET_MEMBER_NAME_CHECKED(UBodySetup, DefaultInstance))->IsValidHandle() )
 		{
 			IDetailCategoryBuilder& PhysicsCategory = DetailBuilder.EditCategory("Physics");
 			IDetailCategoryBuilder& CollisionCategory = DetailBuilder.EditCategory("Collision");
 
-			TSharedPtr<IPropertyHandle> BodyInstanceHandler = DetailBuilder.GetProperty("DefaultInstance");
+			TSharedPtr<IPropertyHandle> BodyInstanceHandler = DetailBuilder.GetProperty(GET_MEMBER_NAME_CHECKED(UBodySetup, DefaultInstance));
 			DetailBuilder.HideProperty(BodyInstanceHandler);
 
-			TSharedPtr<IPropertyHandle> CollisionTraceHandler = DetailBuilder.GetProperty("CollisionTraceFlag");
+			TSharedPtr<IPropertyHandle> CollisionTraceHandler = DetailBuilder.GetProperty(GET_MEMBER_NAME_CHECKED(UBodySetup, CollisionTraceFlag));
 			DetailBuilder.HideProperty(CollisionTraceHandler);
 
 			// add physics properties to physics category
@@ -41,7 +41,7 @@ void FBodySetupDetails::CustomizeDetails( IDetailLayoutBuilder& DetailBuilder )
 			{
 				TSharedPtr<IPropertyHandle> ChildProperty = BodyInstanceHandler->GetChildHandle(ChildIndex);
 				FString Category = FObjectEditorUtils::GetCategory(ChildProperty->GetProperty());
-				if (ChildProperty->GetProperty()->GetName() == TEXT("bSimulatePhysics") || ChildProperty->GetProperty()->GetName() == TEXT("bAutoWeld"))
+				if (ChildProperty->GetProperty()->GetFName() == GET_MEMBER_NAME_CHECKED(FBodyInstance, bSimulatePhysics) || ChildProperty->GetProperty()->GetFName() == GET_MEMBER_NAME_CHECKED(FBodyInstance, bAutoWeld))
 				{
 					// skip bSimulatePhysics
 					// this is because we don't want bSimulatePhysics to show up 
@@ -51,7 +51,7 @@ void FBodySetupDetails::CustomizeDetails( IDetailLayoutBuilder& DetailBuilder )
 					//also hide bAutoWeld for phat
 					continue;
 				}
-				else if (ChildProperty->GetProperty()->GetName() == TEXT("MassInKg"))
+				else if (ChildProperty->GetProperty()->GetFName() == GET_MEMBER_NAME_CHECKED(FBodyInstance, MassInKg))
 				{
 					PhysicsCategory.AddCustomRow(TEXT("Mass"), false)
 						.IsEnabled(TAttribute<bool>(this, &FBodySetupDetails::IsBodyMassEnabled))
diff --git a/Engine/Source/Editor/DetailCustomizations/Private/BrushDetails.cpp b/Engine/Source/Editor/DetailCustomizations/Private/BrushDetails.cpp
--- a/Engine/Source/Editor/DetailCustomizations/Private/BrushDetails.cpp
+++ b/Engine/Source/Editor/DetailCustomizations/Private/BrushDetails.cpp
@@ -175,7 +175,7 @@ void FBrushDetails::CustomizeDetails( IDetailLayoutBuilder& DetailLayout )
 	};
 
 	// Hide the brush builder if it is NULL
-	TSharedRef<IPropertyHandle> BrushBuilderPropertyHandle = DetailLayout.GetProperty("BrushBuilder");
+	TSharedRef<IPropertyHandle> BrushBuilderPropertyHandle = DetailLayout.GetProperty(GET_MEMBER_NAME_CHECKED(ABrush, BrushBuilder));
 	UObject* BrushBuilderObject = nullptr;
 	BrushBuilderPropertyHandle->GetValue(BrushBuilderObject);
 	if(BrushBuilderObject == nullptr)
@@ -189,7 +189,7 @@ void FBrushDetails::CustomizeDetails( IDetailLayoutBuilder& DetailLayout )
 
 	IDetailCategoryBuilder& BrushBuilderCategory = DetailLayout.EditCategory( "BrushSettings", TEXT(""), ECategoryPriority::Important );
 
-	BrushBuilderCategory.AddProperty( TEXT("BrushType") );
+	BrushBuilderCategory.AddProperty( GET_MEMBER_NAME_CHECKED(ABrush, BrushType) );
 	BrushBuilderCategory.AddCustomRow( LOCTEXT("BrushShape", "Brush Shape").ToString() )
 	.NameContent()
 	[
diff --git a/Engine/Source/Editor/DetailCustomizations/Private/DestructibleMeshDetails.cpp b/Engine/Source/Editor/DetailCustomizations/Private/DestructibleMeshDetails.cpp
--- a/Engine/Source/Editor/DetailCustomizations/Private/DestructibleMeshDetails.cpp
+++ b/Engine/Source/Editor/DetailCustomizations/Private/DestructibleMeshDetails.cpp
@@ -5,6 +5,7 @@
 #include "ScopedTransaction.h"
 #include "ObjectEditorUtils.h"
 #include "IDocumentation.h"
+#include "Engine/DestructibleMesh.h"
 
 #define LOCTEXT_NAMESPACE "DestructibleMeshDetails"
 
@@ -42,14 +43,14 @@ void AddStructToDetails(FName CategoryName, FName PropertyName, IDetailLayoutBui
 void FDestructibleMeshDetails::CustomizeDetails(IDetailLayoutBuilder& DetailBuilder)
 {
 		//we always hide bodysetup as it's not useful in this editor
-		TSharedPtr<IPropertyHandle> BodySetupHandler = DetailBuilder.GetProperty("BodySetup");
+		TSharedPtr<IPropertyHandle> BodySetupHandler = DetailBuilder.GetProperty(GET_MEMBER_NAME_CHECKED(UDestructibleMesh, BodySetup));
 		if (BodySetupHandler.IsValid())
 		{
 			DetailBuilder.HideProperty(BodySetupHandler);
 		}
 		
 		//rest of customization is just moving stuff out of DefaultDestructibleParameters so it's nicer to view
-		TSharedPtr<IPropertyHandle> DefaultParams = DetailBuilder.GetProperty("DefaultDestructibleParameters");
+		TSharedPtr<IPropertyHandle> DefaultParams = DetailBuilder.GetProperty(GET_MEMBER_NAME_CHECKED(UDestructibleMesh, DefaultDestructibleParameters));
 		if (DefaultParams.IsValid() == false)
 		{
 			return;
diff --git a/Engine/Source/Editor/DetailCustomizations/Private/LightComponentDetails.cpp b/Engine/Source/Editor/DetailCustomizations/Private/LightComponentDetails.cpp
--- a/Engine/Source/Editor/DetailCustomizations/Private/LightComponentDetails.cpp
+++ b/Engine/Source/Editor/DetailCustomizations/Private/LightComponentDetails.cpp
@@ -14,22 +14,22 @@ TSharedRef<IDetailCustomization> FLightComponentDetails::MakeInstance()
 void FLightComponentDetails::CustomizeDetails( IDetailLayoutBuilder& DetailBuilder )
 {
 	// Mobility property is on the scene component base class not the light component and that is why we have to use USceneComponent::StaticClass
-	TSharedRef<IPropertyHandle> MobilityHandle = DetailBuilder.GetProperty("Mobility", USceneComponent::StaticClass());
+	TSharedRef<IPropertyHandle> MobilityHandle = DetailBuilder.GetProperty(GET_MEMBER_NAME_CHECKED(ULightComponent, Mobility), USceneComponent::StaticClass());
 	// Set a mobility tooltip specific to lights
 	MobilityHandle->SetToolTipText(LOCTEXT("LightMobilityTooltip", "Mobility for lights controls what the light is allowed to do at runtime and therefore what rendering methods are used.\n● A movable light uses fully dynamic lighting and anything can change in game, however it has a large performance cost, typically proportional to the light's influence size.\n● A stationary light will only have its shadowing and bounced lighting from static geometry baked by Lightmass, all other lighting will be dynamic.  It can change color and intensity in game. \n● A static light is fully baked into lightmaps and therefore has no performance cost, but also can't change in game.").ToString());
 
 	IDetailCategoryBuilder& LightCategory = DetailBuilder.EditCategory( "Light", TEXT(""), ECategoryPriority::TypeSpecific );
 
-	LightIntensityProperty = DetailBuilder.GetProperty("Intensity", ULightComponentBase::StaticClass());
-	IESBrightnessTextureProperty = DetailBuilder.GetProperty("IESTexture");
-	IESBrightnessEnabledProperty = DetailBuilder.GetProperty("bUseIESBrightness");
-	IESBrightnessScaleProperty = DetailBuilder.GetProperty("IESBrightnessScale");
+	LightIntensityProperty = DetailBuilder.GetProperty(GET_MEMBER_NAME_CHECKED(ULightComponent, Intensity), ULightComponentBase::StaticClass());
+	IESBrightnessTextureProperty = DetailBuilder.GetProperty(GET_MEMBER_NAME_CHECKED(ULightComponent, IESTexture));
+	IESBrightnessEnabledProperty = DetailBuilder.GetProperty(GET_MEMBER_NAME_CHECKED(ULightComponent, bUseIESBrightness));
+	IESBrightnessScaleProperty = DetailBuilder.GetProperty(GET_MEMBER_NAME_CHECKED(ULightComponent, IESBrightnessScale));
 
 	if( !IESBrightnessEnabledProperty->IsValidHandle() )
 	{
 		// Brightness and color should be listed first
 		LightCategory.AddProperty( LightIntensityProperty );
-		LightCategory.AddProperty( DetailBuilder.GetProperty("LightColor", ULightComponentBase::StaticClass() ) );
+		LightCategory.AddProperty( DetailBuilder.GetProperty(GET_MEMBER_NAME_CHECKED(ULightComponent, LightColor), ULightComponentBase::StaticClass() ) );
 	}
 	else
 	{
diff --git a/Engine/Source/Editor/DetailCustomizations/Private/MeshComponentDetails.cpp b/Engine/Source/Editor/DetailCustomizations/Private/MeshComponentDetails.cpp
--- a/Engine/Source/Editor/DetailCustomizations/Private/MeshComponentDetails.cpp
+++ b/Engine/Source/Editor/DetailCustomizations/Private/MeshComponentDetails.cpp
@@ -12,7 +12,7 @@ TSharedRef<IDetailCustomization> FMeshComponentDetails::MakeInstance()
 void FMeshComponentDetails::CustomizeDetails( IDetailLayoutBuilder& DetailLayout )
 {
 	RenderingCategory = &DetailLayout.EditCategory("Rendering");
-	TSharedRef<IPropertyHandle> MaterialProperty = DetailLayout.GetProperty( "Materials" );
+	TSharedRef<IPropertyHandle> MaterialProperty = DetailLayout.GetProperty( GET_MEMBER_NAME_CHECKED(UMeshComponent, Materials) );
 
 	if( MaterialProperty->IsValidHandle() )
 	{
diff --git a/Engine/Source/Editor/DetailCustomizations/Private/PhysicsConstraintComponentDetails.cpp b/Engine/Source/Editor/DetailCustomizations/Private/PhysicsConstraintComponentDetails.cpp
--- a/Engine/Source/Editor/DetailCustomizations/Private/PhysicsConstraintComponentDetails.cpp
+++ b/Engine/Source/Editor/DetailCustomizations/Private/PhysicsConstraintComponentDetails.cpp
@@ -30,13 +30,13 @@ void FPhysicsConstraintComponentDetails::CustomizeDetails( IDetailLayoutBuilder&
 
 		if (Objects[i]->IsA(UPhysicsConstraintTemplate::StaticClass()))
 		{
-			ConstraintInstance = DetailBuilder.GetProperty("DefaultInstance");
+			ConstraintInstance = DetailBuilder.GetProperty(GET_MEMBER_NAME_CHECKED(UPhysicsConstraintTemplate, DefaultInstance));
 			bInPhAT = true;
 			break;
 		}
 		else if (Objects[i]->IsA(UPhysicsConstraintComponent::StaticClass()))
 		{
-			ConstraintInstance = DetailBuilder.GetProperty("ConstraintInstance");
+			ConstraintInstance = DetailBuilder.GetProperty(GET_MEMBER_NAME_CHECKED(UPhysicsConstraintComponent, ConstraintInstance));
 			ConstraintComp = (UPhysicsConstraintComponent*)Objects[i].Get();
 			OwningConstraintActor = Cast<APhysicsConstraintActor>(ConstraintComp->GetOwner());
 			break;
@@ -179,25 +179,25 @@ void FPhysicsConstraintComponentDetails::CustomizeDetails( IDetailLayoutBuilder&
 		}
 		
 
-		LinearLimitCat.AddProperty(ConstraintInstance->GetChildHandle("LinearLimitSize").ToSharedRef())
+		LinearLimitCat.AddProperty(ConstraintInstance->GetChildHandle(GET_MEMBER_NAME_CHECKED(FConstraintInstance, LinearLimitSize)).ToSharedRef())
 			.Visibility(TAttribute<EVisibility>::Create(TAttribute<EVisibility>::FGetter::CreateSP(this, &FPhysicsConstraintComponentDetails::IsPropertyVisible, EPropertyType::LinearLimit)));
-		LinearLimitCat.AddProperty(ConstraintInstance->GetChildHandle("bLinearLimitSoft").ToSharedRef())
+		LinearLimitCat.AddProperty(ConstraintInstance->GetChildHandle(GET_MEMBER_NAME_CHECKED(FConstraintInstance, bLinearLimitSoft)).ToSharedRef())
 			.Visibility(TAttribute<EVisibility>::Create(TAttribute<EVisibility>::FGetter::CreateSP(this, &FPhysicsConstraintComponentDetails::IsPropertyVisible, EPropertyType::LinearLimit)));
-		LinearLimitCat.AddProperty(ConstraintInstance->GetChildHandle("LinearLimitStiffness").ToSharedRef())
+		LinearLimitCat.AddProperty(ConstraintInstance->GetChildHandle(GET_MEMBER_NAME_CHECKED(FConstraintInstance, LinearLimitStiffness)).ToSharedRef())
 			.Visibility(TAttribute<EVisibility>::Create(TAttribute<EVisibility>::FGetter::CreateSP(this, &FPhysicsConstraintComponentDetails::IsPropertyVisible, EPropertyType::LinearLimit)));
-		LinearLimitCat.AddProperty(ConstraintInstance->GetChildHandle("LinearLimitDamping").ToSharedRef())
+		LinearLimitCat.AddProperty(ConstraintInstance->GetChildHandle(GET_MEMBER_NAME_CHECKED(FConstraintInstance, LinearLimitDamping)).ToSharedRef())
 			.Visibility(TAttribute<EVisibility>::Create(TAttribute<EVisibility>::FGetter::CreateSP(this, &FPhysicsConstraintComponentDetails::IsPropertyVisible, EPropertyType::LinearLimit)));
-		LinearLimitCat.AddProperty(ConstraintInstance->GetChildHandle("bLinearBreakable").ToSharedRef());
-		LinearLimitCat.AddProperty(ConstraintInstance->GetChildHandle("LinearBreakThreshold").ToSharedRef());
+		LinearLimitCat.AddProperty(ConstraintInstance->GetChildHandle(GET_MEMBER_NAME_CHECKED(FConstraintInstance, bLinearBreakable)).ToSharedRef());
+		LinearLimitCat.AddProperty(ConstraintInstance->GetChildHandle(GET_MEMBER_NAME_CHECKED(FConstraintInstance, LinearBreakThreshold)).ToSharedRef());
 	}
 
 	// Angular Limits
 	{
 		IDetailCategoryBuilder& AngularLimitCat = DetailBuilder.EditCategory("Angular Limits", TEXT(""));
 
-		AngularSwing1MotionProperty = ConstraintInstance->GetChildHandle("AngularSwing1Motion");
-		AngularSwing2MotionProperty = ConstraintInstance->GetChildHandle("AngularSwing2Motion");
-		AngularTwistMotionProperty = ConstraintInstance->GetChildHandle("AngularTwistMotion");
+		AngularSwing1MotionProperty = ConstraintInstance->GetChildHandle(GET_MEMBER_NAME_CHECKED(FConstraintInstance, AngularSwing1Motion));
+		AngularSwing2MotionProperty = ConstraintInstance->GetChildHandle(GET_MEMBER_NAME_CHECKED(FConstraintInstance, AngularSwing2Motion));
+		AngularTwistMotionProperty = ConstraintInstance->GetChildHandle(GET_MEMBER_NAME_CHECKED(FConstraintInstance, AngularTwistMotion));
 
 		TArray<TSharedPtr<FString>> AngularLimitOptionNames;
 		TArray<TSharedPtr<FString>> AngularLimitOptionTooltips;
@@ -271,100 +271,100 @@ void FPhysicsConstraintComponentDetails::CustomizeDetails( IDetailLayoutBuilder&
 				];
 		}
 
-		AngularLimitCat.AddProperty(ConstraintInstance->GetChildHandle("Swing1LimitAngle").ToSharedRef())
+		AngularLimitCat.AddProperty(ConstraintInstance->GetChildHandle(GET_MEMBER_NAME_CHECKED(FConstraintInstance, Swing1LimitAngle)).ToSharedRef())
 			.Visibility(TAttribute<EVisibility>::Create(TAttribute<EVisibility>::FGetter::CreateSP(this, &FPhysicsConstraintComponentDetails::IsPropertyVisible, EPropertyType::AngularSwing1Limit)));
-		AngularLimitCat.AddProperty(ConstraintInstance->GetChildHandle("TwistLimitAngle").ToSharedRef())
+		AngularLimitCat.AddProperty(ConstraintInstance->GetChildHandle(GET_MEMBER_NAME_CHECKED(FConstraintInstance, TwistLimitAngle)).ToSharedRef())
 			.Visibility(TAttribute<EVisibility>::Create(TAttribute<EVisibility>::FGetter::CreateSP(this, &FPhysicsConstraintComponentDetails::IsPropertyVisible, EPropertyType::AngularTwistLimit)));
-		AngularLimitCat.AddProperty(ConstraintInstance->GetChildHandle("Swing2LimitAngle").ToSharedRef())
+		AngularLimitCat.AddProperty(ConstraintInstance->GetChildHandle(GET_MEMBER_NAME_CHECKED(FConstraintInstance, Swing2LimitAngle)).ToSharedRef())
 			.Visibility(TAttribute<EVisibility>::Create(TAttribute<EVisibility>::FGetter::CreateSP(this, &FPhysicsConstraintComponentDetails::IsPropertyVisible, EPropertyType::AngularSwing2Limit)));
 
 
-		AngularLimitCat.AddProperty(ConstraintInstance->GetChildHandle("bSwingLimitSoft").ToSharedRef())
+		AngularLimitCat.AddProperty(ConstraintInstance->GetChildHandle(GET_MEMBER_NAME_CHECKED(FConstraintInstance, bSwingLimitSoft)).ToSharedRef())
 			.Visibility(TAttribute<EVisibility>::Create(TAttribute<EVisibility>::FGetter::CreateSP(this, &FPhysicsConstraintComponentDetails::IsPropertyVisible, EPropertyType::AngularSwingLimit)));
-		AngularLimitCat.AddProperty(ConstraintInstance->GetChildHandle("SwingLimitStiffness").ToSharedRef())
+		AngularLimitCat.AddProperty(ConstraintInstance->GetChildHandle(GET_MEMBER_NAME_CHECKED(FConstraintInstance, SwingLimitStiffness)).ToSharedRef())
 			.Visibility(TAttribute<EVisibility>::Create(TAttribute<EVisibility>::FGetter::CreateSP(this, &FPhysicsConstraintComponentDetails::IsPropertyVisible, EPropertyType::AngularSwingLimit)));
-		AngularLimitCat.AddProperty(ConstraintInstance->GetChildHandle("SwingLimitDamping").ToSharedRef())
+		AngularLimitCat.AddProperty(ConstraintInstance->GetChildHandle(GET_MEMBER_NAME_CHECKED(FConstraintInstance, SwingLimitDamping)).ToSharedRef())
 			.Visibility(TAttribute<EVisibility>::Create(TAttribute<EVisibility>::FGetter::CreateSP(this, &FPhysicsConstraintComponentDetails::IsPropertyVisible, EPropertyType::AngularSwingLimit)));
 
 		
 
-		AngularLimitCat.AddProperty(ConstraintInstance->GetChildHandle("bTwistLimitSoft").ToSharedRef())
+		AngularLimitCat.AddProperty(ConstraintInstance->GetChildHandle(GET_MEMBER_NAME_CHECKED(FConstraintInstance, bTwistLimitSoft)).ToSharedRef())
 			.Visibility(TAttribute<EVisibility>::Create(TAttribute<EVisibility>::FGetter::CreateSP(this, &FPhysicsConstraintComponentDetails::IsPropertyVisible, EPropertyType::AngularTwistLimit)));
-		AngularLimitCat.AddProperty(ConstraintInstance->GetChildHandle("TwistLimitStiffness").ToSharedRef())
+		AngularLimitCat.AddProperty(ConstraintInstance->GetChildHandle(GET_MEMBER_NAME_CHECKED(FConstraintInstance, TwistLimitStiffness)).ToSharedRef())
 			.Visibility(TAttribute<EVisibility>::Create(TAttribute<EVisibility>::FGetter::CreateSP(this, &FPhysicsConstraintComponentDetails::IsPropertyVisible, EPropertyType::AngularTwistLimit)));
-		AngularLimitCat.AddProperty(ConstraintInstance->GetChildHandle("TwistLimitDamping").ToSharedRef())
+		AngularLimitCat.AddProperty(ConstraintInstance->GetChildHandle(GET_MEMBER_NAME_CHECKED(FConstraintInstance, TwistLimitDamping)).ToSharedRef())
 			.Visibility(TAttribute<EVisibility>::Create(TAttribute<EVisibility>::FGetter::CreateSP(this, &FPhysicsConstraintComponentDetails::IsPropertyVisible, EPropertyType::AngularTwistLimit)));
 		
 		if (bInPhAT == false)
 		{
-			AngularLimitCat.AddProperty(ConstraintInstance->GetChildHandle("AngularRotationOffset").ToSharedRef())
+			AngularLimitCat.AddProperty(ConstraintInstance->GetChildHandle(GET_MEMBER_NAME_CHECKED(FConstraintInstance, AngularRotationOffset)).ToSharedRef())
 				.Visibility(TAttribute<EVisibility>::Create(TAttribute<EVisibility>::FGetter::CreateSP(this, &FPhysicsConstraintComponentDetails::IsPropertyVisible, EPropertyType::AngularAnyLimit)));
 
 		}
 		else
 		{
-			AngularLimitCat.AddProperty(ConstraintInstance->GetChildHandle("AngularRotationOffset").ToSharedRef())
+			AngularLimitCat.AddProperty(ConstraintInstance->GetChildHandle(GET_MEMBER_NAME_CHECKED(FConstraintInstance, AngularRotationOffset)).ToSharedRef())
 				.Visibility(EVisibility::Collapsed);
 		}
 
-		AngularLimitCat.AddProperty(ConstraintInstance->GetChildHandle("bAngularBreakable").ToSharedRef());
-		AngularLimitCat.AddProperty(ConstraintInstance->GetChildHandle("AngularBreakThreshold").ToSharedRef());
+		AngularLimitCat.AddProperty(ConstraintInstance->GetChildHandle(GET_MEMBER_NAME_CHECKED(FConstraintInstance, bAngularBreakable)).ToSharedRef());
+		AngularLimitCat.AddProperty(ConstraintInstance->GetChildHandle(GET_MEMBER_NAME_CHECKED(FConstraintInstance, AngularBreakThreshold)).ToSharedRef());
 	}
 
 	// Linear Drive
 	{
 		IDetailCategoryBuilder& LinearMotorCat = DetailBuilder.EditCategory("LinearMotor", TEXT(""));
 
-		LinearPositionDriveProperty = ConstraintInstance->GetChildHandle("bLinearPositionDrive");
-		LinearVelocityDriveProperty = ConstraintInstance->GetChildHandle("bLinearVelocityDrive");
+		LinearPositionDriveProperty = ConstraintInstance->GetChildHandle(GET_MEMBER_NAME_CHECKED(FConstraintInstance, bLinearPositionDrive));
+		LinearVelocityDriveProperty = ConstraintInstance->GetChildHandle(GET_MEMBER_NAME_CHECKED(FConstraintInstance, bLinearVelocityDrive));
 
 		IDetailGroup& PositionGroup = LinearMotorCat.AddGroup("Linear Position Drive", TEXT(""));
 		PositionGroup.HeaderProperty(LinearPositionDriveProperty.ToSharedRef());
-		PositionGroup.AddPropertyRow(ConstraintInstance->GetChildHandle("bLinearXPositionDrive").ToSharedRef())
+		PositionGroup.AddPropertyRow(ConstraintInstance->GetChildHandle(GET_MEMBER_NAME_CHECKED(FConstraintInstance, bLinearXPositionDrive)).ToSharedRef())
 			.Visibility(TAttribute<EVisibility>::Create(TAttribute<EVisibility>::FGetter::CreateSP(this, &FPhysicsConstraintComponentDetails::IsPropertyVisible, EPropertyType::LinearPositionDrive)));
-		PositionGroup.AddPropertyRow(ConstraintInstance->GetChildHandle("bLinearYPositionDrive").ToSharedRef())
+		PositionGroup.AddPropertyRow(ConstraintInstance->GetChildHandle(GET_MEMBER_NAME_CHECKED(FConstraintInstance, bLinearYPositionDrive)).ToSharedRef())
 			.Visibility(TAttribute<EVisibility>::Create(TAttribute<EVisibility>::FGetter::CreateSP(this, &FPhysicsConstraintComponentDetails::IsPropertyVisible, EPropertyType::LinearPositionDrive)));
-		PositionGroup.AddPropertyRow(ConstraintInstance->GetChildHandle("bLinearZPositionDrive").ToSharedRef())
+		PositionGroup.AddPropertyRow(ConstraintInstance->GetChildHandle(GET_MEMBER_NAME_CHECKED(FConstraintInstance, bLinearZPositionDrive)).ToSharedRef())
 			.Visibility(TAttribute<EVisibility>::Create(TAttribute<EVisibility>::FGetter::CreateSP(this, &FPhysicsConstraintComponentDetails::IsPropertyVisible, EPropertyType::LinearPositionDrive)));
-		PositionGroup.AddPropertyRow(ConstraintInstance->GetChildHandle("LinearPositionTarget").ToSharedRef())
+		PositionGroup.AddPropertyRow(ConstraintInstance->GetChildHandle(GET_MEMBER_NAME_CHECKED(FConstraintInstance, LinearPositionTarget)).ToSharedRef())
 			.Visibility(TAttribute<EVisibility>::Create(TAttribute<EVisibility>::FGetter::CreateSP(this, &FPhysicsConstraintComponentDetails::IsPropertyVisible, EPropertyType::LinearPositionDrive)));
 
 		IDetailGroup& VelocityGroup = LinearMotorCat.AddGroup("Linear Velocity Drive", TEXT(""));
 		VelocityGroup.HeaderProperty(LinearVelocityDriveProperty.ToSharedRef());
-		VelocityGroup.AddPropertyRow(ConstraintInstance->GetChildHandle("LinearVelocityTarget").ToSharedRef())
+		VelocityGroup.AddPropertyRow(ConstraintInstance->GetChildHandle(GET_MEMBER_NAME_CHECKED(FConstraintInstance, LinearVelocityTarget)).ToSharedRef())
 			.Visibility(TAttribute<EVisibility>::Create(TAttribute<EVisibility>::FGetter::CreateSP(this, &FPhysicsConstraintComponentDetails::IsPropertyVisible, EPropertyType::LinearVelocityDrive)));
 
-		LinearMotorCat.AddProperty(ConstraintInstance->GetChildHandle("LinearDriveSpring").ToSharedRef())
+		LinearMotorCat.AddProperty(ConstraintInstance->GetChildHandle(GET_MEMBER_NAME_CHECKED(FConstraintInstance, LinearDriveSpring)).ToSharedRef())
 			.Visibility(TAttribute<EVisibility>::Create(TAttribute<EVisibility>::FGetter::CreateSP(this, &FPhysicsConstraintComponentDetails::IsPropertyVisible, EPropertyType::LinearAnyDrive)));
-		LinearMotorCat.AddProperty(ConstraintInstance->GetChildHandle("LinearDriveDamping").ToSharedRef())
+		LinearMotorCat.AddProperty(ConstraintInstance->GetChildHandle(GET_MEMBER_NAME_CHECKED(FConstraintInstance, LinearDriveDamping)).ToSharedRef())
 			.Visibility(TAttribute<EVisibility>::Create(TAttribute<EVisibility>::FGetter::CreateSP(this, &FPhysicsConstraintComponentDetails::IsPropertyVisible, EPropertyType::LinearAnyDrive)));
-		LinearMotorCat.AddProperty(ConstraintInstance->GetChildHandle("LinearDriveForceLimit").ToSharedRef())
+		LinearMotorCat.AddProperty(ConstraintInstance->GetChildHandle(GET_MEMBER_NAME_CHECKED(FConstraintInstance, LinearDriveForceLimit)).ToSharedRef())
 			.Visibility(TAttribute<EVisibility>::Create(TAttribute<EVisibility>::FGetter::CreateSP(this, &FPhysicsConstraintComponentDetails::IsPropertyVisible, EPropertyType::LinearAnyDrive)));
 	}
 
 	// Angular Drive
 	{
-		AngularPositionDriveProperty = ConstraintInstance->GetChildHandle("bAngularOrientationDrive");
-		AngularVelocityDriveProperty = ConstraintInstance->GetChildHandle("bAngularVelocityDrive");
+		AngularPositionDriveProperty = ConstraintInstance->GetChildHandle(GET_MEMBER_NAME_CHECKED(FConstraintInstance, bAngularOrientationDrive));
+		AngularVelocityDriveProperty = ConstraintInstance->GetChildHandle(GET_MEMBER_NAME_CHECKED(FConstraintInstance, bAngularVelocityDrive));
 
 		IDetailCategoryBuilder& AngularMotorCat = DetailBuilder.EditCategory("AngularMotor", TEXT(""));
 
 		IDetailGroup& PositionGroup = AngularMotorCat.AddGroup("Angular Orientation Drive", TEXT(""));
 		PositionGroup.HeaderProperty(AngularPositionDriveProperty.ToSharedRef());
 		
-		PositionGroup.AddPropertyRow(ConstraintInstance->GetChildHandle("AngularOrientationTarget").ToSharedRef())
+		PositionGroup.AddPropertyRow(ConstraintInstance->GetChildHandle(GET_MEMBER_NAME_CHECKED(FConstraintInstance, AngularOrientationTarget)).ToSharedRef())
 			.Visibility(TAttribute<EVisibility>::Create(TAttribute<EVisibility>::FGetter::CreateSP(this, &FPhysicsConstraintComponentDetails::IsPropertyVisible, EPropertyType::AngularPositionDrive)));
 
 		IDetailGroup& VelocityGroup = AngularMotorCat.AddGroup("Angular Velocity Drive", TEXT(""));
 		VelocityGroup.HeaderProperty(AngularVelocityDriveProperty.ToSharedRef());
-		VelocityGroup.AddPropertyRow(ConstraintInstance->GetChildHandle("AngularVelocityTarget").ToSharedRef())
+		VelocityGroup.AddPropertyRow(ConstraintInstance->GetChildHandle(GET_MEMBER_NAME_CHECKED(FConstraintInstance, AngularVelocityTarget)).ToSharedRef())
 			.Visibility(TAttribute<EVisibility>::Create(TAttribute<EVisibility>::FGetter::CreateSP(this, &FPhysicsConstraintComponentDetails::IsPropertyVisible, EPropertyType::AngularVelocityDrive)));
-		VelocityGroup.AddPropertyRow(ConstraintInstance->GetChildHandle("AngularDriveForceLimit").ToSharedRef())
+		VelocityGroup.AddPropertyRow(ConstraintInstance->GetChildHandle(GET_MEMBER_NAME_CHECKED(FConstraintInstance, AngularDriveForceLimit)).ToSharedRef())
 			.Visibility(TAttribute<EVisibility>::Create(TAttribute<EVisibility>::FGetter::CreateSP(this, &FPhysicsConstraintComponentDetails::IsPropertyVisible, EPropertyType::AngularVelocityDrive)));
 
-		AngularMotorCat.AddProperty(ConstraintInstance->GetChildHandle("AngularDriveSpring").ToSharedRef())
+		AngularMotorCat.AddProperty(ConstraintInstance->GetChildHandle(GET_MEMBER_NAME_CHECKED(FConstraintInstance, AngularDriveSpring)).ToSharedRef())
 			.Visibility(TAttribute<EVisibility>::Create(TAttribute<EVisibility>::FGetter::CreateSP(this, &FPhysicsConstraintComponentDetails::IsPropertyVisible, EPropertyType::AngularAnyDrive)));
-		AngularMotorCat.AddProperty(ConstraintInstance->GetChildHandle("AngularDriveDamping").ToSharedRef())
+		AngularMotorCat.AddProperty(ConstraintInstance->GetChildHandle(GET_MEMBER_NAME_CHECKED(FConstraintInstance, AngularDriveDamping)).ToSharedRef())
 			.Visibility(TAttribute<EVisibility>::Create(TAttribute<EVisibility>::FGetter::CreateSP(this, &FPhysicsConstraintComponentDetails::IsPropertyVisible, EPropertyType::AngularAnyDrive)));
 	}
 }
diff --git a/Engine/Source/Editor/DetailCustomizations/Private/PointLightComponentDetails.cpp b/Engine/Source/Editor/DetailCustomizations/Private/PointLightComponentDetails.cpp
--- a/Engine/Source/Editor/DetailCustomizations/Private/PointLightComponentDetails.cpp
+++ b/Engine/Source/Editor/DetailCustomizations/Private/PointLightComponentDetails.cpp
@@ -13,7 +13,7 @@ TSharedRef<IDetailCustomization> FPointLightComponentDetails::MakeInstance()
 
 void FPointLightComponentDetails::CustomizeDetails( IDetailLayoutBuilder& DetailBuilder )
 {
-	TSharedPtr<IPropertyHandle> LightIntensityProperty = DetailBuilder.GetProperty("Intensity", ULightComponentBase::StaticClass());
+	TSharedPtr<IPropertyHandle> LightIntensityProperty = DetailBuilder.GetProperty(GET_MEMBER_NAME_CHECKED(ULightComponentBase, Intensity), ULightComponentBase::StaticClass());
 
 	// Inverse squared falloff point lights (the default) are in units of lumens, instead of just being a brightness scale
 	LightIntensityProperty->GetProperty()->SetMetaData("UIMin",TEXT("0.0f"));
diff --git a/Engine/Source/Editor/DetailCustomizations/Private/PrimitiveComponentDetails.cpp b/Engine/Source/Editor/DetailCustomizations/Private/PrimitiveComponentDetails.cpp
--- a/Engine/Source/Editor/DetailCustomizations/Private/PrimitiveComponentDetails.cpp
+++ b/Engine/Source/Editor/DetailCustomizations/Private/PrimitiveComponentDetails.cpp
@@ -144,12 +144,12 @@ EVisibility FPrimitiveComponentDetails::IsMassVisible(bool bOverrideMass) const
 
 void FPrimitiveComponentDetails::CustomizeDetails( IDetailLayoutBuilder& DetailBuilder )
 {
-	TSharedRef<IPropertyHandle> MobilityHandle = DetailBuilder.GetProperty("Mobility", USceneComponent::StaticClass());
+	TSharedRef<IPropertyHandle> MobilityHandle = DetailBuilder.GetProperty(GET_MEMBER_NAME_CHECKED(UPrimitiveComponent, Mobility), USceneComponent::StaticClass());
 	MobilityHandle->SetToolTipText(LOCTEXT("PrimitiveMobilityTooltip", "Mobility for primitive components controls how they can be modified in game and therefore how they interact with lighting and physics.\n● A movable primitive component can be changed in game, but requires dynamic lighting and shadowing from lights which have a large performance cost.\n● A static primitive component can't be changed in game, but can have its lighting baked, which allows rendering to be very efficient.").ToString());
 
-	if ( DetailBuilder.GetProperty("BodyInstance")->IsValidHandle() )
+	if ( DetailBuilder.GetProperty(GET_MEMBER_NAME_CHECKED(UPrimitiveComponent, BodyInstance))->IsValidHandle() )
 	{
-		TSharedPtr<IPropertyHandle> BodyInstanceHandler = DetailBuilder.GetProperty("BodyInstance");
+		TSharedPtr<IPropertyHandle> BodyInstanceHandler = DetailBuilder.GetProperty(GET_MEMBER_NAME_CHECKED(UPrimitiveComponent, BodyInstance));
 		uint32 NumChildren = 0;
 		BodyInstanceHandler->GetNumChildren(NumChildren);
 
diff --git a/Engine/Source/Editor/DetailCustomizations/Private/SceneComponentDetails.cpp b/Engine/Source/Editor/DetailCustomizations/Private/SceneComponentDetails.cpp
--- a/Engine/Source/Editor/DetailCustomizations/Private/SceneComponentDetails.cpp
+++ b/Engine/Source/Editor/DetailCustomizations/Private/SceneComponentDetails.cpp
@@ -608,12 +608,12 @@ void FSceneComponentDetails::MakeTransformDetails( IDetailLayoutBuilder& DetailB
 
 
 	// Hide the transform properties so they don't show up
-	DetailBuilder.HideProperty( DetailBuilder.GetProperty( "bAbsoluteLocation" ) );
-	DetailBuilder.HideProperty( DetailBuilder.GetProperty( "bAbsoluteRotation" ) );
-	DetailBuilder.HideProperty( DetailBuilder.GetProperty( "bAbsoluteScale" ) );
-	DetailBuilder.HideProperty( DetailBuilder.GetProperty( "RelativeLocation" ) );
-	DetailBuilder.HideProperty( DetailBuilder.GetProperty( "RelativeRotation" ) );
-	DetailBuilder.HideProperty( DetailBuilder.GetProperty( "RelativeScale3D" ) );
+	DetailBuilder.HideProperty( DetailBuilder.GetProperty( GET_MEMBER_NAME_CHECKED(USceneComponent, bAbsoluteLocation) ) );
+	DetailBuilder.HideProperty( DetailBuilder.GetProperty( GET_MEMBER_NAME_CHECKED(USceneComponent, bAbsoluteRotation) ) );
+	DetailBuilder.HideProperty( DetailBuilder.GetProperty( GET_MEMBER_NAME_CHECKED(USceneComponent, bAbsoluteScale) ) );
+	DetailBuilder.HideProperty( DetailBuilder.GetProperty( GET_MEMBER_NAME_CHECKED(USceneComponent, RelativeLocation) ) );
+	DetailBuilder.HideProperty( DetailBuilder.GetProperty( GET_MEMBER_NAME_CHECKED(USceneComponent, RelativeRotation) ) );
+	DetailBuilder.HideProperty( DetailBuilder.GetProperty( GET_MEMBER_NAME_CHECKED(USceneComponent, RelativeScale3D) ) );
 
 	// Determine whether or not we are editing Blueprint defaults through the CDO
 	bool bIsEditingBlueprintDefaults = false;
diff --git a/Engine/Source/Editor/DetailCustomizations/Private/SkinnedMeshComponentDetails.cpp b/Engine/Source/Editor/DetailCustomizations/Private/SkinnedMeshComponentDetails.cpp
--- a/Engine/Source/Editor/DetailCustomizations/Private/SkinnedMeshComponentDetails.cpp
+++ b/Engine/Source/Editor/DetailCustomizations/Private/SkinnedMeshComponentDetails.cpp
@@ -19,10 +19,10 @@ void FSkinnedMeshComponentDetails::CustomizeDetails(IDetailLayoutBuilder& Detail
 	IDetailCategoryBuilder& LODCategory = DetailBuilder.EditCategory("LOD", TEXT("Level of Detail"), ECategoryPriority::Default);
 	
 	// show extra field about actually used physics asset, but make sure to show it under physics asset override
-	TSharedRef<IPropertyHandle> PhysicsAssetProperty =  DetailBuilder.GetProperty("PhysicsAssetOverride");
+	TSharedRef<IPropertyHandle> PhysicsAssetProperty =  DetailBuilder.GetProperty(GET_MEMBER_NAME_CHECKED(USkinnedMeshComponent, PhysicsAssetOverride));
 	if (PhysicsAssetProperty->IsValidHandle())
 	{
-		PhysicsCategory.AddProperty("PhysicsAssetOverride");
+		PhysicsCategory.AddProperty(GET_MEMBER_NAME_CHECKED(USkinnedMeshComponent, PhysicsAssetOverride));
 		CreateActuallyUsedPhysicsAssetWidget(PhysicsCategory.AddCustomRow( LOCTEXT("CurrentPhysicsAsset", "Currently used Physics Asset").ToString(), true), &DetailBuilder);
 	}
 }
diff --git a/Engine/Source/Editor/DetailCustomizations/Private/SkyLightComponentDetails.cpp b/Engine/Source/Editor/DetailCustomizations/Private/SkyLightComponentDetails.cpp
--- a/Engine/Source/Editor/DetailCustomizations/Private/SkyLightComponentDetails.cpp
+++ b/Engine/Source/Editor/DetailCustomizations/Private/SkyLightComponentDetails.cpp
@@ -17,11 +17,11 @@ TSharedRef<IDetailCustomization> FSkyLightComponentDetails::MakeInstance()
 void FSkyLightComponentDetails::CustomizeDetails( IDetailLayoutBuilder& DetailLayout )
 {
 	// Mobility property is on the scene component base class not the light component and that is why we have to use USceneComponent::StaticClass
-	TSharedRef<IPropertyHandle> MobilityHandle = DetailLayout.GetProperty("Mobility", USceneComponent::StaticClass());
+	TSharedRef<IPropertyHandle> MobilityHandle = DetailLayout.GetProperty(GET_MEMBER_NAME_CHECKED(USkyLightComponent, Mobility), USceneComponent::StaticClass());
 	// Set a mobility tooltip specific to lights
 	MobilityHandle->SetToolTipText(LOCTEXT("SkyLightMobilityTooltip", "Mobility for sky light components determines what rendering methods will be used.  A Stationary sky light has its shadowing baked into Bent Normal AO by Lightmass, but its lighting can be changed in game.").ToString());
 
-	TSharedPtr<IPropertyHandle> LightIntensityProperty = DetailLayout.GetProperty("Intensity", ULightComponentBase::StaticClass());
+	TSharedPtr<IPropertyHandle> LightIntensityProperty = DetailLayout.GetProperty(GET_MEMBER_NAME_CHECKED(USkyLightComponent, Intensity), ULightComponentBase::StaticClass());
 	// Point lights need to override the ui min and max for units of lumens, so we have to undo that
 	LightIntensityProperty->GetProperty()->SetMetaData("UIMin",TEXT("0.0f"));
 	LightIntensityProperty->GetProperty()->SetMetaData("UIMax",TEXT("20.0f"));
diff --git a/Engine/Source/Editor/Kismet/Private/BlueprintDetailsCustomization.cpp b/Engine/Source/Editor/Kismet/Private/BlueprintDetailsCustomization.cpp
--- a/Engine/Source/Editor/Kismet/Private/BlueprintDetailsCustomization.cpp
+++ b/Engine/Source/Editor/Kismet/Private/BlueprintDetailsCustomization.cpp
@@ -4889,7 +4889,7 @@ FChildActorComponentDetails::FChildActorComponentDetails(TWeakPtr<FBlueprintEdit
 
 void FChildActorComponentDetails::CustomizeDetails(IDetailLayoutBuilder& DetailBuilder)
 {
-	TSharedPtr<IPropertyHandle> ActorClassProperty = DetailBuilder.GetProperty("ChildActorClass");
+	TSharedPtr<IPropertyHandle> ActorClassProperty = DetailBuilder.GetProperty(GET_MEMBER_NAME_CHECKED(UChildActorComponent, ChildActorClass));
 	if (ActorClassProperty->IsValidHandle())
 	{
 		if (BlueprintEditorPtr.IsValid())
diff --git a/Engine/Source/Editor/UnrealEd/Classes/Editor/EditorEngine.h b/Engine/Source/Editor/UnrealEd/Classes/Editor/EditorEngine.h
--- a/Engine/Source/Editor/UnrealEd/Classes/Editor/EditorEngine.h
+++ b/Engine/Source/Editor/UnrealEd/Classes/Editor/EditorEngine.h
@@ -652,9 +652,6 @@ public:
 	bool	HandleBugItGoCommand( const TCHAR* Str, FOutputDevice& Ar );
 	bool	HandleTagSoundsCommand( const TCHAR* Str, FOutputDevice& Ar );
 	bool	HandlecheckSoundsCommand( const TCHAR* Str, FOutputDevice& Ar );
-	bool	HandlePruneAnimSetsCommand( const TCHAR* Str, FOutputDevice& Ar, UWorld* InWorld );
-	bool	HandleBakeAnimSetsCommand( const TCHAR* Str, FOutputDevice& Ar, UWorld* InWorld );
-	bool	HandleCleanMatineeAnimSetsCommand( const TCHAR* Str, FOutputDevice& Ar, UWorld* InWorld );
 	bool	HandleFixupBadAnimNotifiersCommand( const TCHAR* Str, FOutputDevice& Ar );
 	bool	HandleSetDetailModeCommand( const TCHAR* Str, FOutputDevice& Ar );
 	bool	HandleSetDetailModeViewCommand( const TCHAR* Str, FOutputDevice& Ar, UWorld* InWorld );
@@ -768,59 +765,6 @@ public:
 	bool Exec_Transaction(const TCHAR* Str, FOutputDevice& Ar);
 	bool Exec_Particle(const TCHAR* Str, FOutputDevice& Ar);
 
-	/**
-	 *	Gather the status of BakeAndPrune settings for the referenced AnimSets in the given level
-	 *
-	 *	@param	InLevel							The source level
-	 *	@param	OutAnimSetSkipBakeAndPruneMap	List of anim sets to not touch during the operation
-	 *											This is filled in during the bake, and used during the prune
-	 *											If NULL, bake all possible animsets
-	 *
-	 *	@return	bool							true if successful, false if not
-	 */
-	bool GatherBakeAndPruneStatus(ULevel* InLevel, TMap<FString,bool>* OutAnimSetSkipBakeAndPruneMap);
-
-	/**
-	 *	Clear the BakeAndPrune status arrays in the given level
-	 *
-	 *	@param	InLevel							The source level
-	 *
-	 *	@return	bool							true if successful, false if not
-	 */
-	bool ClearBakeAndPruneStatus(ULevel* InLevel);
-
-	/**
-	 *	Bake the anim sets in the given level
-	 *
-	 *	@param	InLevel							The source level
-	 *	@param	InAnimSetSkipBakeAndPruneMap	List of anim sets to not touch during the operation
-	 *											If NULL, bake all possible animsets
-	 *
-	 *	@return	bool							true if successful, false if not
-	 */
-	bool BakeAnimSetsInLevel(ULevel* InLevel, TMap<FString,bool>* InAnimSetSkipBakeAndPruneMap);
-
-	/**
-	 *	Prune the anim sets in the given level
-	 *
-	 *	@param	InLevel							The source level
-	 *	@param	InAnimSetSkipBakeAndPruneMap	List of anim sets to not touch during the operation
-	 *											Filled in during the bake, and used during the prune
-	 *											If NULL, prine all possible animsets
-	 *
-	 *	@return	bool							true if successful, false if not
-	 */
-	bool PruneAnimSetsInLevel(ULevel* InLevel, TMap<FString,bool>* InAnimSetSkipBakeAndPruneMap);
-
-	/**
-	 *	Clear unreferenced AnimSets from InterpData Groups
-	 *
-	 *	@param	InLevel		The source level to clean up
-	 *
-	 *	@return	bool		true if successful, false if not
-	 */
-	bool ClearUnreferenceAnimSetsFromGroups(ULevel* InLevel);
-
 	/**
 	 * Executes each line of text in a file sequentially, as if each were a separate command
 	 *
diff --git a/Engine/Source/Editor/UnrealEd/Private/EditorServer.cpp b/Engine/Source/Editor/UnrealEd/Private/EditorServer.cpp
--- a/Engine/Source/Editor/UnrealEd/Private/EditorServer.cpp
+++ b/Engine/Source/Editor/UnrealEd/Private/EditorServer.cpp
@@ -4717,94 +4717,6 @@ bool UEditorEngine::Exec_Particle(const TCHAR* Str, FOutputDevice& Ar)
 }
 
 
-bool UEditorEngine::GatherBakeAndPruneStatus(ULevel* InLevel, TMap<FString,bool>* OutAnimSetSkipBakeAndPruneMap)
-{
-	if (OutAnimSetSkipBakeAndPruneMap != NULL)
-	{
-		for (TObjectIterator<UInterpData> IterpDataIt; IterpDataIt; ++IterpDataIt)
-		{
-			UInterpData* InterpData = *IterpDataIt;
-			if (InterpData->IsIn(InLevel))
-			{
-				// Gather the animsets from this InterpData
-				if (InterpData->bShouldBakeAndPrune == false)
-				{
-					// None of the anim sets in this data should be bake and pruned.
-					for (int32 AnimSetIdx = 0; AnimSetIdx < InterpData->BakeAndPruneStatus.Num(); AnimSetIdx++)
-					{
-						// If the anim set is referenced in *any* interp data that should not bake & prune,
-						// then it should not be bake and pruned in any of them...
-						FAnimSetBakeAndPruneStatus& Status = InterpData->BakeAndPruneStatus[AnimSetIdx];
-						OutAnimSetSkipBakeAndPruneMap->Add(Status.AnimSetName, true);
-					}
-				}
-				else
-				{
-					for (int32 AnimSetIdx = 0; AnimSetIdx < InterpData->BakeAndPruneStatus.Num(); AnimSetIdx++)
-					{
-						// If the anim set is referenced in *any* interp data that should not bake & prune,
-						// then it should not be bake and pruned in any of them...
-						FAnimSetBakeAndPruneStatus& Status = InterpData->BakeAndPruneStatus[AnimSetIdx];
-						bool* pShouldBakeAndPrune = OutAnimSetSkipBakeAndPruneMap->Find(Status.AnimSetName);
-						if (pShouldBakeAndPrune == NULL)
-						{
-							// First instance of the AnimSet
-							OutAnimSetSkipBakeAndPruneMap->Add(Status.AnimSetName, Status.bSkipBakeAndPrune);
-						}
-						else
-						{
-							// If found, SkipBakeAndPrune == true trumps any previous setting
-							if (*pShouldBakeAndPrune == false)
-							{
-								OutAnimSetSkipBakeAndPruneMap->Add(Status.AnimSetName, Status.bSkipBakeAndPrune);
-							}
-						}
-					}
-				}
-			}
-		}
-	}
-
-	return true;
-}
-
-
-bool UEditorEngine::ClearBakeAndPruneStatus(ULevel* InLevel)
-{
-	for (TObjectIterator<UInterpData> IterpDataIt; IterpDataIt; ++IterpDataIt)
-	{
-		UInterpData* InterpData = *IterpDataIt;
-		if (InterpData->IsIn(InLevel))
-		{
-			InterpData->BakeAndPruneStatus.Empty();
-		}
-	}
-
-	return true;
-}
-
-
-bool UEditorEngine::BakeAnimSetsInLevel(ULevel* InLevel, TMap<FString,bool>* InAnimSetSkipBakeAndPruneMap)
-{
-	// @TODO: ANIM: we'll need a feature to cut the animation to the size
-	return false;
-}
-
-
-bool UEditorEngine::PruneAnimSetsInLevel(ULevel* InLevel, TMap<FString,bool>* InAnimSetSkipBakeAndPruneMap)
-{
-	//@TODO: ANIM
-	return false;
-}
-
-
-bool UEditorEngine::ClearUnreferenceAnimSetsFromGroups(ULevel* InLevel)
-{
-	//@TODO: ANIM
-	return true;
-}
-
-
 void UEditorEngine::ExecFile( UWorld* InWorld, const TCHAR* InFilename, FOutputDevice& Ar )
 {
 	FString FileTextContents;
@@ -5380,18 +5292,6 @@ bool UEditorEngine::Exec( UWorld* InWorld, const TCHAR* Stream, FOutputDevice& A
 	{
 		return HandlecheckSoundsCommand( Str, Ar );
 	}
-	else if( FParse::Command(&Str,TEXT("PRUNEANIMSETS")) )
-	{
-		return HandlePruneAnimSetsCommand( Str, Ar, InWorld );
-	}
-	else if( FParse::Command(&Str,TEXT("BAKEANIMSETS")) )
-	{
-		return HandleBakeAnimSetsCommand( Str, Ar, InWorld );
-	}
-	else if( FParse::Command(&Str,TEXT("CLEANMATINEEANIMSETS")) )
-	{
-		return HandleCleanMatineeAnimSetsCommand( Str, Ar, InWorld );
-	}
 	else if( FParse::Command(&Str,TEXT("FIXUPBADANIMNOTIFIERS")) )
 	{
 		return HandleFixupBadAnimNotifiersCommand( Str, Ar );
@@ -6010,29 +5910,6 @@ bool UEditorEngine::HandlecheckSoundsCommand( const TCHAR* Str, FOutputDevice& A
 		return true;
 }
 
-bool UEditorEngine::HandlePruneAnimSetsCommand( const TCHAR* Str, FOutputDevice& Ar, UWorld* InWorld )
-{
-	TMap<FString,bool> AnimSetSkipBakeAndPruneMap;
-	GatherBakeAndPruneStatus(InWorld->GetCurrentLevel(), &AnimSetSkipBakeAndPruneMap);
-	PruneAnimSetsInLevel(InWorld->GetCurrentLevel(), &AnimSetSkipBakeAndPruneMap);
-	return true;
-}
-
-bool UEditorEngine::HandleBakeAnimSetsCommand( const TCHAR* Str, FOutputDevice& Ar, UWorld* InWorld )
-{
-	TMap<FString,bool> AnimSetSkipBakeAndPruneMap;
-	GatherBakeAndPruneStatus(InWorld->GetCurrentLevel(), &AnimSetSkipBakeAndPruneMap);
-	BakeAnimSetsInLevel(InWorld->GetCurrentLevel(), &AnimSetSkipBakeAndPruneMap);
-	return true;
-}
-
-bool UEditorEngine::HandleCleanMatineeAnimSetsCommand( const TCHAR* Str, FOutputDevice& Ar, UWorld* InWorld )
-{
-	// Clear out unreferenced animsets from groups...
-	ClearUnreferenceAnimSetsFromGroups(InWorld->GetCurrentLevel());
-	return true;
-}
-
 bool UEditorEngine::HandleFixupBadAnimNotifiersCommand( const TCHAR* Str, FOutputDevice& Ar )
 {
 	// Clear out unreferenced animsets from groups...
diff --git a/Engine/Source/Editor/UnrealEd/Private/Factories/PackFactory.cpp b/Engine/Source/Editor/UnrealEd/Private/Factories/PackFactory.cpp
--- a/Engine/Source/Editor/UnrealEd/Private/Factories/PackFactory.cpp
+++ b/Engine/Source/Editor/UnrealEd/Private/Factories/PackFactory.cpp
@@ -139,8 +139,8 @@ namespace PackFactoryHelper
 		PackConfig.ProcessInputFileContents(ConfigString);
 
 		// Input Settings
-		static UArrayProperty* ActionMappingsProp = FindFieldChecked<UArrayProperty>(UInputSettings::StaticClass(), "ActionMappings");
-		static UArrayProperty* AxisMappingsProp = FindFieldChecked<UArrayProperty>(UInputSettings::StaticClass(), "AxisMappings");
+		static UArrayProperty* ActionMappingsProp = FindFieldChecked<UArrayProperty>(UInputSettings::StaticClass(), GET_MEMBER_NAME_CHECKED(UInputSettings, ActionMappings));
+		static UArrayProperty* AxisMappingsProp = FindFieldChecked<UArrayProperty>(UInputSettings::StaticClass(), GET_MEMBER_NAME_CHECKED(UInputSettings, AxisMappings));
 
 		UInputSettings* InputSettingsCDO = GetMutableDefault<UInputSettings>();
 		bool bCheckedOut = false;
diff --git a/Engine/Source/Editor/UnrealEd/Private/Kismet2/ComponentEditorUtils.cpp b/Engine/Source/Editor/UnrealEd/Private/Kismet2/ComponentEditorUtils.cpp
--- a/Engine/Source/Editor/UnrealEd/Private/Kismet2/ComponentEditorUtils.cpp
+++ b/Engine/Source/Editor/UnrealEd/Private/Kismet2/ComponentEditorUtils.cpp
@@ -105,19 +105,19 @@ void FComponentEditorUtils::PropagateTransformPropertyChange(
 		USceneComponent* InstancedSceneComponent = FComponentEditorUtils::GetSceneComponent(ArchetypeInstances[InstanceIndex], InSceneComponentTemplate);
 		if(InstancedSceneComponent != nullptr && !UpdatedComponents.Contains(InstancedSceneComponent))
 		{
-			static const UProperty* RelativeLocationProperty = FindFieldChecked<UProperty>( USceneComponent::StaticClass(), "RelativeLocation" );
+			static const UProperty* RelativeLocationProperty = FindFieldChecked<UProperty>( USceneComponent::StaticClass(), GET_MEMBER_NAME_CHECKED(USceneComponent, RelativeLocation));
 			if(RelativeLocationProperty != nullptr)
 			{
 				PropagateTransformPropertyChange(InstancedSceneComponent, RelativeLocationProperty, OldDefaultTransform.Trans, NewDefaultTransform.Trans, UpdatedComponents);
 			}
 
-			static const UProperty* RelativeRotationProperty = FindFieldChecked<UProperty>( USceneComponent::StaticClass(), "RelativeRotation" );
+			static const UProperty* RelativeRotationProperty = FindFieldChecked<UProperty>( USceneComponent::StaticClass(), GET_MEMBER_NAME_CHECKED(USceneComponent, RelativeLocation) );
 			if(RelativeRotationProperty != nullptr)
 			{
 				PropagateTransformPropertyChange(InstancedSceneComponent, RelativeRotationProperty, OldDefaultTransform.Rot, NewDefaultTransform.Rot, UpdatedComponents);
 			}
 
-			static const UProperty* RelativeScale3DProperty = FindFieldChecked<UProperty>( USceneComponent::StaticClass(), "RelativeScale3D" );
+			static const UProperty* RelativeScale3DProperty = FindFieldChecked<UProperty>( USceneComponent::StaticClass(), GET_MEMBER_NAME_CHECKED(USceneComponent, RelativeScale3D) );
 			if(RelativeScale3DProperty != nullptr)
 			{
 				PropagateTransformPropertyChange(InstancedSceneComponent, RelativeScale3DProperty, OldDefaultTransform.Scale, NewDefaultTransform.Scale, UpdatedComponents);
diff --git a/Engine/Source/Editor/UnrealEd/Private/MaterialGraphSchema.cpp b/Engine/Source/Editor/UnrealEd/Private/MaterialGraphSchema.cpp
--- a/Engine/Source/Editor/UnrealEd/Private/MaterialGraphSchema.cpp
+++ b/Engine/Source/Editor/UnrealEd/Private/MaterialGraphSchema.cpp
@@ -720,7 +720,7 @@ void UMaterialGraphSchema::GetMaterialFunctionActions(FGraphActionMenuBuilder& A
 					const FString LibraryCategoriesString = AssetData.TagsAndValues.FindRef("LibraryCategories");
 					if ( !LibraryCategoriesString.IsEmpty() )
 					{
-						UArrayProperty* LibraryCategoriesProperty = FindFieldChecked<UArrayProperty>(UMaterialFunction::StaticClass(), TEXT("LibraryCategories"));
+						UArrayProperty* LibraryCategoriesProperty = FindFieldChecked<UArrayProperty>(UMaterialFunction::StaticClass(), GET_MEMBER_NAME_CHECKED(UMaterialFunction, LibraryCategories));
 						uint8* DestAddr = (uint8*)(&LibraryCategories);
 						LibraryCategoriesProperty->ImportText(*LibraryCategoriesString, DestAddr, PPF_None, NULL, GWarn);
 					}
diff --git a/Engine/Source/Editor/UnrealEd/Private/Settings/LevelEditorPlaySettingsCustomization.h b/Engine/Source/Editor/UnrealEd/Private/Settings/LevelEditorPlaySettingsCustomization.h
--- a/Engine/Source/Editor/UnrealEd/Private/Settings/LevelEditorPlaySettingsCustomization.h
+++ b/Engine/Source/Editor/UnrealEd/Private/Settings/LevelEditorPlaySettingsCustomization.h
@@ -300,10 +300,10 @@ public:
 		IDetailCategoryBuilder& PlayInNewWindowCategory = LayoutBuilder.EditCategory("PlayInNewWindow");
 		{
 			// new window size
-			TSharedRef<IPropertyHandle> WindowHeightHandle = LayoutBuilder.GetProperty("NewWindowHeight");
-			TSharedRef<IPropertyHandle> WindowWidthHandle = LayoutBuilder.GetProperty("NewWindowWidth");
-			TSharedRef<IPropertyHandle> WindowPositionHandle = LayoutBuilder.GetProperty("NewWindowPosition");
-			TSharedRef<IPropertyHandle> CenterNewWindowHandle = LayoutBuilder.GetProperty("CenterNewWindow");
+			TSharedRef<IPropertyHandle> WindowHeightHandle = LayoutBuilder.GetProperty(GET_MEMBER_NAME_CHECKED(ULevelEditorPlaySettings, NewWindowHeight));
+			TSharedRef<IPropertyHandle> WindowWidthHandle = LayoutBuilder.GetProperty(GET_MEMBER_NAME_CHECKED(ULevelEditorPlaySettings, NewWindowWidth));
+			TSharedRef<IPropertyHandle> WindowPositionHandle = LayoutBuilder.GetProperty(GET_MEMBER_NAME_CHECKED(ULevelEditorPlaySettings, NewWindowPosition));
+			TSharedRef<IPropertyHandle> CenterNewWindowHandle = LayoutBuilder.GetProperty(GET_MEMBER_NAME_CHECKED(ULevelEditorPlaySettings, CenterNewWindow));
 
 			WindowHeightHandle->MarkHiddenByCustomization();
 			WindowWidthHandle->MarkHiddenByCustomization();
@@ -343,8 +343,8 @@ public:
 		IDetailCategoryBuilder& PlayInStandaloneCategory = LayoutBuilder.EditCategory("PlayInStandaloneGame");
 		{
 			// standalone window size
-			TSharedRef<IPropertyHandle> WindowHeightHandle = LayoutBuilder.GetProperty("StandaloneWindowHeight");
-			TSharedRef<IPropertyHandle> WindowWidthHandle = LayoutBuilder.GetProperty("StandaloneWindowWidth");
+			TSharedRef<IPropertyHandle> WindowHeightHandle = LayoutBuilder.GetProperty(GET_MEMBER_NAME_CHECKED(ULevelEditorPlaySettings, StandaloneWindowHeight));
+			TSharedRef<IPropertyHandle> WindowWidthHandle = LayoutBuilder.GetProperty(GET_MEMBER_NAME_CHECKED(ULevelEditorPlaySettings, StandaloneWindowWidth));
 
 			WindowHeightHandle->MarkHiddenByCustomization();
 			WindowWidthHandle->MarkHiddenByCustomization();
@@ -364,7 +364,7 @@ public:
 				];
 
 			// command line options
-			TSharedPtr<IPropertyHandle> DisableStandaloneSoundProperty = LayoutBuilder.GetProperty("DisableStandaloneSound");
+			TSharedPtr<IPropertyHandle> DisableStandaloneSoundProperty = LayoutBuilder.GetProperty(GET_MEMBER_NAME_CHECKED(ULevelEditorPlaySettings, DisableStandaloneSound));
 
 			DisableStandaloneSoundProperty->MarkHiddenByCustomization();
 
diff --git a/Engine/Source/Runtime/Engine/Classes/Matinee/InterpData.h b/Engine/Source/Runtime/Engine/Classes/Matinee/InterpData.h
--- a/Engine/Source/Runtime/Engine/Classes/Matinee/InterpData.h
+++ b/Engine/Source/Runtime/Engine/Classes/Matinee/InterpData.h
@@ -4,31 +4,6 @@
 #pragma once
 #include "InterpData.generated.h"
 
-USTRUCT()
-struct FAnimSetBakeAndPruneStatus
-{
-	GENERATED_USTRUCT_BODY()
-
-	/** Name of the anim set referenced in Matinee */
-	UPROPERTY(Category=AnimSetBakeAndPruneStatus, VisibleAnywhere, BlueprintReadWrite)
-	FString AnimSetName;
-
-	/** If true, the animation is in a GroupAnimSets array, but is unused */
-	UPROPERTY(Category=AnimSetBakeAndPruneStatus, VisibleAnywhere, BlueprintReadWrite)
-	uint32 bReferencedButUnused:1;
-
-	/** If true, skip BakeAndPrune on this anim set during cooking */
-	UPROPERTY(EditAnywhere, Category=AnimSetBakeAndPruneStatus)
-	uint32 bSkipBakeAndPrune:1;
-
-
-	FAnimSetBakeAndPruneStatus()
-		: bReferencedButUnused(false)
-		, bSkipBakeAndPrune(false)
-	{
-	}
-
-};
 
 /**
  * Interpolation data, containing keyframe tracks, event tracks etc.
@@ -82,10 +57,6 @@ class UInterpData : public UObject
 	UPROPERTY(EditAnywhere, Category=InterpData)
 	uint32 bShouldBakeAndPrune:1;
 
-	/** AnimSets referenced by this matinee, and whether to allow bake and prune on them during cooking. */
-	UPROPERTY(EditAnywhere, editfixedsize, Category=InterpData)
-	TArray<struct FAnimSetBakeAndPruneStatus> BakeAndPruneStatus;
-
 	/** Cached version of the director group, if any, for easy access while in game */
 	UPROPERTY(transient)
 	class UInterpGroupDirector* CachedDirectorGroup;
@@ -121,9 +92,6 @@ class UInterpData : public UObject
 	/** Update the AllEventNames array */
 	ENGINE_API void UpdateEventNames();
 
-#if WITH_EDITOR
-	void UpdateBakeAndPruneStatus();
-#endif
 #if WITH_EDITORONLY_DATA
 	void CreateDefaultFilters();
 #endif
diff --git a/Engine/Source/Runtime/Engine/Classes/Matinee/InterpGroup.h b/Engine/Source/Runtime/Engine/Classes/Matinee/InterpGroup.h
--- a/Engine/Source/Runtime/Engine/Classes/Matinee/InterpGroup.h
+++ b/Engine/Source/Runtime/Engine/Classes/Matinee/InterpGroup.h
@@ -117,9 +117,6 @@ private:
 public:
 	// Begin UObject interface.
 	virtual void PostLoad() override;
-#if WITH_EDITOR
-	virtual void PostEditChangeProperty(FPropertyChangedEvent& PropertyChangedEvent) override;
-#endif // WITH_EDITOR
 	// End UObject interface.
 
 
diff --git a/Engine/Source/Runtime/Engine/Classes/Matinee/InterpGroupCamera.h b/Engine/Source/Runtime/Engine/Classes/Matinee/InterpGroupCamera.h
--- a/Engine/Source/Runtime/Engine/Classes/Matinee/InterpGroupCamera.h
+++ b/Engine/Source/Runtime/Engine/Classes/Matinee/InterpGroupCamera.h
@@ -68,13 +68,6 @@ class UInterpGroupCamera : public UInterpGroup
 	/** When compress, tolerance option **/
 	UPROPERTY(EditAnywhere, Category=InterpGroupCamera)
 	float CompressTolerance;
-
-
-	// Begin UOBject Interface
-#if WITH_EDITOR
-	virtual void PostEditChangeProperty(FPropertyChangedEvent& PropertyChangedEvent) override;
-#endif // WITH_EDITOR
-	// End UOBject Interface
 };
 
 
diff --git a/Engine/Source/Runtime/Engine/Classes/PhysicsEngine/BodyInstance.h b/Engine/Source/Runtime/Engine/Classes/PhysicsEngine/BodyInstance.h
--- a/Engine/Source/Runtime/Engine/Classes/PhysicsEngine/BodyInstance.h
+++ b/Engine/Source/Runtime/Engine/Classes/PhysicsEngine/BodyInstance.h
@@ -188,13 +188,6 @@ public:
 	UPROPERTY(EditAnywhere, AdvancedDisplay, BlueprintReadOnly, Category=Physics)
 	uint32 bEnableGravity:1;
 
-	/**
-	 * If true, this body will be put into the asynchronous physics scene. If false, it will be put into the synchronous physics scene.
-	 * If the body is static, it will be placed into both scenes regardless of the value of bUseAsyncScene.
-	 */
-	UPROPERTY(EditAnywhere, AdvancedDisplay, BlueprintReadWrite, Category=Physics)
-	uint32 bUseAsyncScene:1;
-
 	/** If true, it will update mass when scale changes **/
 	UPROPERTY()
 	uint32 bUpdateMassWhenScaleChanges:1;
@@ -236,6 +229,13 @@ public:
 
 protected:
 
+	/**
+	 * If true, this body will be put into the asynchronous physics scene. If false, it will be put into the synchronous physics scene.
+	 * If the body is static, it will be placed into both scenes regardless of the value of bUseAsyncScene.
+	 */
+	UPROPERTY(EditAnywhere, AdvancedDisplay, BlueprintReadWrite, Category=Physics)
+	uint32 bUseAsyncScene:1;
+
 	/** Whether this instance of the object has its own custom walkable slope override setting. */
 	UPROPERTY(EditAnywhere, AdvancedDisplay, BlueprintReadOnly, Category=Physics)
 	uint32 bOverrideWalkableSlopeOnInstance:1;
@@ -396,6 +396,9 @@ public:
 	/** Returns whether this body wants (and can) use the async scene. */
 	bool UseAsyncScene() const;
 
+	/** Indicates whether this body should use the async scene. Must be called before body is init'd, will assert otherwise. Will have no affect if there is no async scene. */
+	void SetUseAsyncScene(bool bNewUseAsyncScene);
+
 #if WITH_PHYSX
 	/**
 	 * Return the PxRigidActor from the given scene (see EPhysicsSceneType), if SceneType is in the range [0, PST_MAX).
diff --git a/Engine/Source/Runtime/Engine/Private/DestructibleComponent.cpp b/Engine/Source/Runtime/Engine/Private/DestructibleComponent.cpp
--- a/Engine/Source/Runtime/Engine/Private/DestructibleComponent.cpp
+++ b/Engine/Source/Runtime/Engine/Private/DestructibleComponent.cpp
@@ -33,7 +33,7 @@ UDestructibleComponent::UDestructibleComponent(const FObjectInitializer& ObjectI
 
 	bHasCustomNavigableGeometry = EHasCustomNavigableGeometry::EvenIfNotCollidable;
 
-	BodyInstance.bUseAsyncScene = true;
+	BodyInstance.SetUseAsyncScene(true);
 	BodyInstance.bEnableCollision_DEPRECATED = true;
 	static FName CollisionProfileName(TEXT("Destructible"));
 	SetCollisionProfileName(CollisionProfileName);
diff --git a/Engine/Source/Runtime/Engine/Private/DialogueVoice.cpp b/Engine/Source/Runtime/Engine/Private/DialogueVoice.cpp
--- a/Engine/Source/Runtime/Engine/Private/DialogueVoice.cpp
+++ b/Engine/Source/Runtime/Engine/Private/DialogueVoice.cpp
@@ -29,12 +29,12 @@ FString UDialogueVoice::GetDesc()
 {
 	FString SummaryString;
 	{
-		UByteProperty* GenderProperty = CastChecked<UByteProperty>( FindFieldChecked<UProperty>( GetClass(), "Gender" ) );
+		UByteProperty* GenderProperty = CastChecked<UByteProperty>(FindFieldChecked<UProperty>(GetClass(), GET_MEMBER_NAME_CHECKED(UDialogueVoice, Gender)));
 		SummaryString += GenderProperty->Enum->GetEnumText(Gender).ToString();
 
 		if( Plurality != EGrammaticalNumber::Singular )
 		{
-			UByteProperty* PluralityProperty = CastChecked<UByteProperty>( FindFieldChecked<UProperty>( GetClass(), "Plurality" ) );
+			UByteProperty* PluralityProperty = CastChecked<UByteProperty>( FindFieldChecked<UProperty>( GetClass(), GET_MEMBER_NAME_CHECKED(UDialogueVoice, Plurality)) );
 
 			SummaryString += ", ";
 			SummaryString += PluralityProperty->Enum->GetEnumText(Plurality).ToString();
diff --git a/Engine/Source/Runtime/Engine/Private/Interpolation.cpp b/Engine/Source/Runtime/Engine/Private/Interpolation.cpp
--- a/Engine/Source/Runtime/Engine/Private/Interpolation.cpp
+++ b/Engine/Source/Runtime/Engine/Private/Interpolation.cpp
@@ -1157,7 +1157,7 @@ UInterpGroupInst* AMatineeActor::FindGroupInst(const AActor* Actor) const
 void AMatineeActor::PostEditChangeProperty( FPropertyChangedEvent& PropertyChangedEvent )
 {
 	Super::PostEditChangeProperty( PropertyChangedEvent );
-	if( PropertyChangedEvent.Property && PropertyChangedEvent.Property->GetName() == TEXT("MatineeData") )
+	if( PropertyChangedEvent.Property && PropertyChangedEvent.Property->GetFName() == GET_MEMBER_NAME_CHECKED(AMatineeActor, MatineeData) )
 	{
 		// Create new entries
 		if( MatineeData )
@@ -1828,9 +1828,6 @@ void UInterpData::PostLoad(void)
 	// Ensure the cached director group is emptied out
 	CachedDirectorGroup = NULL;
 
-#if WITH_EDITOR
-	UpdateBakeAndPruneStatus();
-#endif
 #if WITH_EDITORONLY_DATA
 	if (GIsEditor && DefaultFilters.Num() == 0)
 	{
@@ -1948,92 +1945,6 @@ void UInterpData::UpdateEventNames()
 	}
 }
 
-#if WITH_EDITOR
-void UInterpData::UpdateBakeAndPruneStatus()
-{
-	if (!GIsEditor)
-	{
-		return;
-	}
-
-#if 0 // @todoanim : bakeandprune
-	// Check for anim sets that are referenced
-	TMap<FString,bool> UsedAnimSetNames;
-	TMap<FString,bool> GroupAnimSetNames;
-	TArray<FString> FoundAnimSetNames;
-	for (int32 GroupIdx = 0; GroupIdx < InterpGroups.Num(); GroupIdx++)
-	{
-		UInterpGroup* Group = InterpGroups(GroupIdx);
-		if (Group != NULL)
-		{
-			if (Group->GroupAnimSets.Num() > 0)
-			{
-				for (int32 DumpIdx = 0; DumpIdx < Group->GroupAnimSets.Num(); DumpIdx++)
-				{
-					UAnimSet* FoundSet = Group->GroupAnimSets(DumpIdx);
-					if (FoundSet != NULL)
-					{
-						GroupAnimSetNames.Set(FoundSet->GetPathName(), true);
-						FoundAnimSetNames.AddUniqueItem(FoundSet->GetPathName());
-					}
-				}
-			}
-			// Iterate over all tracks to find anim control tracks and their anim sequences.
-			// We only want to tag animsets that are acutally used.
-			for (int32 TrackIndex = 0; TrackIndex < Group->InterpTracks.Num(); TrackIndex++)
-			{
-				UInterpTrack* InterpTrack = Group->InterpTracks(TrackIndex);
-				UInterpTrackAnimControl* AnimControl = Cast<UInterpTrackAnimControl>(InterpTrack);				
-				if (AnimControl != NULL)
-				{
-					// Iterate over all track key/ sequences and find the associated sequence.
-					for (int32 TrackKeyIndex = 0; TrackKeyIndex < AnimControl->AnimSeqs.Num(); TrackKeyIndex++)
-					{
-						const FAnimControlTrackKey& TrackKey = AnimControl->AnimSeqs[TrackKeyIndex];
-						UAnimSequence* AnimSequence = TrackKey.AnimSeq;
-						if (AnimSequence != NULL)
-						{
-							UsedAnimSetNames.Set(AnimSequence->GetPathName(), true);
-							FoundAnimSetNames.AddUniqueItem(AnimSequence->GetPathName());
-						}
-					}
-				}
-			}
-		}
-	}
-
-	// Add any new ones found
-	for (int32 AnimSetIdx = 0; AnimSetIdx < FoundAnimSetNames.Num(); AnimSetIdx++)
-	{
-		FString AnimSetName = FoundAnimSetNames(AnimSetIdx);
-		bool bUsed = (UsedAnimSetNames.Find(AnimSetName) != NULL);
-		bool bInGroupList = (GroupAnimSetNames.Find(AnimSetName) != NULL);
-		bool bFound = false;
-		for (int32 CheckIdx = 0; CheckIdx < BakeAndPruneStatus.Num(); CheckIdx++)
-		{
-			FAnimSetBakeAndPruneStatus& Status = BakeAndPruneStatus(CheckIdx);
-			if (Status.AnimSetName == AnimSetName)
-			{
-				bFound = true;
-				Status.bReferencedButUnused = !bUsed;
-				break;
-			}
-		}
-		if (bFound == false)
-		{
-			// Add it
-			int32 NewIdx = BakeAndPruneStatus.AddZeroed();
-			FAnimSetBakeAndPruneStatus& Status = BakeAndPruneStatus(NewIdx);
-			Status.AnimSetName = AnimSetName;
-			Status.bSkipBakeAndPrune = false;
-			Status.bReferencedButUnused = !bUsed;
-		}
-	}
-#endif
-}
-#endif
-
-
 /*-----------------------------------------------------------------------------
  UInterpGroup
 -----------------------------------------------------------------------------*/
@@ -2077,23 +1988,6 @@ void UInterpGroup::PostLoad()
 	}
 }
 
-#if WITH_EDITOR
-void UInterpGroup::PostEditChangeProperty(FPropertyChangedEvent& PropertyChangedEvent)
-{
-	Super::PostEditChangeProperty(PropertyChangedEvent);
-
-	if (PropertyChangedEvent.Property && PropertyChangedEvent.Property->GetName() == TEXT("GroupAnimSets"))
-	{
-		// Update the interp data bake and prune list
-		UInterpData* InterpData = Cast<UInterpData>(GetOuter());
-		if (InterpData != NULL)
-		{
-			InterpData->UpdateBakeAndPruneStatus();
-		}
-	}
-}
-#endif // WITH_EDITOR
-
 void UInterpGroup::UpdateGroup(float NewPosition, UInterpGroupInst* GrInst, bool bPreview, bool bJump)
 {
 	check( InterpTracks.Num() == GrInst->TrackInst.Num() );
@@ -9659,42 +9553,6 @@ UInterpGroupInstDirector::UInterpGroupInstDirector(const FObjectInitializer& Obj
 {
 }
 
-#if WITH_EDITOR
-void UInterpGroupCamera::PostEditChangeProperty(FPropertyChangedEvent& PropertyChangedEvent)
-{
-#if WITH_EDITORONLY_DATA
-	if (PropertyChangedEvent.Property && PropertyChangedEvent.Property->GetName() == TEXT("AnimSeqName"))
-	{
-		UCameraAnim * CameraAnim = CastChecked<UCameraAnim>(GetOuter());
-		UInterpGroup* Group = CameraAnim->PreviewInterpGroup;
-		// find first InterpGroup, not me
-		if ( Group )
-		{
-			TArray<UInterpTrack*> AnimTracks;
-			Group->FindTracksByClass(UInterpTrackAnimControl::StaticClass(), AnimTracks);
-			if (AnimTracks.Num() > 0)
-			{
-				UInterpTrackAnimControl* AnimTrack = CastChecked<UInterpTrackAnimControl>(AnimTracks[0]);
-				if (AnimTrack->AnimSeqs.Num() > 0)
-				{
-					FAnimControlTrackKey& SeqKey = AnimTrack->AnimSeqs[ 0 ];
-					SeqKey.AnimSeq = Target.AnimSeq;	
-				}
-				else
-				{
-					int32 KeyIndex = AnimTrack->AddKeyframe(0.0f, NULL, CIM_Linear);
-					FAnimControlTrackKey& NewSeqKey = AnimTrack->AnimSeqs[ KeyIndex ];
-					NewSeqKey.AnimSeq = Target.AnimSeq;	
-				}
-			}
-		}
-	}
-#endif
-
-	Super::PostEditChangeProperty(PropertyChangedEvent);
-}
-#endif // WITH_EDITOR
-
 #if WITH_EDITORONLY_DATA
 /** Returns SpriteComponent subobject **/
 UBillboardComponent* AMatineeActor::GetSpriteComponent() const { return SpriteComponent; }
diff --git a/Engine/Source/Runtime/Engine/Private/PhysicsEngine/BodyInstance.cpp b/Engine/Source/Runtime/Engine/Private/PhysicsEngine/BodyInstance.cpp
--- a/Engine/Source/Runtime/Engine/Private/PhysicsEngine/BodyInstance.cpp
+++ b/Engine/Source/Runtime/Engine/Private/PhysicsEngine/BodyInstance.cpp
@@ -3842,6 +3842,15 @@ bool FBodyInstance::UseAsyncScene() const
 	return bUseAsyncScene && bHasAsyncScene;
 }
 
+
+void FBodyInstance::SetUseAsyncScene(bool bNewUseAsyncScene)
+{
+	// Invalid to call this if the body has been created
+	check(BodyInstancePtr == nullptr);
+	// Set flag
+	bUseAsyncScene = bNewUseAsyncScene;
+}
+
 ////////////////////////////////////////////////////////////////////////////
 // FBodyInstanceEditorHelpers
 
diff --git a/Engine/Source/Runtime/Engine/Private/PrimitiveComponent.cpp b/Engine/Source/Runtime/Engine/Private/PrimitiveComponent.cpp
--- a/Engine/Source/Runtime/Engine/Private/PrimitiveComponent.cpp
+++ b/Engine/Source/Runtime/Engine/Private/PrimitiveComponent.cpp
@@ -494,7 +494,7 @@ void UPrimitiveComponent::PostEditChangeProperty(FPropertyChangedEvent& Property
 		}
 
 		// we need to reregister the primitive if the min draw distance changed to propagate the change to the rendering thread
-		if (PropertyThatChanged->GetName() == TEXT("MinDrawDistance"))
+		if (PropertyThatChanged->GetFName() == GET_MEMBER_NAME_CHECKED(UPrimitiveComponent, MinDrawDistance))
 		{
 			MarkRenderStateDirty();
 		}
diff --git a/Engine/Source/Runtime/Engine/Private/ReflectionCaptureComponent.cpp b/Engine/Source/Runtime/Engine/Private/ReflectionCaptureComponent.cpp
--- a/Engine/Source/Runtime/Engine/Private/ReflectionCaptureComponent.cpp
+++ b/Engine/Source/Runtime/Engine/Private/ReflectionCaptureComponent.cpp
@@ -1112,7 +1112,7 @@ void UReflectionCaptureComponent::PostEditChangeProperty(FPropertyChangedEvent&
 	// AActor::PostEditChange will ForceUpdateComponents()
 	Super::PostEditChangeProperty(PropertyChangedEvent);
 
-	if (PropertyChangedEvent.Property && PropertyChangedEvent.Property->GetName() == TEXT("bVisible"))
+	if (PropertyChangedEvent.Property && PropertyChangedEvent.Property->GetFName() == GET_MEMBER_NAME_CHECKED(UReflectionCaptureComponent, bVisible))
 	{
 		SetCaptureIsDirty();
 	}
@@ -1442,7 +1442,7 @@ void USphereReflectionCaptureComponent::PostEditChangeProperty(FPropertyChangedE
 	Super::PostEditChangeProperty(PropertyChangedEvent);
 
 	if (PropertyChangedEvent.Property && 
-		(PropertyChangedEvent.Property->GetName() == TEXT("InfluenceRadius")))
+		(PropertyChangedEvent.Property->GetFName() == GET_MEMBER_NAME_CHECKED(USphereReflectionCaptureComponent, InfluenceRadius)))
 	{
 		SetCaptureIsDirty();
 	}
@@ -1474,7 +1474,7 @@ void UBoxReflectionCaptureComponent::PostEditChangeProperty(FPropertyChangedEven
 	// AActor::PostEditChange will ForceUpdateComponents()
 	Super::PostEditChangeProperty(PropertyChangedEvent);
 
-	if (PropertyChangedEvent.Property && PropertyChangedEvent.Property->GetName() == TEXT("BoxTransitionDistance"))
+	if (PropertyChangedEvent.Property && PropertyChangedEvent.Property->GetFName() == GET_MEMBER_NAME_CHECKED(UBoxReflectionCaptureComponent, BoxTransitionDistance))
 	{
 		SetCaptureIsDirty();
 	}
diff --git a/Engine/Source/Runtime/Engine/Private/StaticMesh.cpp b/Engine/Source/Runtime/Engine/Private/StaticMesh.cpp
--- a/Engine/Source/Runtime/Engine/Private/StaticMesh.cpp
+++ b/Engine/Source/Runtime/Engine/Private/StaticMesh.cpp
@@ -1413,7 +1413,7 @@ void UStaticMesh::PostEditChangeProperty(FPropertyChangedEvent& PropertyChangedE
 	if (!bAutoComputeLODScreenSize
 		&& RenderData
 		&& PropertyThatChanged
-		&& PropertyThatChanged->GetName() == TEXT("bAutoComputeLODDistance"))
+		&& PropertyThatChanged->GetFName() == GET_MEMBER_NAME_CHECKED(UStaticMesh, bAutoComputeLODScreenSize))
 		{
 		for (int32 LODIndex = 1; LODIndex < SourceModels.Num(); ++LODIndex)
 		{
@@ -1427,8 +1427,8 @@ void UStaticMesh::PostEditChangeProperty(FPropertyChangedEvent& PropertyChangedE
 
 	// Only unbuild lighting for properties which affect static lighting
 	if (!PropertyThatChanged 
-		|| PropertyThatChanged->GetName() == TEXT("LightMapResolution")
-		|| PropertyThatChanged->GetName() == TEXT("LightMapCoordinateIndex"))
+		|| PropertyThatChanged->GetFName() == GET_MEMBER_NAME_CHECKED(UStaticMesh, LightMapResolution)
+		|| PropertyThatChanged->GetFName() == GET_MEMBER_NAME_CHECKED(UStaticMesh, LightMapCoordinateIndex))
 	{
 		FStaticMeshComponentRecreateRenderStateContext Context(this, true);		
 		SetLightingGuid();
diff --git a/Engine/Source/Runtime/Landscape/Private/LandscapeEdit.cpp b/Engine/Source/Runtime/Landscape/Private/LandscapeEdit.cpp
--- a/Engine/Source/Runtime/Landscape/Private/LandscapeEdit.cpp
+++ b/Engine/Source/Runtime/Landscape/Private/LandscapeEdit.cpp
@@ -3941,7 +3941,7 @@ void ULandscapeComponent::SetLOD(bool bForcedLODChanged, int32 InLODValue)
 void ULandscapeComponent::PreEditChange(UProperty* PropertyThatWillChange)
 {
 	Super::PreEditChange(PropertyThatWillChange);
-	if (GIsEditor && PropertyThatWillChange && (PropertyThatWillChange->GetName() == TEXT("ForcedLOD") || PropertyThatWillChange->GetName() == TEXT("LODBias")))
+	if (GIsEditor && PropertyThatWillChange && (PropertyThatWillChange->GetFName() == GET_MEMBER_NAME_CHECKED(ULandscapeComponent, ForcedLOD) || PropertyThatWillChange->GetFName() == GET_MEMBER_NAME_CHECKED(ULandscapeComponent, LODBias)))
 	{
 		// PreEdit unregister component and re-register after PostEdit so we will lose XYtoComponentMap for this component
 		ULandscapeInfo* Info = GetLandscapeInfo(false);
