Flags to lock BodyInstance rotate & translate axis
[https://github.com/EpicGames/UnrealEngine/pull/563] (kylawl)
Also use GET_MEMBER_NAME_CHECKED more in PrimComp details customization

diff --git a/Engine/Source/Editor/DetailCustomizations/Private/PrimitiveComponentDetails.cpp b/Engine/Source/Editor/DetailCustomizations/Private/PrimitiveComponentDetails.cpp
--- a/Engine/Source/Editor/DetailCustomizations/Private/PrimitiveComponentDetails.cpp
+++ b/Engine/Source/Editor/DetailCustomizations/Private/PrimitiveComponentDetails.cpp
@@ -97,6 +97,22 @@ EVisibility FPrimitiveComponentDetails::IsCustomLockedAxisSelected() const
 	return bVisible ? EVisibility::Visible : EVisibility::Collapsed;
 }
 
+EVisibility FPrimitiveComponentDetails::IsLockAxisEnabled() const
+{
+	bool bVisible = false;
+	if (LockedAxisProperty.IsValid())
+	{
+		uint8 LockedAxis;
+		if (LockedAxisProperty->GetValue(LockedAxis) == FPropertyAccess::Success)
+		{
+			bVisible = static_cast<ELockedAxis::Type>(LockedAxis) != ELockedAxis::Default;
+		}
+	}
+
+	return bVisible ? EVisibility::Visible : EVisibility::Collapsed;
+}
+
+
 bool FPrimitiveComponentDetails::IsAutoWeldEditable() const
 {
 	for (int32 i = 0; i < ObjectsCustomized.Num(); ++i)
@@ -185,42 +201,47 @@ void FPrimitiveComponentDetails::CustomizeDetails( IDetailLayoutBuilder& DetailB
 			{
 				TSharedPtr<IPropertyHandle> ChildProperty = BodyInstanceHandler->GetChildHandle(ChildIndex);
 				FString Category = FObjectEditorUtils::GetCategory(ChildProperty->GetProperty());
-				FString PropName = ChildProperty->GetProperty()->GetName();
+				FName PropName = ChildProperty->GetProperty()->GetFName();
 				if (Category == TEXT("Physics"))
 				{
 					// Only permit modifying bSimulatePhysics when the body has some geometry.
-					if (PropName == TEXT("bSimulatePhysics"))
+					if (PropName == GET_MEMBER_NAME_CHECKED(FBodyInstance, bSimulatePhysics))
 					{
 						PhysicsCategory.AddProperty(ChildProperty).EditCondition(TAttribute<bool>(this, &FPrimitiveComponentDetails::IsSimulatePhysicsEditable), NULL);
 					}
-					else if (PropName == TEXT("bUseAsyncScene"))
+					else if (PropName == FName("bUseAsyncScene")) // Can't use GET_MEMBER_NAME_CHECKED because protected
 					{
 						//we only enable bUseAsyncScene if the project uses an AsyncScene
 						PhysicsCategory.AddProperty(ChildProperty).EditCondition(TAttribute<bool>(this, &FPrimitiveComponentDetails::IsUseAsyncEditable), NULL);
 					}
-					else if (PropName == TEXT("LockedAxisMode"))
+					else if (PropName == GET_MEMBER_NAME_CHECKED(FBodyInstance, LockedAxisMode))
 					{
 						LockedAxisProperty = ChildProperty;
 						PhysicsCategory.AddProperty(ChildProperty);
 					}
-					else if (PropName == TEXT("CustomLockedAxis"))
+					else if (PropName == GET_MEMBER_NAME_CHECKED(FBodyInstance, CustomLockedAxis))
 					{
 						//we only enable bUseAsyncScene if the project uses an AsyncScene
 						PhysicsCategory.AddProperty(ChildProperty).Visibility(TAttribute<EVisibility>(this, &FPrimitiveComponentDetails::IsCustomLockedAxisSelected));
 					}
-					else if (PropName == TEXT("bAutoWeld"))
+					else if (PropName == GET_MEMBER_NAME_CHECKED(FBodyInstance, bLockTranslation) || PropName == GET_MEMBER_NAME_CHECKED(FBodyInstance, bLockRotation))
+					{
+						// Only show the lock translation/rotation options if some lock axis is selected
+						PhysicsCategory.AddProperty(ChildProperty).Visibility(TAttribute<EVisibility>(this, &FPrimitiveComponentDetails::IsLockAxisEnabled));
+					}
+					else if (PropName == GET_MEMBER_NAME_CHECKED(FBodyInstance, bAutoWeld))
 					{
 						PhysicsCategory.AddProperty(ChildProperty).Visibility(TAttribute<EVisibility>(this, &FPrimitiveComponentDetails::IsAutoWeldVisible))
 																  .EditCondition(TAttribute<bool>(this, &FPrimitiveComponentDetails::IsAutoWeldEditable), NULL);
 					}
-					else if (PropName == TEXT("bOverrideMass"))
+					else if (PropName == GET_MEMBER_NAME_CHECKED(FBodyInstance, bOverrideMass))
 					{
 						if (bDisplayMassOverride)
 						{
 							PhysicsCategory.AddProperty(ChildProperty);
 						}
 					}
-					else if (PropName == TEXT("MassInKg"))
+					else if (PropName == GET_MEMBER_NAME_CHECKED(FBodyInstance, MassInKg))
 					{
 						
 						if (bDisplayMass)
diff --git a/Engine/Source/Editor/DetailCustomizations/Private/PrimitiveComponentDetails.h b/Engine/Source/Editor/DetailCustomizations/Private/PrimitiveComponentDetails.h
--- a/Engine/Source/Editor/DetailCustomizations/Private/PrimitiveComponentDetails.h
+++ b/Engine/Source/Editor/DetailCustomizations/Private/PrimitiveComponentDetails.h
@@ -36,6 +36,7 @@ private:
 	TArray< TWeakObjectPtr<UObject> > ObjectsCustomized;
 	TSharedPtr<IPropertyHandle> LockedAxisProperty;
 	EVisibility IsCustomLockedAxisSelected() const;
+	EVisibility IsLockAxisEnabled() const;
 
 	bool IsAutoWeldEditable() const;
 	EVisibility IsAutoWeldVisible() const;
diff --git a/Engine/Source/Runtime/Engine/Classes/PhysicsEngine/BodyInstance.h b/Engine/Source/Runtime/Engine/Classes/PhysicsEngine/BodyInstance.h
--- a/Engine/Source/Runtime/Engine/Classes/PhysicsEngine/BodyInstance.h
+++ b/Engine/Source/Runtime/Engine/Classes/PhysicsEngine/BodyInstance.h
@@ -209,6 +209,14 @@ public:
 	UPROPERTY(EditAnywhere, Category = Physics)
 	FVector CustomLockedAxis;
 
+	/** When a Locked Axis Mode is selected, will lock translation on the specified axis*/
+	UPROPERTY(EditAnywhere, Category = Physics, meta=(DisplayName = "Lock Axis Translation"))
+	uint32 bLockTranslation : 1;
+	
+	/** When a Locked Axis Mode is selected, will lock rotation to the specified axis*/
+	UPROPERTY(EditAnywhere, Category = Physics, meta=(DisplayName = "Lock Axis Rotation"))
+	uint32 bLockRotation : 1;
+
 	/** Locks physical movement along axis. */
 	void SetDOFLock(ELockedAxis::Type NewAxisMode);
 
diff --git a/Engine/Source/Runtime/Engine/Private/PhysicsEngine/BodyInstance.cpp b/Engine/Source/Runtime/Engine/Private/PhysicsEngine/BodyInstance.cpp
--- a/Engine/Source/Runtime/Engine/Private/PhysicsEngine/BodyInstance.cpp
+++ b/Engine/Source/Runtime/Engine/Private/PhysicsEngine/BodyInstance.cpp
@@ -210,6 +210,8 @@ FBodyInstance::FBodyInstance()
 , MassInKg(100.f)
 , LockedAxisMode(0)
 , CustomLockedAxis(FVector::ZeroVector)
+, bLockTranslation(true)
+, bLockRotation(true)
 , DOFConstraint(NULL)
 , WeldParent(NULL)
 , bUseAsyncScene(false)
@@ -581,11 +583,11 @@ void FBodyInstance::CreateDOFLock()
 		DOFConstraint->bTwistLimitSoft = false;
 		DOFConstraint->bLinearLimitSoft = false;
 		//set all rotation to free
-		DOFConstraint->AngularSwing1Motion = EAngularConstraintMotion::ACM_Locked;
-		DOFConstraint->AngularSwing2Motion = EAngularConstraintMotion::ACM_Locked;
+		DOFConstraint->AngularSwing1Motion = bLockRotation ? EAngularConstraintMotion::ACM_Locked : EAngularConstraintMotion::ACM_Free;
+		DOFConstraint->AngularSwing2Motion = bLockRotation ? EAngularConstraintMotion::ACM_Locked : EAngularConstraintMotion::ACM_Free;
 		DOFConstraint->AngularTwistMotion = EAngularConstraintMotion::ACM_Free;
 
-		DOFConstraint->LinearXMotion = ELinearConstraintMotion::LCM_Locked;
+		DOFConstraint->LinearXMotion = bLockTranslation ? ELinearConstraintMotion::LCM_Locked : ELinearConstraintMotion::LCM_Free;
 		DOFConstraint->LinearYMotion = ELinearConstraintMotion::LCM_Free;
 		DOFConstraint->LinearZMotion = ELinearConstraintMotion::LCM_Free;
 
