Don't crash in the case where the FKeyStructCustomization::GetCurrentKey is called when the struct no longer exists.

diff --git a/Engine/Source/Editor/DetailCustomizations/Private/KeyStructCustomization.cpp b/Engine/Source/Editor/DetailCustomizations/Private/KeyStructCustomization.cpp
--- a/Engine/Source/Editor/DetailCustomizations/Private/KeyStructCustomization.cpp
+++ b/Engine/Source/Editor/DetailCustomizations/Private/KeyStructCustomization.cpp
@@ -60,27 +60,26 @@ FKey FKeyStructCustomization::GetCurrentKey() const
 {
 	TArray<void*> StructPtrs;
 	PropertyHandle->AccessRawData(StructPtrs);
-	check(StructPtrs.Num() != 0);
-	FKey* SelectedKey = (FKey*)StructPtrs[0];
-
-	bool bMultipleValues = false;
-	for(int32 StructPtrIndex = 1; StructPtrIndex < StructPtrs.Num(); ++StructPtrIndex)
+	if (StructPtrs.Num() != 0)
 	{
-		if(*(FKey*)StructPtrs[StructPtrIndex] != *SelectedKey)
+		FKey* SelectedKey = (FKey*)StructPtrs[0];
+
+		bool bMultipleValues = false;
+		for(int32 StructPtrIndex = 1; StructPtrIndex < StructPtrs.Num(); ++StructPtrIndex)
 		{
-			bMultipleValues = true;
-			break;
+			if(*(FKey*)StructPtrs[StructPtrIndex] != *SelectedKey)
+			{
+				bMultipleValues = true;
+				break;
+			}
 		}
-	}
 
-	if( !bMultipleValues && SelectedKey )
-	{
-		return *SelectedKey;
-	}
-	else
-	{
-		return FKey();
+		if( !bMultipleValues && SelectedKey )
+		{
+			return *SelectedKey;
+		}
 	}
+	return FKey();
 }
 
 void FKeyStructCustomization::OnKeyChanged(TSharedPtr<FKey> SelectedKey)
