UE-2057 Show collision profile description as tooltip in combo box dropdown

diff --git a/Engine/Source/Editor/DetailCustomizations/Private/BodyInstanceCustomization.cpp b/Engine/Source/Editor/DetailCustomizations/Private/BodyInstanceCustomization.cpp
--- a/Engine/Source/Editor/DetailCustomizations/Private/BodyInstanceCustomization.cpp
+++ b/Engine/Source/Editor/DetailCustomizations/Private/BodyInstanceCustomization.cpp
@@ -640,9 +640,18 @@ FString FBodyInstanceCustomization::GetObjectTypeComboBoxContent() const
 
 TSharedRef<SWidget> FBodyInstanceCustomization::MakeCollisionProfileComboWidget(TSharedPtr<FString> InItem)
 {
+	FString ProfileMessage;
+
+	FCollisionResponseTemplate ProfileData;
+	if (CollisionProfile->GetProfileTemplate(FName(**InItem), ProfileData))
+	{
+		ProfileMessage = ProfileData.HelpMessage;
+	}
+
 	return
 		SNew(STextBlock)
 		.Text(*InItem)
+		.ToolTipText(FText::FromString(ProfileMessage))
 		.Font(IDetailLayoutBuilder::GetDetailFont());
 }
 
