Include fixed for NonUnity build;

diff --git a/Engine/Source/Editor/DetailCustomizations/Private/DetailCustomizations.cpp b/Engine/Source/Editor/DetailCustomizations/Private/DetailCustomizations.cpp
--- a/Engine/Source/Editor/DetailCustomizations/Private/DetailCustomizations.cpp
+++ b/Engine/Source/Editor/DetailCustomizations/Private/DetailCustomizations.cpp
@@ -81,6 +81,7 @@
 #include "ParticleSystemComponentDetails.h"
 #include "ParticleSysParamStructCustomization.h"
 #include "RawDistributionVectorStructCustomization.h"
+#include "CollisionProfileNameCustomization.h"
 #include "DocumentationActorDetails.h"
 #include "MediaPlayerCustomization.h"
 #include "MediaSoundWaveCustomization.h"
