Removed Blueprint Notify for State Transition event.

- Nobody knew what this does and I think it's meant to be used for any anim notify, but that doesn't work because it's not showing the right BPs. Am removing it until we have a reason to add this back.

diff --git a/Engine/Source/Editor/DetailCustomizations/Private/AnimTransitionNodeDetails.cpp b/Engine/Source/Editor/DetailCustomizations/Private/AnimTransitionNodeDetails.cpp
--- a/Engine/Source/Editor/DetailCustomizations/Private/AnimTransitionNodeDetails.cpp
+++ b/Engine/Source/Editor/DetailCustomizations/Private/AnimTransitionNodeDetails.cpp
@@ -367,38 +367,9 @@ EVisibility FAnimTransitionNodeDetails::GetBlendGraphButtonVisibility() const
 void FAnimTransitionNodeDetails::CreateTransitionEventPropertyWidgets(IDetailCategoryBuilder& TransitionCategory, FString TransitionName)
 {
 	TSharedPtr<IPropertyHandle> NameProperty = TransitionCategory.GetParentLayout().GetProperty(*(TransitionName + TEXT(".NotifyName")));
-	TSharedPtr<IPropertyHandle> NotifyProperty = TransitionCategory.GetParentLayout().GetProperty(*(TransitionName + TEXT(".Notify")));
 
 	TransitionCategory.AddProperty( NameProperty )
-		.DisplayName( LOCTEXT("CreateTransition_CustomBlueprintEvent", "Custom Blueprint Event").ToString() )
-		.IsEnabled( TAttribute<bool>::Create(TAttribute<bool>::FGetter::CreateSP( this, &FAnimTransitionNodeDetails::ObjIsNULL, NotifyProperty ) ) );
-
-	TransitionCategory.AddProperty( NotifyProperty )
-		.DisplayName( LOCTEXT("CreateTransition_NotifyBlueprintName", "Blueprint Notify").ToString() )
-		.IsEnabled( TAttribute<bool>::Create(TAttribute<bool>::FGetter::CreateSP( this, &FAnimTransitionNodeDetails::NameIsNone, NameProperty ) ) );
-}
-
-bool FAnimTransitionNodeDetails::NameIsNone(TSharedPtr<IPropertyHandle> Property) const
-{
-	if (Property.IsValid())
-	{
-		FName Name = NAME_None;
-		FPropertyAccess::Result Result = Property->GetValue(Name);
-		return (Name == NAME_None);
-	}
-	return true;
-}
-
-bool FAnimTransitionNodeDetails::ObjIsNULL(TSharedPtr<IPropertyHandle> Property) const
-{
-	if (Property.IsValid())
-	{
-		UObject* Obj = NULL;
-		FPropertyAccess::Result Result = Property->GetValue(Obj);
-		return (Obj == NULL);
-
-	}
-	return true;
+		.DisplayName( LOCTEXT("CreateTransition_CustomBlueprintEvent", "Custom Blueprint Event").ToString() );
 }
 
 TSharedRef<SWidget> FAnimTransitionNodeDetails::GetWidgetForInlineShareMenu(FString TypeName, FString SharedName, bool bIsCurrentlyShared,  FOnClicked PromoteClick, FOnClicked DemoteClick, FOnGetContent GetContentMenu)
diff --git a/Engine/Source/Editor/DetailCustomizations/Private/AnimTransitionNodeDetails.h b/Engine/Source/Editor/DetailCustomizations/Private/AnimTransitionNodeDetails.h
--- a/Engine/Source/Editor/DetailCustomizations/Private/AnimTransitionNodeDetails.h
+++ b/Engine/Source/Editor/DetailCustomizations/Private/AnimTransitionNodeDetails.h
@@ -18,10 +18,6 @@ public:
 	virtual void CustomizeDetails( IDetailLayoutBuilder& DetailBuilder ) override;
 
 protected:
-
-	bool NameIsNone(TSharedPtr<IPropertyHandle> Property) const;
-	bool ObjIsNULL(TSharedPtr<IPropertyHandle> Property) const;
-
 	void CreateTransitionEventPropertyWidgets(IDetailCategoryBuilder& TransitionCategory, FString TransitionName);
 
 	FReply OnClickEditBlendGraph();
