Fixed issue with adding engine trace channel

diff --git a/Engine/Source/Editor/DetailCustomizations/Private/CollisionProfileDetails.cpp b/Engine/Source/Editor/DetailCustomizations/Private/CollisionProfileDetails.cpp
--- a/Engine/Source/Editor/DetailCustomizations/Private/CollisionProfileDetails.cpp
+++ b/Engine/Source/Editor/DetailCustomizations/Private/CollisionProfileDetails.cpp
@@ -1792,7 +1792,8 @@ void FCollisionProfileDetails::RefreshChannelList(bool bTraceType)
 
 		for(auto Iter = CollisionProfile->DefaultChannelResponses.CreateIterator(); Iter; ++Iter)
 		{
-			if(Iter->bTraceType)
+			// only display game channels
+			if(Iter->Channel >= ECC_GameTraceChannel1 && Iter->bTraceType)
 			{
 				TraceChannelList.Add(MakeShareable(new FChannelListItem(MakeShareable(new FCustomChannelSetup(*Iter)))));
 			}
@@ -1804,7 +1805,8 @@ void FCollisionProfileDetails::RefreshChannelList(bool bTraceType)
 
 		for(auto Iter = CollisionProfile->DefaultChannelResponses.CreateIterator(); Iter; ++Iter)
 		{
-			if(!Iter->bTraceType)
+			// only display game channels
+			if(Iter->Channel >= ECC_GameTraceChannel1 && !Iter->bTraceType)
 			{
 				ObjectChannelList.Add(MakeShareable(new FChannelListItem(MakeShareable(new FCustomChannelSetup(*Iter)))));
 			}
