Improved the appearance of some details customizations.

diff --git a/Engine/Source/Editor/DetailCustomizations/Private/ActorDetails.cpp b/Engine/Source/Editor/DetailCustomizations/Private/ActorDetails.cpp
--- a/Engine/Source/Editor/DetailCustomizations/Private/ActorDetails.cpp
+++ b/Engine/Source/Editor/DetailCustomizations/Private/ActorDetails.cpp
@@ -699,7 +699,7 @@ void FActorDetails::AddUtilityBlueprintRows( IDetailCategoryBuilder& BlueprintCa
 		FFormatNamedArguments Args;
 		Args.Add( TEXT( "ActorLabel" ), ActorLabel );
 
-		const FText FindButtonLabel = FText::Format( LOCTEXT( "FindInLevelScript", "Find {ActorLabel} in Level Blueprint" ), Args );
+		const FText FindButtonLabel = LOCTEXT( "FindInLevelScript", "Find in Level Blueprint" );
 		const FText FindButtonToolTip = FText::Format( LOCTEXT( "FindInLevelScript_ToolTip", "Search for uses of {ActorLabel} in the Level Blueprint" ), Args );
 
 		BlueprintCategory.AddCustomRow( LOCTEXT( "FindInLevelScript_Filter", "Find in Level Script" ).ToString() )
@@ -712,7 +712,7 @@ void FActorDetails::AddUtilityBlueprintRows( IDetailCategoryBuilder& BlueprintCa
 			.Padding( 0, 0, 10, 0 )
 			[
 				SNew( SBox )
-				.WidthOverride( 250 )
+				.WidthOverride( 200 )
 				[
 					SNew( SButton )
 					.ToolTipText( FindButtonToolTip )
@@ -730,7 +730,7 @@ void FActorDetails::AddUtilityBlueprintRows( IDetailCategoryBuilder& BlueprintCa
 		];
 
 		// Add Level Events
-		const FText AddEventLabel = FText::Format( LOCTEXT( "ScriptingEvents", "Add Level Events for {ActorLabel}" ), Args );
+		const FText AddEventLabel = LOCTEXT( "ScriptingEvents", "Add Level Events" );
 		const FText AddEventToolTip = FText::Format( LOCTEXT( "ScripingEvents_Tooltip", "Adds or views events for {ActorLabel} in the Level Blueprint" ), Args );
 
 		BlueprintCategory.AddCustomRow( LOCTEXT( "ScriptingEvents_Filter", "Level Blueprint Events" ).ToString() )
@@ -743,7 +743,7 @@ void FActorDetails::AddUtilityBlueprintRows( IDetailCategoryBuilder& BlueprintCa
 			.Padding( 0, 0, 10, 0 )
 			[
 				SNew( SBox )
-				.WidthOverride( 250 )
+				.WidthOverride( 200 )
 				[
 					SNew( SComboButton )
 					.IsEnabled( this, &FActorDetails::IsActorValidForLevelScript )
@@ -773,7 +773,7 @@ void FActorDetails::AddUtilityBlueprintRows( IDetailCategoryBuilder& BlueprintCa
 		.Padding( 0, 0, 10, 0 )
 		[
 			SNew( SBox )
-			.WidthOverride( 250 )
+			.WidthOverride( 200 )
 			[
 				SNew( SButton )
 				.ToolTipText( LOCTEXT( "CreateHarvestBlueprint_ToolTip", "Harvest Components from Selected Actors and create Blueprint" ) )
@@ -784,6 +784,7 @@ void FActorDetails::AddUtilityBlueprintRows( IDetailCategoryBuilder& BlueprintCa
 				[
 					SNew( STextBlock )
 					.Text( LOCTEXT( "ReplaceWithAmalgamBlueprint", "Replace With Composited Blueprint" ) )
+					.Font(IDetailLayoutBuilder::GetDetailFont())
 				]
 			]
 		]
@@ -993,37 +994,42 @@ void FActorDetails::AddSingleBlueprintRow( IDetailCategoryBuilder& BlueprintCate
 	if(!InBlueprint)
 	{
 		BlueprintCategory.AddCustomRow( LOCTEXT("CreateBlueprintFilterString", "Create Blueprint").ToString(), true )
+		.WholeRowContent()
+		.MinDesiredWidth(200)
+		.MaxDesiredWidth(200)
 		[
 			SNew( SHorizontalBox )
 			+SHorizontalBox::Slot()
 			.AutoWidth()
 			[
 				SNew(SBorder)
-					.BorderImage(FEditorStyle::GetBrush("NoBorder"))
-					.Padding(0)
-					.ToolTipText(CreateBlueprintToolTip)
+				.BorderImage(FEditorStyle::GetBrush("NoBorder"))
+				.Padding(0)
+				.ToolTipText(CreateBlueprintToolTip)
 				[
 					SNew( SButton )
-						.IsEnabled(bCanCreateActorBlueprint)
-						.ButtonColorAndOpacity(FLinearColor(0.2f, 0.4f, 0.6f, 1.0f))
-						.OnClicked(FOnClicked::CreateRaw(this, 	&FActorDetails::OnPickBlueprintPathClicked, false ))
+					.IsEnabled(bCanCreateActorBlueprint)
+					.ButtonColorAndOpacity(FLinearColor(0.2f, 0.4f, 0.6f, 1.0f))
+					.OnClicked(FOnClicked::CreateRaw(this, 	&FActorDetails::OnPickBlueprintPathClicked, false ))
+					.ContentPadding(2)
 					[
 						SNew( SHorizontalBox )
 						+SHorizontalBox::Slot()
-							.VAlign(VAlign_Center)
-							.AutoWidth()
+						.VAlign(VAlign_Center)
+						.AutoWidth()
 						[
 							SNew(STextBlock)
-								.Text( FText::Format( LOCTEXT("CreateBlueprintFromActor", "Create {0} Blueprint..."), FText::FromString( SelectedActorClassName ) ) )
+							.Text( FText::Format( LOCTEXT("CreateBlueprintFromActor", "Create {0} Blueprint..."), FText::FromString( SelectedActorClassName ) ) )
+							.Font(IDetailLayoutBuilder::GetDetailFont())
 						]
 
 						+SHorizontalBox::Slot()
-							.VAlign(VAlign_Center)
-							.AutoWidth()
-							.Padding( 4.0f, 0.0f, 0.0f, 0.0f )
+						.VAlign(VAlign_Center)
+						.AutoWidth()
+						.Padding( 4.0f, 0.0f, 0.0f, 0.0f )
 						[
 							SNew( SImage )
-								.Image(FEditorStyle::GetBrush(TEXT("Kismet.CreateBlueprint")))
+							.Image(FEditorStyle::GetBrush(TEXT("Kismet.CreateBlueprint")))
 						]
 					]
 				]
diff --git a/Engine/Source/Editor/DetailCustomizations/Private/ReflectionCaptureDetails.cpp b/Engine/Source/Editor/DetailCustomizations/Private/ReflectionCaptureDetails.cpp
--- a/Engine/Source/Editor/DetailCustomizations/Private/ReflectionCaptureDetails.cpp
+++ b/Engine/Source/Editor/DetailCustomizations/Private/ReflectionCaptureDetails.cpp
@@ -30,20 +30,27 @@ void FReflectionCaptureDetails::CustomizeDetails( IDetailLayoutBuilder& DetailLa
 		}
 	}
 
-	DetailLayout.EditCategory( "ReflectionCapture" )
-	.AddCustomRow( NSLOCTEXT("ReflectionCaptureDetails", "UpdateReflectionCaptures", "Update Captures").ToString() )
+	DetailLayout.EditCategory("ReflectionCapture")
+	.AddCustomRow(NSLOCTEXT("ReflectionCaptureDetails", "UpdateReflectionCaptures", "Update Captures").ToString())
+	.NameContent()
 	[
-		SNew(SHorizontalBox)
-		+SHorizontalBox::Slot()
-		.FillWidth(1.f)
-		.Padding(0, 5, 10, 5)
+		SNullWidget::NullWidget
+	]
+	.ValueContent()
+	[
+		SNew(SBox)
+		.WidthOverride(125)
 		[
 			SNew(SButton)
 			.ContentPadding(3)
 			.VAlign(VAlign_Center)
 			.HAlign(HAlign_Center)
-			.OnClicked( this, &FReflectionCaptureDetails::OnUpdateReflectionCaptures )
-			.Text( NSLOCTEXT("ReflectionCaptureDetails", "UpdateReflectionCaptures", "Update Captures").ToString() )
+			.OnClicked(this, &FReflectionCaptureDetails::OnUpdateReflectionCaptures)
+			[
+				SNew(STextBlock)
+				.Text(NSLOCTEXT("ReflectionCaptureDetails", "UpdateReflectionCaptures", "Update Captures").ToString())
+				.Font(IDetailLayoutBuilder::GetDetailFont())
+			]
 		]
 	];
 }
diff --git a/Engine/Source/Editor/DetailCustomizations/Private/StaticMeshActorDetails.cpp b/Engine/Source/Editor/DetailCustomizations/Private/StaticMeshActorDetails.cpp
--- a/Engine/Source/Editor/DetailCustomizations/Private/StaticMeshActorDetails.cpp
+++ b/Engine/Source/Editor/DetailCustomizations/Private/StaticMeshActorDetails.cpp
@@ -45,9 +45,17 @@ void FStaticMeshActorDetails::CustomizeDetails( IDetailLayoutBuilder& DetailBuil
 	const FString CreateBlockingVolumeString = LOCTEXT("BlockingVolumeMenu", "Create Blocking Volume").ToString();
 
 	StaticMeshCategory.AddCustomRow( CreateBlockingVolumeString, bForAdvanced )
+	.NameContent()
 	[
-		SNew( SComboButton )
-		.ToolTipText( LOCTEXT("CreateBlockingVolumeTooltip", "Creates a blocking volume from the static mesh"))
+		SNullWidget::NullWidget
+	]
+	.ValueContent()
+	.VAlign(VAlign_Center)
+	.MaxDesiredWidth(250)
+	[
+		SNew(SComboButton)
+		.VAlign(VAlign_Center)
+		.ToolTipText(LOCTEXT("CreateBlockingVolumeTooltip", "Creates a blocking volume from the static mesh"))
 		.ButtonContent()
 		[
 			SNew( STextBlock )
diff --git a/Engine/Source/Editor/UnrealEd/Public/GameModeInfoCustomizer.h b/Engine/Source/Editor/UnrealEd/Public/GameModeInfoCustomizer.h
--- a/Engine/Source/Editor/UnrealEd/Public/GameModeInfoCustomizer.h
+++ b/Engine/Source/Editor/UnrealEd/Public/GameModeInfoCustomizer.h
@@ -65,12 +65,16 @@ public:
 			+SHorizontalBox::Slot()
 			.AutoWidth()
 			[
-				SNew(SClassPropertyEntryBox)
-				.AllowNone(bAllowNone)
-				.MetaClass(MetaClass)
-				.IsEnabled(this, &FGameModeInfoCustomizer::AllowModifyGameMode)
-				.SelectedClass(this, &FGameModeInfoCustomizer::OnGetDefaultClass, DefaultClassPropertyName)
-				.OnSetClass(FOnSetClass::CreateSP(this, &FGameModeInfoCustomizer::OnSetDefaultClass, DefaultClassPropertyName))
+				SNew(SBox)
+				.WidthOverride(125)
+				[
+					SNew(SClassPropertyEntryBox)
+					.AllowNone(bAllowNone)
+					.MetaClass(MetaClass)
+					.IsEnabled(this, &FGameModeInfoCustomizer::AllowModifyGameMode)
+					.SelectedClass(this, &FGameModeInfoCustomizer::OnGetDefaultClass, DefaultClassPropertyName)
+					.OnSetClass(FOnSetClass::CreateSP(this, &FGameModeInfoCustomizer::OnSetDefaultClass, DefaultClassPropertyName))
+				]
 			]
 			+SHorizontalBox::Slot()
 			.AutoWidth()
@@ -118,11 +122,15 @@ public:
 			+SHorizontalBox::Slot()
 			.AutoWidth()
 			[
-				SNew(SClassPropertyEntryBox)
-				.AllowNone(bAllowNone)
-				.MetaClass(AGameMode::StaticClass())
-				.SelectedClass(this, &FGameModeInfoCustomizer::GetCurrentGameModeClass)
-				.OnSetClass(FOnSetClass::CreateSP(this, &FGameModeInfoCustomizer::SetCurrentGameModeClass))
+				SNew(SBox)
+				.WidthOverride(125)
+				[
+					SNew(SClassPropertyEntryBox)
+					.AllowNone(bAllowNone)
+					.MetaClass(AGameMode::StaticClass())
+					.SelectedClass(this, &FGameModeInfoCustomizer::GetCurrentGameModeClass)
+					.OnSetClass(FOnSetClass::CreateSP(this, &FGameModeInfoCustomizer::SetCurrentGameModeClass))
+				]
 			]
 
 			+SHorizontalBox::Slot()
