Converted the property editor to use FText

The display name, filter string, and tooltips for the property editor (and associated detail customizations) are now stored as FText rather than FString. This allows us to remove SLATE_TEXT_ARGUMENT from the detail customization widgets.

diff --git a/Engine/Plugins/2D/Paper2D/Source/Paper2DEditor/Private/SpriteEditor/SpriteDetailsCustomization.cpp b/Engine/Plugins/2D/Paper2D/Source/Paper2DEditor/Private/SpriteEditor/SpriteDetailsCustomization.cpp
--- a/Engine/Plugins/2D/Paper2D/Source/Paper2DEditor/Private/SpriteEditor/SpriteDetailsCustomization.cpp
+++ b/Engine/Plugins/2D/Paper2D/Source/Paper2DEditor/Private/SpriteEditor/SpriteDetailsCustomization.cpp
@@ -20,7 +20,7 @@ TSharedRef<IDetailCustomization> FSpriteDetailsCustomization::MakeInstance()
 
 FDetailWidgetRow& FSpriteDetailsCustomization::GenerateWarningRow(IDetailCategoryBuilder& WarningCategory, bool bExperimental, const FText& WarningText, const FText& Tooltip, const FString& ExcerptLink, const FString& ExcerptName)
 {
-	const FString SearchString = WarningText.ToString();
+	const FText SearchString = WarningText;
 	const FSlateBrush* WarningIcon = FEditorStyle::GetBrush(bExperimental ? "PropertyEditor.ExperimentalClass" : "PropertyEditor.EarlyAccessClass");
 
 	FDetailWidgetRow& WarningRow = WarningCategory.AddCustomRow(SearchString)
@@ -55,7 +55,7 @@ FDetailWidgetRow& FSpriteDetailsCustomization::GenerateWarningRow(IDetailCategor
 void FSpriteDetailsCustomization::CustomizeDetails(IDetailLayoutBuilder& DetailLayout)
 {
 	// Make sure sprite properties are near the top
-	IDetailCategoryBuilder& SpriteCategory = DetailLayout.EditCategory("Sprite", TEXT(""), ECategoryPriority::TypeSpecific);
+	IDetailCategoryBuilder& SpriteCategory = DetailLayout.EditCategory("Sprite", FText::GetEmpty(), ECategoryPriority::TypeSpecific);
 	BuildSpriteSection(SpriteCategory, DetailLayout);
 
 	// Build the rendering category
@@ -95,7 +95,7 @@ void FSpriteDetailsCustomization::BuildRenderingSection(IDetailCategoryBuilder&
 	// Add the rendering geometry mode into the parent container (renamed)
 	const FString RenderGeometryTypePropertyPath = FString::Printf(TEXT("%s.%s"), GET_MEMBER_NAME_STRING_CHECKED(UPaperSprite, RenderGeometry), GET_MEMBER_NAME_STRING_CHECKED(FSpritePolygonCollection, GeometryType));
 	RenderingCategory.AddProperty(DetailLayout.GetProperty(*RenderGeometryTypePropertyPath))
-		.DisplayName(LOCTEXT("RenderGeometryType", "Render Geometry Type").ToString());
+		.DisplayName(LOCTEXT("RenderGeometryType", "Render Geometry Type"));
 
 	// Show the rendering geometry settings
 	TSharedRef<IPropertyHandle> RenderGeometry = DetailLayout.GetProperty(GET_MEMBER_NAME_CHECKED(UPaperSprite, RenderGeometry));
@@ -104,7 +104,7 @@ void FSpriteDetailsCustomization::BuildRenderingSection(IDetailCategoryBuilder&
 	// Add the render polygons into advanced (renamed)
 	const FString RenderGeometryPolygonsPropertyPath = FString::Printf(TEXT("%s.%s"), GET_MEMBER_NAME_STRING_CHECKED(UPaperSprite, RenderGeometry), GET_MEMBER_NAME_STRING_CHECKED(FSpritePolygonCollection, Polygons));
 	RenderingCategory.AddProperty(DetailLayout.GetProperty(*RenderGeometryPolygonsPropertyPath), EPropertyLocation::Advanced)
-		.DisplayName(LOCTEXT("RenderPolygons", "Render Polygons").ToString());
+		.DisplayName(LOCTEXT("RenderPolygons", "Render Polygons"));
 }
 
 void FSpriteDetailsCustomization::BuildCollisionSection(IDetailCategoryBuilder& CollisionCategory, IDetailLayoutBuilder& DetailLayout)
@@ -142,7 +142,7 @@ void FSpriteDetailsCustomization::BuildCollisionSection(IDetailCategoryBuilder&
 		CollisionGeometryTypeProperty->AddRestriction(PreventDicedRestriction.ToSharedRef());
 
 		CollisionCategory.AddProperty(CollisionGeometryTypeProperty)
-			.DisplayName(LOCTEXT("CollisionGeometryType", "Collision Geometry Type").ToString())
+			.DisplayName(LOCTEXT("CollisionGeometryType", "Collision Geometry Type"))
 			.Visibility(ParticipatesInPhysics);
 	}
 
@@ -157,7 +157,7 @@ void FSpriteDetailsCustomization::BuildCollisionSection(IDetailCategoryBuilder&
 	// Add the collision polygons into advanced (renamed)
 	const FString CollisionGeometryPolygonsPropertyPath = FString::Printf(TEXT("%s.%s"), GET_MEMBER_NAME_STRING_CHECKED(UPaperSprite, CollisionGeometry), GET_MEMBER_NAME_STRING_CHECKED(FSpritePolygonCollection, Polygons));
 	CollisionCategory.AddProperty(DetailLayout.GetProperty(*CollisionGeometryPolygonsPropertyPath), EPropertyLocation::Advanced)
-		.DisplayName(LOCTEXT("CollisionPolygons", "Collision Polygons").ToString())
+		.DisplayName(LOCTEXT("CollisionPolygons", "Collision Polygons"))
 		.Visibility(ParticipatesInPhysics);
 
 	// Show the default body instance (and only it) from the body setup (if it exists)
diff --git a/Engine/Plugins/2D/Paper2D/Source/Paper2DEditor/Private/TileMapEditing/PaperTileMapDetailsCustomization.cpp b/Engine/Plugins/2D/Paper2D/Source/Paper2DEditor/Private/TileMapEditing/PaperTileMapDetailsCustomization.cpp
--- a/Engine/Plugins/2D/Paper2D/Source/Paper2DEditor/Private/TileMapEditing/PaperTileMapDetailsCustomization.cpp
+++ b/Engine/Plugins/2D/Paper2D/Source/Paper2DEditor/Private/TileMapEditing/PaperTileMapDetailsCustomization.cpp
@@ -52,7 +52,7 @@ void FPaperTileMapDetailsCustomization::CustomizeDetails(IDetailLayoutBuilder& D
 	if (bEditingActor)
 	{
 		TileMapCategory
-		.AddCustomRow(TEXT("EnterEditMode"))
+		.AddCustomRow(LOCTEXT( "EnterTileMapEditMode", "Enter Edit Mode" ))
 		[
 			SNew(SVerticalBox)
 			+SVerticalBox::Slot()
@@ -73,12 +73,12 @@ void FPaperTileMapDetailsCustomization::CustomizeDetails(IDetailLayoutBuilder& D
 	{
 		IDetailCategoryBuilder& LayersCategory = DetailLayout.EditCategory("Tile Layers");
 
-		LayersCategory.AddCustomRow(TEXT("Tile layer list"))
+		LayersCategory.AddCustomRow(LOCTEXT("TileLayerList", "Tile layer list"))
 		[
 			SNew(STileLayerList, TileMap)
 		];
 
-		LayersCategory.AddCustomRow(TEXT("Add Layer, Add Collision Layer"))
+		LayersCategory.AddCustomRow(LOCTEXT("AddLayerOrCollisionLayer", "Add Layer, Add Collision Layer"))
 		[
 			SNew(SVerticalBox)
 			+SVerticalBox::Slot()
diff --git a/Engine/Plugins/Editor/MacGraphicsSwitching/Source/MacGraphicsSwitching/Private/MacGraphicsSwitchingSettingsDetails.cpp b/Engine/Plugins/Editor/MacGraphicsSwitching/Source/MacGraphicsSwitching/Private/MacGraphicsSwitchingSettingsDetails.cpp
--- a/Engine/Plugins/Editor/MacGraphicsSwitching/Source/MacGraphicsSwitching/Private/MacGraphicsSwitchingSettingsDetails.cpp
+++ b/Engine/Plugins/Editor/MacGraphicsSwitching/Source/MacGraphicsSwitching/Private/MacGraphicsSwitchingSettingsDetails.cpp
@@ -35,7 +35,7 @@ void FMacGraphicsSwitchingSettingsDetails::CustomizeDetails( IDetailLayoutBuilde
 	}
 	
 	IDetailCategoryBuilder& AccessorCategory = DetailLayout.EditCategory( "OpenGL" );
-	AccessorCategory.AddCustomRow( LOCTEXT("PreferredRenderer", "Preferred Renderer").ToString() )
+	AccessorCategory.AddCustomRow( LOCTEXT("PreferredRenderer", "Preferred Renderer") )
 	.NameContent()
 	[
 		PreferredRendererPropertyHandle->CreatePropertyNameWidget()
diff --git a/Engine/Plugins/Experimental/CodeView/Source/CodeView/Private/CodeViewPlugin.cpp b/Engine/Plugins/Experimental/CodeView/Source/CodeView/Private/CodeViewPlugin.cpp
--- a/Engine/Plugins/Experimental/CodeView/Source/CodeView/Private/CodeViewPlugin.cpp
+++ b/Engine/Plugins/Experimental/CodeView/Source/CodeView/Private/CodeViewPlugin.cpp
@@ -34,12 +34,12 @@ public:
 			// the widget is destroyed.
 			const bool bShouldInitiallyExpand = CodeViewWidget->IsReadyToPopulate();
 
-			DetailBuilder.EditCategory( "CodeView", NSLOCTEXT("ActorDetails", "CodeViewSection", "Code View").ToString(), ECategoryPriority::Uncommon )
+			DetailBuilder.EditCategory( "CodeView", NSLOCTEXT("ActorDetails", "CodeViewSection", "Code View"), ECategoryPriority::Uncommon )
 				.InitiallyCollapsed( !bShouldInitiallyExpand )
 				// The expansion state should not be restored
 				.RestoreExpansionState( false )
 				.OnExpansionChanged( FOnBooleanValueChanged::CreateSP( CodeViewWidget, &CodeView::SCodeView::OnDetailSectionExpansionChanged ) )
-				.AddCustomRow( NSLOCTEXT("ActorDetails", "CodeViewSection", "Code View").ToString() )
+				.AddCustomRow( NSLOCTEXT("ActorDetails", "CodeViewSection", "Code View") )
 				[
 					// @todo editcode1: Width of item is too big for detail view?!
 					CodeViewWidget
diff --git a/Engine/Source/Developer/Android/AndroidPlatformEditor/Private/AndroidTargetSettingsCustomization.cpp b/Engine/Source/Developer/Android/AndroidPlatformEditor/Private/AndroidTargetSettingsCustomization.cpp
--- a/Engine/Source/Developer/Android/AndroidPlatformEditor/Private/AndroidTargetSettingsCustomization.cpp
+++ b/Engine/Source/Developer/Android/AndroidPlatformEditor/Private/AndroidTargetSettingsCustomization.cpp
@@ -67,13 +67,13 @@ void FAndroidTargetSettingsCustomization::BuildAppManifestSection(IDetailLayoutB
 
 	SetupForPlatformAttribute = PlatformSetupMessage->GetReadyToGoAttribute();
 
-	AppManifestCategory.AddCustomRow(TEXT("Warning"), false)
+	AppManifestCategory.AddCustomRow(LOCTEXT("Warning", "Warning"), false)
 		.WholeRowWidget
 		[
 			PlatformSetupMessage
 		];
 
-	AppManifestCategory.AddCustomRow(TEXT("App Manifest Hyperlink"), false)
+	AppManifestCategory.AddCustomRow(LOCTEXT("AppManifestHyperlink", "App Manifest Hyperlink"), false)
 		.WholeRowWidget
 		[
 			SNew(SBox)
@@ -85,7 +85,7 @@ void FAndroidTargetSettingsCustomization::BuildAppManifestSection(IDetailLayoutB
 			]
 		];
 
-	AppManifestCategory.AddCustomRow(TEXT("App Manifest"), false)
+	AppManifestCategory.AddCustomRow(LOCTEXT("AppManifestLabel", "App Manifest"), false)
 		.IsEnabled(SetupForPlatformAttribute)
 		.NameContent()
 		[
@@ -132,13 +132,13 @@ void FAndroidTargetSettingsCustomization::BuildAppManifestSection(IDetailLayoutB
 
 	SetupForGooglePlayAttribute = GooglePlaySetupMessage->GetReadyToGoAttribute();
 
-	GooglePlayCategory.AddCustomRow(TEXT("Warning"), false)
+	GooglePlayCategory.AddCustomRow(LOCTEXT("Warning", "Warning"), false)
 		.WholeRowWidget
 		[
 			GooglePlaySetupMessage
 		];
 
-	GooglePlayCategory.AddCustomRow(TEXT("App ID Hyperlink"), false)
+	GooglePlayCategory.AddCustomRow(LOCTEXT("AppIDHyperlink", "App ID Hyperlink"), false)
 		.WholeRowWidget
 		[
 			SNew(SBox)
@@ -173,7 +173,7 @@ void FAndroidTargetSettingsCustomization::BuildIconSection(IDetailLayoutBuilder&
 	// Icon category
 	IDetailCategoryBuilder& IconCategory = DetailLayout.EditCategory(TEXT("Icons"));
 
-	IconCategory.AddCustomRow(TEXT("Icons Hyperlink"), false)
+	IconCategory.AddCustomRow(LOCTEXT("IconsHyperlink", "Icons Hyperlink"), false)
 		.WholeRowWidget
 		[
 			SNew(SBox)
@@ -190,7 +190,7 @@ void FAndroidTargetSettingsCustomization::BuildIconSection(IDetailLayoutBuilder&
 		const FString AutomaticImagePath = EngineAndroidPath / Info.IconPath;
 		const FString TargetImagePath = GameAndroidPath / Info.IconPath;
 
-		IconCategory.AddCustomRow(Info.IconName.ToString())
+		IconCategory.AddCustomRow(Info.IconName)
 		.NameContent()
 		[
 			SNew(SHorizontalBox)
diff --git a/Engine/Source/Developer/InternationalizationSettings/Private/InternationalizationSettingsModelDetails.cpp b/Engine/Source/Developer/InternationalizationSettings/Private/InternationalizationSettingsModelDetails.cpp
--- a/Engine/Source/Developer/InternationalizationSettings/Private/InternationalizationSettingsModelDetails.cpp
+++ b/Engine/Source/Developer/InternationalizationSettings/Private/InternationalizationSettingsModelDetails.cpp
@@ -130,7 +130,7 @@ void FInternationalizationSettingsModelDetails::CustomizeDetails( IDetailLayoutB
 	// For use in the Slate macros below, the type must be typedef'd to compile.
 	typedef FCulturePtr ThreadSafeCulturePtr;
 
-	CategoryBuilder.AddCustomRow("Language")
+	CategoryBuilder.AddCustomRow(LOCTEXT("EditorLanguageLabel", "Language"))
 	.NameContent()
 	[
 		SNew(SHorizontalBox)
@@ -163,7 +163,7 @@ void FInternationalizationSettingsModelDetails::CustomizeDetails( IDetailLayoutB
 
 	const FText RegionToolTipText = LOCTEXT("EditorRegionTooltip", "Change the Editor region (requires restart to take effect)");
 
-	CategoryBuilder.AddCustomRow("Region")
+	CategoryBuilder.AddCustomRow(LOCTEXT("EditorRegionLabel", "Region"))
 	.NameContent()
 	[
 		SNew(SHorizontalBox)
@@ -197,7 +197,7 @@ void FInternationalizationSettingsModelDetails::CustomizeDetails( IDetailLayoutB
 
 	const FText FieldNamesToolTipText = LOCTEXT("EditorFieldNamesTooltip", "Toggle showing localized field names (requires restart to take effect)");
 
-	CategoryBuilder.AddCustomRow("FieldNames")
+	CategoryBuilder.AddCustomRow(LOCTEXT("EditorFieldNamesLabel", "Use Localized Field Names"))
 	.NameContent()
 	[
 		SNew(SHorizontalBox)
@@ -220,7 +220,7 @@ void FInternationalizationSettingsModelDetails::CustomizeDetails( IDetailLayoutB
 		.OnCheckStateChanged(this, &FInternationalizationSettingsModelDetails::ShoudLoadLocalizedFieldNamesCheckChanged)
 	];
 
-	CategoryBuilder.AddCustomRow("RestartWarning")
+	CategoryBuilder.AddCustomRow(LOCTEXT("EditorRestartWarningLabel", "RestartWarning"))
 	.Visibility( TAttribute<EVisibility>(this, &FInternationalizationSettingsModelDetails::GetInternationalizationRestartRowVisibility) )
 	.WholeRowContent()
 	.HAlign(HAlign_Center)
diff --git a/Engine/Source/Developer/iOS/IOSPlatformEditor/Private/IOSTargetSettingsCustomization.cpp b/Engine/Source/Developer/iOS/IOSPlatformEditor/Private/IOSTargetSettingsCustomization.cpp
--- a/Engine/Source/Developer/iOS/IOSPlatformEditor/Private/IOSTargetSettingsCustomization.cpp
+++ b/Engine/Source/Developer/iOS/IOSPlatformEditor/Private/IOSTargetSettingsCustomization.cpp
@@ -134,7 +134,7 @@ void FIOSTargetSettingsCustomization::BuildPListSection(IDetailLayoutBuilder& De
 			]
 		];*/
 
-	ProvisionCategory.AddCustomRow(TEXT("Mobile Provision"), false)
+	ProvisionCategory.AddCustomRow(LOCTEXT("ProvisionLabel", "Provision"), false)
 		.NameContent()
 		[
 			SNew(SHorizontalBox)
@@ -174,7 +174,7 @@ void FIOSTargetSettingsCustomization::BuildPListSection(IDetailLayoutBuilder& De
 			]
 		];
 
-	ProvisionCategory.AddCustomRow(TEXT("Certificate"), false)
+	ProvisionCategory.AddCustomRow(LOCTEXT("CertificateLabel", "Certificate"), false)
 		.NameContent()
 		[
 			SNew(SHorizontalBox)
@@ -213,13 +213,13 @@ void FIOSTargetSettingsCustomization::BuildPListSection(IDetailLayoutBuilder& De
 					]
 			];
 
-	AppManifestCategory.AddCustomRow(TEXT("Warning"), false)
+	AppManifestCategory.AddCustomRow(LOCTEXT("Warning", "Warning"), false)
 		.WholeRowWidget
 		[
 			PlatformSetupMessage
 		];
 
-	AppManifestCategory.AddCustomRow(TEXT("Info.plist Hyperlink"), false)
+	AppManifestCategory.AddCustomRow(LOCTEXT("InfoPlistHyperlink", "Info.plist Hyperlink"), false)
 		.WholeRowWidget
 		[
 			SNew(SBox)
@@ -232,7 +232,7 @@ void FIOSTargetSettingsCustomization::BuildPListSection(IDetailLayoutBuilder& De
 		];
 
 
-	AppManifestCategory.AddCustomRow(TEXT("Info.plist"), false)
+	AppManifestCategory.AddCustomRow(LOCTEXT("InfoPlist", "Info.plist"), false)
 		.IsEnabled(SetupForPlatformAttribute)
 		.NameContent()
 		[
@@ -282,30 +282,30 @@ void FIOSTargetSettingsCustomization::BuildPListSection(IDetailLayoutBuilder& De
 			.ToolTip(Tip); \
 	}
 
-	SETUP_PLIST_PROP(BundleDisplayName, BundleCategory, TEXT("Specifies the the display name for the application. This will be displayed under the icon on the device."));
-	SETUP_PLIST_PROP(BundleName, BundleCategory, TEXT("Specifies the the name of the application bundle. This is the short name for the application bundle."));
-	SETUP_PLIST_PROP(BundleIdentifier, BundleCategory, TEXT("Specifies the bundle identifier for the application."));
-	SETUP_PLIST_PROP(VersionInfo, BundleCategory, TEXT("Specifies the version for the application."));
-	SETUP_PLIST_PROP(bSupportsPortraitOrientation, OrientationCategory, TEXT("Supports default portrait orientation. Landscape will not be supported."));
-	SETUP_PLIST_PROP(bSupportsUpsideDownOrientation, OrientationCategory, TEXT("Supports upside down portrait orientation. Landscape will not be supported."));
-	SETUP_PLIST_PROP(bSupportsLandscapeLeftOrientation, OrientationCategory, TEXT("Supports left landscape orientation. Protrait will not be supported."));
-	SETUP_PLIST_PROP(bSupportsLandscapeRightOrientation, OrientationCategory, TEXT("Supports right landscape orientation. Protrait will not be supported."));
+	SETUP_PLIST_PROP(BundleDisplayName, BundleCategory, LOCTEXT("BundleDisplayNameToolTip", "Specifies the the display name for the application. This will be displayed under the icon on the device."));
+	SETUP_PLIST_PROP(BundleName, BundleCategory, LOCTEXT("BundleNameToolTip", "Specifies the the name of the application bundle. This is the short name for the application bundle."));
+	SETUP_PLIST_PROP(BundleIdentifier, BundleCategory, LOCTEXT("BundleIdentifierToolTip", "Specifies the bundle identifier for the application."));
+	SETUP_PLIST_PROP(VersionInfo, BundleCategory, LOCTEXT("VersionInfoToolTip", "Specifies the version for the application."));
+	SETUP_PLIST_PROP(bSupportsPortraitOrientation, OrientationCategory, LOCTEXT("SupportsPortraitOrientationToolTip", "Supports default portrait orientation. Landscape will not be supported."));
+	SETUP_PLIST_PROP(bSupportsUpsideDownOrientation, OrientationCategory, LOCTEXT("SupportsUpsideDownOrientationToolTip", "Supports upside down portrait orientation. Landscape will not be supported."));
+	SETUP_PLIST_PROP(bSupportsLandscapeLeftOrientation, OrientationCategory, LOCTEXT("SupportsLandscapeLeftOrientationToolTip", "Supports left landscape orientation. Protrait will not be supported."));
+	SETUP_PLIST_PROP(bSupportsLandscapeRightOrientation, OrientationCategory, LOCTEXT("SupportsLandscapeRightOrientationToolTip", "Supports right landscape orientation. Protrait will not be supported."));
 	
-	SETUP_PLIST_PROP(bSupportsMetal, RenderCategory, TEXT("Whether or not to add support for Metal API (requires IOS8 and A7 processors)."));
-	SETUP_PLIST_PROP(bSupportsOpenGLES2, RenderCategory, TEXT("Whether or not to add support for OpenGL ES2 (if this is false, then your game should specify minimum IOS8 version and use \"metal\" instead of \"opengles-2\" in UIRequiredDeviceCapabilities)"));
+	SETUP_PLIST_PROP(bSupportsMetal, RenderCategory, LOCTEXT("SupportsMetalToolTip", "Whether or not to add support for Metal API (requires IOS8 and A7 processors)."));
+	SETUP_PLIST_PROP(bSupportsOpenGLES2, RenderCategory, LOCTEXT("SupportsOpenGLES2ToolTip", "Whether or not to add support for OpenGL ES2 (if this is false, then your game should specify minimum IOS8 version and use \"metal\" instead of \"opengles-2\" in UIRequiredDeviceCapabilities)"));
 
-	SETUP_PLIST_PROP(bSupportsIPad, DeviceCategory, TEXT("Whether or not to add support for iPad devices"));
-	SETUP_PLIST_PROP(bSupportsIPhone, DeviceCategory, TEXT("Whether or not to add support for iPhone devices"));
+	SETUP_PLIST_PROP(bSupportsIPad, DeviceCategory, LOCTEXT("SupportsIPadToolTip", "Whether or not to add support for iPad devices"));
+	SETUP_PLIST_PROP(bSupportsIPhone, DeviceCategory, LOCTEXT("SupportsIPhoneToolTip", "Whether or not to add support for iPhone devices"));
 
-	SETUP_PLIST_PROP(MinimumiOSVersion, OSInfoCategory, TEXT("WMinimum iOS version this game supports"));
+	SETUP_PLIST_PROP(MinimumiOSVersion, OSInfoCategory, LOCTEXT("MinimumiOSVersionToolTip", "Minimum iOS version this game supports"));
 
-	FString DisabledTip = TEXT("This requires GitHub source.");
-	SETUP_NONROCKET_PROP(bDevForArmV7, BuildCategory, TEXT("Enable ArmV7 support? (this will be used if all type are unchecked)"), DisabledTip);
-	SETUP_NONROCKET_PROP(bDevForArm64, BuildCategory, TEXT("Enable Arm64 support?"), DisabledTip);
-	SETUP_NONROCKET_PROP(bDevForArmV7S, BuildCategory, TEXT("Enable ArmV7s support?"), DisabledTip);
-	SETUP_NONROCKET_PROP(bShipForArmV7, BuildCategory, TEXT("Enable ArmV7 support? (this will be used if all type are unchecked)"), DisabledTip);
-	SETUP_NONROCKET_PROP(bShipForArm64, BuildCategory, TEXT("Enable Arm64 support?"), DisabledTip);
-	SETUP_NONROCKET_PROP(bShipForArmV7S, BuildCategory, TEXT("Enable ArmV7s support?"), DisabledTip);
+	const FText DisabledTip = LOCTEXT("GitHubSourceRequiredToolTip", "This requires GitHub source.");
+	SETUP_NONROCKET_PROP(bDevForArmV7, BuildCategory, LOCTEXT("DevForArmV7ToolTip", "Enable ArmV7 support? (this will be used if all type are unchecked)"), DisabledTip);
+	SETUP_NONROCKET_PROP(bDevForArm64, BuildCategory, LOCTEXT("DevForArm64ToolTip", "Enable Arm64 support?"), DisabledTip);
+	SETUP_NONROCKET_PROP(bDevForArmV7S, BuildCategory, LOCTEXT("DevForArmV7SToolTip", "Enable ArmV7s support?"), DisabledTip);
+	SETUP_NONROCKET_PROP(bShipForArmV7, BuildCategory, LOCTEXT("ShipForArmV7ToolTip", "Enable ArmV7 support? (this will be used if all type are unchecked)"), DisabledTip);
+	SETUP_NONROCKET_PROP(bShipForArm64, BuildCategory, LOCTEXT("ShipForArm64ToolTip", "Enable Arm64 support?"), DisabledTip);
+	SETUP_NONROCKET_PROP(bShipForArmV7S, BuildCategory, LOCTEXT("ShipForArmV7SToolTip", "Enable ArmV7s support?"), DisabledTip);
 
 #undef SETUP_PLIST_PROP
 #undef SETUP_NONROCKET_PROP
@@ -490,7 +490,7 @@ void FIOSTargetSettingsCustomization::BuildImageRow(IDetailLayoutBuilder& Detail
 	const FString AutomaticImagePath = EngineGraphicsPath / Info.IconPath;
 	const FString TargetImagePath = GameGraphicsPath / Info.IconPath;
 
-	Category.AddCustomRow(Info.IconName.ToString())
+	Category.AddCustomRow(Info.IconName)
 		.NameContent()
 		[
 			SNew(SHorizontalBox)
diff --git a/Engine/Source/Editor/BehaviorTreeEditor/Private/DetailCustomizations/BehaviorDecoratorDetails.cpp b/Engine/Source/Editor/BehaviorTreeEditor/Private/DetailCustomizations/BehaviorDecoratorDetails.cpp
--- a/Engine/Source/Editor/BehaviorTreeEditor/Private/DetailCustomizations/BehaviorDecoratorDetails.cpp
+++ b/Engine/Source/Editor/BehaviorTreeEditor/Private/DetailCustomizations/BehaviorDecoratorDetails.cpp
@@ -15,7 +15,7 @@ TSharedRef<IDetailCustomization> FBehaviorDecoratorDetails::MakeInstance()
 BEGIN_SLATE_FUNCTION_BUILD_OPTIMIZATION
 void FBehaviorDecoratorDetails::CustomizeDetails( IDetailLayoutBuilder& DetailLayout )
 {
-	FString AbortModeDesc = LOCTEXT("ObserverTitle","Observer aborts").ToString();
+	FText AbortModeDesc = LOCTEXT("ObserverTitle","Observer aborts");
 	PropUtils = &(DetailLayout.GetPropertyUtilities().Get());
 
 	TArray<TWeakObjectPtr<UObject> > EditedObjects;
diff --git a/Engine/Source/Editor/BehaviorTreeEditor/Private/DetailCustomizations/BlackboardDataDetails.cpp b/Engine/Source/Editor/BehaviorTreeEditor/Private/DetailCustomizations/BlackboardDataDetails.cpp
--- a/Engine/Source/Editor/BehaviorTreeEditor/Private/DetailCustomizations/BlackboardDataDetails.cpp
+++ b/Engine/Source/Editor/BehaviorTreeEditor/Private/DetailCustomizations/BlackboardDataDetails.cpp
@@ -4,6 +4,8 @@
 #include "BlackboardDataDetails.h"
 #include "Runtime/AIModule/Classes/BehaviorTree/BlackboardData.h"
 
+#define LOCTEXT_NAMESPACE "BlackboardDataDetails"
+
 TSharedRef<IDetailCustomization> FBlackboardDataDetails::MakeInstance(FOnGetSelectedBlackboardItemIndex InOnGetSelectedBlackboardItemIndex)
 {
 	return MakeShareable( new FBlackboardDataDetails(InOnGetSelectedBlackboardItemIndex) );
@@ -35,7 +37,7 @@ void FBlackboardDataDetails::CustomizeDetails( IDetailLayoutBuilder& DetailLayou
 
 			IDetailCategoryBuilder& DetailCategoryBuilder = DetailLayout.EditCategory("Key");
 			TSharedPtr<IPropertyHandle> EntryNameProperty = KeyHandle->GetChildHandle(GET_MEMBER_NAME_CHECKED(FBlackboardEntry, EntryName));
-			DetailCategoryBuilder.AddCustomRow(TEXT("Entry Name"))
+			DetailCategoryBuilder.AddCustomRow(LOCTEXT("EntryNameLabel", "Entry Name"))
 			.NameContent()
 			[
 				EntryNameProperty->CreatePropertyNameWidget()
@@ -62,3 +64,5 @@ void FBlackboardDataDetails::CustomizeDetails( IDetailLayoutBuilder& DetailLayou
 		}	
 	}
 }
+
+#undef LOCTEXT_NAMESPACE
diff --git a/Engine/Source/Editor/Blutility/Private/BlutilityDetailsPanel.cpp b/Engine/Source/Editor/Blutility/Private/BlutilityDetailsPanel.cpp
--- a/Engine/Source/Editor/Blutility/Private/BlutilityDetailsPanel.cpp
+++ b/Engine/Source/Editor/Blutility/Private/BlutilityDetailsPanel.cpp
@@ -68,7 +68,7 @@ void FEditorUtilityInstanceDetails::CustomizeDetails(IDetailLayoutBuilder& Detai
 		const APlacedEditorUtilityBase* PlacedActorCDO = Cast<const APlacedEditorUtilityBase>(Class->GetDefaultObject());
 		if (PlacedActorCDO)
 		{
-			ActionsCategory.AddCustomRow( PlacedActorCDO->HelpText )
+			ActionsCategory.AddCustomRow( FText::FromString(PlacedActorCDO->HelpText) )
 			[
 				SNew(STextBlock)
 				.Text(PlacedActorCDO->HelpText)
@@ -78,7 +78,7 @@ void FEditorUtilityInstanceDetails::CustomizeDetails(IDetailLayoutBuilder& Detai
 		const UGlobalEditorUtilityBase* GlobalBlutilityCDO = Cast<const UGlobalEditorUtilityBase>(Class->GetDefaultObject());
 		if (GlobalBlutilityCDO)
 		{
-			ActionsCategory.AddCustomRow( GlobalBlutilityCDO->HelpText )
+			ActionsCategory.AddCustomRow( FText::FromString(GlobalBlutilityCDO->HelpText) )
 			[
 				SNew(STextBlock)
 				.Text(GlobalBlutilityCDO->HelpText)
@@ -122,7 +122,7 @@ void FEditorUtilityInstanceDetails::CustomizeDetails(IDetailLayoutBuilder& Detai
 
 		if (NumButtons > 0)
 		{
-			ActionsCategory.AddCustomRow(TEXT(""))
+			ActionsCategory.AddCustomRow(FText::GetEmpty())
 			[
 				WrapBox
 			];
diff --git a/Engine/Source/Editor/DataTableEditor/Private/SRowEditor.cpp b/Engine/Source/Editor/DataTableEditor/Private/SRowEditor.cpp
--- a/Engine/Source/Editor/DataTableEditor/Private/SRowEditor.cpp
+++ b/Engine/Source/Editor/DataTableEditor/Private/SRowEditor.cpp
@@ -264,7 +264,7 @@ void SRowEditor::Construct(const FArguments& InArgs, UDataTable* Changed)
 		ViewArgs.bHideSelectionTip = false;
 		ViewArgs.bShowActorLabel = false;
 
-		StructureDetailsView = PropertyModule.CreateStructureDetailView(ViewArgs, CurrentRow, false, LOCTEXT("RowValue", "Row Value").ToString());
+		StructureDetailsView = PropertyModule.CreateStructureDetailView(ViewArgs, CurrentRow, false, LOCTEXT("RowValue", "Row Value"));
 		StructureDetailsView->GetOnFinishedChangingPropertiesDelegate().AddSP(this, &SRowEditor::OnFinishedChangingProperties);
 	}
 
diff --git a/Engine/Source/Editor/DetailCustomizations/Private/ActorDetails.cpp b/Engine/Source/Editor/DetailCustomizations/Private/ActorDetails.cpp
--- a/Engine/Source/Editor/DetailCustomizations/Private/ActorDetails.cpp
+++ b/Engine/Source/Editor/DetailCustomizations/Private/ActorDetails.cpp
@@ -403,10 +403,10 @@ void FActorDetails::AddLayersCategory( IDetailLayoutBuilder& DetailBuilder )
 
 	FLayersModule& LayersModule = FModuleManager::LoadModuleChecked< FLayersModule >( TEXT("Layers") );
 
-	const FString LayerCategory = LOCTEXT("LayersCategory", "Layers").ToString();
+	const FText LayerCategory = LOCTEXT("LayersCategory", "Layers");
 
 	DetailBuilder.EditCategory( "Layers", LayerCategory, ECategoryPriority::Uncommon )
-	.AddCustomRow( TEXT("") )
+	.AddCustomRow( FText::GetEmpty() )
 	[
 		LayersModule.CreateLayerCloud( SelectedActors )
 	];
@@ -428,7 +428,7 @@ void FActorDetails::AddTransformCategory( IDetailLayoutBuilder& DetailBuilder )
 	
 	TSharedRef<FComponentTransformDetails> TransformDetails = MakeShareable( new FComponentTransformDetails( DetailBuilder.GetDetailsView().GetSelectedObjects(), SelectedActorInfo, DetailBuilder ) );
 
-	IDetailCategoryBuilder& TransformCategory = DetailBuilder.EditCategory( "TransformCommon", LOCTEXT("TransformCommonCategory", "Transform").ToString(), ECategoryPriority::Transform );
+	IDetailCategoryBuilder& TransformCategory = DetailBuilder.EditCategory( "TransformCommon", LOCTEXT("TransformCommonCategory", "Transform"), ECategoryPriority::Transform );
 
 	TransformCategory.AddCustomBuilder( TransformDetails );
 }
@@ -442,9 +442,9 @@ void FActorDetails::AddExperimentalWarningCategory( IDetailLayoutBuilder& Detail
 		const bool bExperimental = SelectedActorInfo.bHaveExperimentalClass;
 
 		const FName CategoryName(TEXT("Warning"));
-		const FString CategoryDisplayName = LOCTEXT("WarningCategoryDisplayName", "Warning").ToString();
+		const FText CategoryDisplayName = LOCTEXT("WarningCategoryDisplayName", "Warning");
 		const FText WarningText = bExperimental ? LOCTEXT("ExperimentalClassWarning", "Uses experimental class") : LOCTEXT("EarlyAccessClassWarning", "Uses early access class");
-		const FString SearchString = WarningText.ToString();
+		const FText SearchString = WarningText;
 		const FText Tooltip = bExperimental ? LOCTEXT("ExperimentalClassTooltip", "Here be dragons!  Uses one or more unsupported 'experimental' classes") : LOCTEXT("EarlyAccessClassTooltip", "Uses one or more 'early access' classes");
 		const FString ExcerptName = bExperimental ? TEXT("ActorUsesExperimentalClass") : TEXT("ActorUsesEarlyAccessClass");
 		const FSlateBrush* WarningIcon = FEditorStyle::GetBrush(bExperimental ? "PropertyEditor.ExperimentalClass" : "PropertyEditor.EarlyAccessClass");
@@ -497,7 +497,7 @@ void FActorDetails::AddActorCategory( IDetailLayoutBuilder& DetailBuilder, const
 	const FSelectedActorInfo& SelectedActorInfo = DetailBuilder.GetDetailsView().GetSelectedActorInfo();
 	TSharedPtr<SVerticalBox> LevelBox;
 
-	IDetailCategoryBuilder& ActorCategory = DetailBuilder.EditCategory("Actor", TEXT(""), ECategoryPriority::Uncommon );
+	IDetailCategoryBuilder& ActorCategory = DetailBuilder.EditCategory("Actor", FText::GetEmpty(), ECategoryPriority::Uncommon );
 
 
 	// Create the info buttons per level
@@ -521,7 +521,7 @@ void FActorDetails::AddActorCategory( IDetailLayoutBuilder& DetailBuilder, const
 		// Create the row for this level
 		TWeakObjectPtr<ULevel> WeakLevelPtr = Level;
 
-		ActorCategory.AddCustomRow( LOCTEXT("SelectionFilter", "Selected").ToString() )
+		ActorCategory.AddCustomRow( LOCTEXT("SelectionFilter", "Selected") )
 		.NameContent()
 		[
 			SNew(SHyperlink)
@@ -543,7 +543,7 @@ void FActorDetails::AddActorCategory( IDetailLayoutBuilder& DetailBuilder, const
 	// WorldSettings should never convert to another class type
 	if( SelectedActorInfo.SelectionClass != AWorldSettings::StaticClass() && SelectedActorInfo.HasConvertableAsset() )
 	{
-		ActorCategory.AddCustomRow( LOCTEXT("ConvertMenu", "Convert").ToString() )
+		ActorCategory.AddCustomRow( LOCTEXT("ConvertMenu", "Convert") )
 		.NameContent()
 		[
 			SNew(STextBlock)
@@ -563,7 +563,7 @@ void FActorDetails::AddBlueprintCategory( IDetailLayoutBuilder& DetailBuilder, c
 	// Create the overarching structure and expose the vertical box to add entries per-blueprint
 	TSharedPtr<SVerticalBox> Box;
 
-	IDetailCategoryBuilder& BlueprintCategory = DetailBuilder.EditCategory("Blueprint", NSLOCTEXT("BlueprintDetails", "BlueprintTitle", "Blueprint").ToString(), ECategoryPriority::Uncommon );
+	IDetailCategoryBuilder& BlueprintCategory = DetailBuilder.EditCategory("Blueprint", NSLOCTEXT("BlueprintDetails", "BlueprintTitle", "Blueprint"), ECategoryPriority::Uncommon );
 
 	// Only show the blueprints section if a single actor is selected
 	if ( SelectedActors.Num() > 0 )
@@ -592,7 +592,7 @@ void FActorDetails::AddBlueprintCategory( IDetailLayoutBuilder& DetailBuilder, c
 void FActorDetails::AddBlutilityCategory( IDetailLayoutBuilder& DetailBuilder, const TMap<UBlueprint*, UObject*>& UniqueBlueprints )
 {
 	// Create the Blutilities Category
-	IDetailCategoryBuilder& BlutilitiesCategory = DetailBuilder.EditCategory("Blutilities", NSLOCTEXT("Blutilities", "BlutilityTitle", "Blutilities").ToString(), ECategoryPriority::Uncommon );
+	IDetailCategoryBuilder& BlutilitiesCategory = DetailBuilder.EditCategory("Blutilities", NSLOCTEXT("Blutilities", "BlutilityTitle", "Blutilities"), ECategoryPriority::Uncommon );
 
 	// Only show the bluetilities section if a single actor is selected
 	if ( SelectedActors.Num() > 0 && DoesActorHaveBlutiltyFunctions() )
@@ -614,7 +614,7 @@ void FActorDetails::AddBlutilityCategory( IDetailLayoutBuilder& DetailBuilder, c
 		const FText ButtonLabel = FText::Format( NSLOCTEXT( "Blutilities", "CallInEditor_ButtonLabel", "Run" ), Args );
 		const FText ButtonToolTip = FText::Format( NSLOCTEXT( "Blutilities", "CallInEditor_ButtonTooltip", "Run Selected Blutility Function on {ActorLabel}" ), Args );
 		// Build Content
-		BlutilitiesCategory.AddCustomRow( NSLOCTEXT( "Blutilities", "CallInEditorHeader", "Blutility Functions").ToString())
+		BlutilitiesCategory.AddCustomRow( NSLOCTEXT( "Blutilities", "CallInEditorHeader", "Blutility Functions"))
 		.WholeRowContent()
 		[
 			SNew(SHorizontalBox)
@@ -670,9 +670,9 @@ void FActorDetails::AddComponentsCategory( IDetailLayoutBuilder& DetailBuilder )
 		SelectedActor = Actors[0].Get();
 	}
 
-	DetailBuilder.EditCategory( "Components", NSLOCTEXT("ActorDetails", "ComponentsSection", "Components").ToString(), ECategoryPriority::Uncommon )
+	DetailBuilder.EditCategory( "Components", NSLOCTEXT("ActorDetails", "ComponentsSection", "Components"), ECategoryPriority::Uncommon )
 		.InitiallyCollapsed( true )
-		.AddCustomRow( NSLOCTEXT("ActorDetails", "ComponentsSection", "Components").ToString() )
+		.AddCustomRow( NSLOCTEXT("ActorDetails", "ComponentsSection", "Components") )
 		[
 			SNew( SComponentsTree, SelectedActor )
 		];
@@ -702,7 +702,7 @@ void FActorDetails::AddUtilityBlueprintRows( IDetailCategoryBuilder& BlueprintCa
 		const FText FindButtonLabel = LOCTEXT( "FindInLevelScript", "Find in Level Blueprint" );
 		const FText FindButtonToolTip = FText::Format( LOCTEXT( "FindInLevelScript_ToolTip", "Search for uses of {ActorLabel} in the Level Blueprint" ), Args );
 
-		BlueprintCategory.AddCustomRow( LOCTEXT( "FindInLevelScript_Filter", "Find in Level Script" ).ToString() )
+		BlueprintCategory.AddCustomRow( LOCTEXT( "FindInLevelScript_Filter", "Find in Level Script" ) )
 		.WholeRowContent()
 		[
 			SNew( SHorizontalBox )
@@ -733,7 +733,7 @@ void FActorDetails::AddUtilityBlueprintRows( IDetailCategoryBuilder& BlueprintCa
 		const FText AddEventLabel = LOCTEXT( "ScriptingEvents", "Add Level Events" );
 		const FText AddEventToolTip = FText::Format( LOCTEXT( "ScripingEvents_Tooltip", "Adds or views events for {ActorLabel} in the Level Blueprint" ), Args );
 
-		BlueprintCategory.AddCustomRow( LOCTEXT( "ScriptingEvents_Filter", "Level Blueprint Events" ).ToString() )
+		BlueprintCategory.AddCustomRow( LOCTEXT( "ScriptingEvents_Filter", "Level Blueprint Events" ) )
 		.WholeRowContent()
 		[
 			SNew( SHorizontalBox )
@@ -763,7 +763,7 @@ void FActorDetails::AddUtilityBlueprintRows( IDetailCategoryBuilder& BlueprintCa
 	}
 
 	// Create the row for this blueprint
-	BlueprintCategory.AddCustomRow( NSLOCTEXT( "BlueprintDetails", "BlueprintFilter", "Blueprints" ).ToString() )
+	BlueprintCategory.AddCustomRow( NSLOCTEXT( "BlueprintDetails", "BlueprintFilter", "Blueprints" ) )
 	.WholeRowContent()
 	[
 		SNew( SHorizontalBox )
@@ -858,7 +858,7 @@ void FActorDetails::AddSingleBlueprintRow( IDetailCategoryBuilder& BlueprintCate
 	FText CurrentBlueprintLong = FText::Format( FText::FromString( TEXT( "{0} {1}" ) ), Name, Description );
 
 	// Blueprint Details
-	BlueprintCategory.AddCustomRow( NSLOCTEXT("BlueprintDetails", "BlueprintFilter", "Blueprints").ToString() )
+	BlueprintCategory.AddCustomRow( NSLOCTEXT("BlueprintDetails", "BlueprintFilter", "Blueprints") )
 	.Visibility( InBlueprint ? EVisibility::Visible : EVisibility::Collapsed )
 	.NameContent()
 	[
@@ -993,7 +993,7 @@ void FActorDetails::AddSingleBlueprintRow( IDetailCategoryBuilder& BlueprintCate
 
 	if(!InBlueprint)
 	{
-		BlueprintCategory.AddCustomRow( LOCTEXT("CreateBlueprintFilterString", "Create Blueprint").ToString(), true )
+		BlueprintCategory.AddCustomRow( LOCTEXT("CreateBlueprintFilterString", "Create Blueprint"), true )
 		.WholeRowContent()
 		.MinDesiredWidth(200)
 		.MaxDesiredWidth(200)
diff --git a/Engine/Source/Editor/DetailCustomizations/Private/ActorMaterialCategory.cpp b/Engine/Source/Editor/DetailCustomizations/Private/ActorMaterialCategory.cpp
--- a/Engine/Source/Editor/DetailCustomizations/Private/ActorMaterialCategory.cpp
+++ b/Engine/Source/Editor/DetailCustomizations/Private/ActorMaterialCategory.cpp
@@ -246,7 +246,7 @@ void FActorMaterialCategory::Create( IDetailLayoutBuilder& DetailBuilder )
 	if( bAnyMaterialsToDisplay )
 	{
 		// Make a category for the materials.
-		IDetailCategoryBuilder& MaterialCategory = DetailBuilder.EditCategory("Materials", TEXT(""), ECategoryPriority::TypeSpecific );
+		IDetailCategoryBuilder& MaterialCategory = DetailBuilder.EditCategory("Materials", FText::GetEmpty(), ECategoryPriority::TypeSpecific );
 
 		MaterialCategory.AddCustomBuilder( MaterialList );
 	}
diff --git a/Engine/Source/Editor/DetailCustomizations/Private/AmbientSoundDetails.cpp b/Engine/Source/Editor/DetailCustomizations/Private/AmbientSoundDetails.cpp
--- a/Engine/Source/Editor/DetailCustomizations/Private/AmbientSoundDetails.cpp
+++ b/Engine/Source/Editor/DetailCustomizations/Private/AmbientSoundDetails.cpp
@@ -34,8 +34,8 @@ void FAmbientSoundDetails::CustomizeDetails( IDetailLayoutBuilder& DetailBuilder
 		}
 	}
 
-	DetailBuilder.EditCategory( "Sound", TEXT(""), ECategoryPriority::Important )
-		.AddCustomRow( TEXT("") )
+	DetailBuilder.EditCategory( "Sound", FText::GetEmpty(), ECategoryPriority::Important )
+		.AddCustomRow( FText::GetEmpty() )
 		[
 			SNew(SVerticalBox)
 			+ SVerticalBox::Slot()
@@ -100,8 +100,8 @@ void FAmbientSoundDetails::CustomizeDetails( IDetailLayoutBuilder& DetailBuilder
 			]
 		];
 
-	DetailBuilder.EditCategory("Attenuation", TEXT(""), ECategoryPriority::TypeSpecific);
-	DetailBuilder.EditCategory("Modulation", TEXT(""), ECategoryPriority::TypeSpecific);
+	DetailBuilder.EditCategory("Attenuation", FText::GetEmpty(), ECategoryPriority::TypeSpecific);
+	DetailBuilder.EditCategory("Modulation", FText::GetEmpty(), ECategoryPriority::TypeSpecific);
 }
 
 bool FAmbientSoundDetails::IsEditSoundCueEnabled() const
diff --git a/Engine/Source/Editor/DetailCustomizations/Private/AnimMontageSegmentDetails.cpp b/Engine/Source/Editor/DetailCustomizations/Private/AnimMontageSegmentDetails.cpp
--- a/Engine/Source/Editor/DetailCustomizations/Private/AnimMontageSegmentDetails.cpp
+++ b/Engine/Source/Editor/DetailCustomizations/Private/AnimMontageSegmentDetails.cpp
@@ -60,15 +60,15 @@ TSharedRef<IDetailCustomization> FAnimMontageSegmentDetails::MakeInstance()
 
 void FAnimMontageSegmentDetails::CustomizeDetails( IDetailLayoutBuilder& DetailBuilder )
 {
-	IDetailCategoryBuilder& SegmentCategory = DetailBuilder.EditCategory("Animation Segment", TEXT("Animation Segment") );
+	IDetailCategoryBuilder& SegmentCategory = DetailBuilder.EditCategory("Animation Segment", LOCTEXT("AnimationSegmentCategoryTitle", "Animation Segment") );
 
-	SegmentCategory.AddProperty("AnimSegment.AnimReference").DisplayName( TEXT("Animation Reference") );
+	SegmentCategory.AddProperty("AnimSegment.AnimReference").DisplayName( LOCTEXT("AnimationReferenceLabel", "Animation Reference") );
 
-	SegmentCategory.AddProperty("AnimSegment.AnimStartTime").DisplayName( TEXT("Start Time") );
-	SegmentCategory.AddProperty("AnimSegment.AnimEndTime").DisplayName( TEXT("End Time") );
+	SegmentCategory.AddProperty("AnimSegment.AnimStartTime").DisplayName( LOCTEXT("StartTimeLabel", "Start Time") );
+	SegmentCategory.AddProperty("AnimSegment.AnimEndTime").DisplayName( LOCTEXT("EndTimeLabel", "End Time") );
 
-	SegmentCategory.AddProperty("AnimSegment.AnimPlayRate").DisplayName( TEXT("Play Rate") );
-	SegmentCategory.AddProperty("AnimSegment.LoopingCount").DisplayName( TEXT("Loop Count") );
+	SegmentCategory.AddProperty("AnimSegment.AnimPlayRate").DisplayName( LOCTEXT("PlayRateLabel", "Play Rate") );
+	SegmentCategory.AddProperty("AnimSegment.LoopingCount").DisplayName( LOCTEXT("LoopCountLabel", "Loop Count") );
 
 	TSharedPtr<IPropertyHandle> InPropertyHandle = DetailBuilder.GetProperty("AnimSegment.AnimReference");
 	UObject *Object = NULL;
@@ -81,7 +81,7 @@ void FAnimMontageSegmentDetails::CustomizeDetails( IDetailLayoutBuilder& DetailB
 		Skeleton = AnimRef->GetSkeleton();
 	}
 
-	SegmentCategory.AddCustomRow(TEXT(""), false)
+	SegmentCategory.AddCustomRow(FText::GetEmpty(), false)
 	[
 		SNew(SAnimationSegmentViewport)
 		.Skeleton(Skeleton)
diff --git a/Engine/Source/Editor/DetailCustomizations/Private/AnimStateNodeDetails.cpp b/Engine/Source/Editor/DetailCustomizations/Private/AnimStateNodeDetails.cpp
--- a/Engine/Source/Editor/DetailCustomizations/Private/AnimStateNodeDetails.cpp
+++ b/Engine/Source/Editor/DetailCustomizations/Private/AnimStateNodeDetails.cpp
@@ -15,10 +15,10 @@ TSharedRef<IDetailCustomization> FAnimStateNodeDetails::MakeInstance()
 
 void FAnimStateNodeDetails::CustomizeDetails( IDetailLayoutBuilder& DetailBuilder )
 {
-	IDetailCategoryBuilder& SegmentCategory = DetailBuilder.EditCategory("Animation State", TEXT("Animation State") );
+	IDetailCategoryBuilder& SegmentCategory = DetailBuilder.EditCategory("Animation State", LOCTEXT("AnimationStateCategoryTitle", "Animation State") );
 
 	{
-		SegmentCategory.AddCustomRow( TEXT("Entered State Event") )
+		SegmentCategory.AddCustomRow( LOCTEXT("EnteredAnimationStateEventLabel", "Entered State Event") )
 		[
 			SNew( STextBlock )
 			.Text( LOCTEXT("EnteredAnimationStateEventLabel", "Entered State Event") )
@@ -29,7 +29,7 @@ void FAnimStateNodeDetails::CustomizeDetails( IDetailLayoutBuilder& DetailBuilde
 
 
 	{
-		SegmentCategory.AddCustomRow( TEXT("Left State Event") )
+		SegmentCategory.AddCustomRow( LOCTEXT("ExitedAnimationStateEventLabel", "Left State Event") )
 		[
 			SNew( STextBlock )
 			.Text( LOCTEXT("ExitedAnimationStateEventLabel", "Left State Event")  )
@@ -40,7 +40,7 @@ void FAnimStateNodeDetails::CustomizeDetails( IDetailLayoutBuilder& DetailBuilde
 
 	{
 
-		SegmentCategory.AddCustomRow( TEXT("Fully Blended State Event") )
+		SegmentCategory.AddCustomRow( LOCTEXT("FullyBlendedAnimationStateEventLabel", "Fully Blended State Event") )
 		[
 			SNew( STextBlock )
 			.Text( LOCTEXT("FullyBlendedAnimationStateEventLabel", "Fully Blended State Event") )
diff --git a/Engine/Source/Editor/DetailCustomizations/Private/AnimTransitionNodeDetails.cpp b/Engine/Source/Editor/DetailCustomizations/Private/AnimTransitionNodeDetails.cpp
--- a/Engine/Source/Editor/DetailCustomizations/Private/AnimTransitionNodeDetails.cpp
+++ b/Engine/Source/Editor/DetailCustomizations/Private/AnimTransitionNodeDetails.cpp
@@ -46,7 +46,7 @@ void FAnimTransitionNodeDetails::CustomizeDetails( IDetailLayoutBuilder& DetailB
 
 	////////////////////////////////////////
 
-	IDetailCategoryBuilder& TransitionCategory = DetailBuilder.EditCategory("Transition", TEXT("Transition") );
+	IDetailCategoryBuilder& TransitionCategory = DetailBuilder.EditCategory("Transition", LOCTEXT("TransitionCategoryTitle", "Transition") );
 
 	if (bTransitionToConduit)
 	{
@@ -60,22 +60,22 @@ void FAnimTransitionNodeDetails::CustomizeDetails( IDetailLayoutBuilder& DetailB
 	}
 	else
 	{
-		TransitionCategory.AddCustomRow( TEXT("Transition") )
+		TransitionCategory.AddCustomRow( LOCTEXT("TransitionEventPropertiesCategoryLabel", "Transition") )
 		[
 			SNew( STextBlock )
 			.Text( LOCTEXT("TransitionEventPropertiesCategoryLabel", "Transition") )
 			.Font( IDetailLayoutBuilder::GetDetailFontBold() )
 		];
 
-		TransitionCategory.AddProperty(GET_MEMBER_NAME_CHECKED(UAnimStateTransitionNode, PriorityOrder)).DisplayName(TEXT("Priority Order"));
-		TransitionCategory.AddProperty(GET_MEMBER_NAME_CHECKED(UAnimStateTransitionNode, Bidirectional)).DisplayName( TEXT("Bidirectional"));
-		TransitionCategory.AddProperty(GET_MEMBER_NAME_CHECKED(UAnimStateTransitionNode, LogicType)).DisplayName( TEXT("Blend Logic") );
+		TransitionCategory.AddProperty(GET_MEMBER_NAME_CHECKED(UAnimStateTransitionNode, PriorityOrder)).DisplayName(LOCTEXT("PriorityOrderLabel", "Priority Order"));
+		TransitionCategory.AddProperty(GET_MEMBER_NAME_CHECKED(UAnimStateTransitionNode, Bidirectional)).DisplayName(LOCTEXT("BidirectionalLabel", "Bidirectional"));
+		TransitionCategory.AddProperty(GET_MEMBER_NAME_CHECKED(UAnimStateTransitionNode, LogicType)).DisplayName(LOCTEXT("BlendLogicLabel", "Blend Logic") );
 
 		UAnimStateTransitionNode* TransNode = TransitionNode.Get();
 		if (TransitionNode != NULL)
 		{
 			// The sharing option for the rule
-			TransitionCategory.AddCustomRow( TEXT("Transition Rule Sharing") )
+			TransitionCategory.AddCustomRow( LOCTEXT("TransitionRuleSharingLabel", "Transition Rule Sharing") )
 			[
 				GetWidgetForInlineShareMenu(TEXT("Transition Rule Sharing"), TransNode->SharedRulesName, TransNode->bSharedRules,
 					FOnClicked::CreateSP(this, &FAnimTransitionNodeDetails::OnPromoteToSharedClick, true),
@@ -101,7 +101,7 @@ void FAnimTransitionNodeDetails::CustomizeDetails( IDetailLayoutBuilder& DetailB
 					CanExecPin = ResultNode->FindPin(TEXT("bCanEnterTransition"));
 				}
 			}
-			TransitionCategory.AddCustomRow( CanExecPin ? CanExecPin->PinFriendlyName.ToString() : FString() )
+			TransitionCategory.AddCustomRow( CanExecPin ? CanExecPin->PinFriendlyName : FText::GetEmpty() )
 			[
 				SNew(SKismetLinearExpression, CanExecPin)
 			];
@@ -110,11 +110,11 @@ void FAnimTransitionNodeDetails::CustomizeDetails( IDetailLayoutBuilder& DetailB
 
 
 
-		IDetailCategoryBuilder& CrossfadeCategory = DetailBuilder.EditCategory("BlendSettings", TEXT("BlendSettings") );
+		IDetailCategoryBuilder& CrossfadeCategory = DetailBuilder.EditCategory("BlendSettings", LOCTEXT("BlendSettingsCategoryTitle", "BlendSettings") );
 		if (TransitionNode != NULL)
 		{
 			// The sharing option for the crossfade settings
-			CrossfadeCategory.AddCustomRow( TEXT("Transition Crossfade Sharing") )
+			CrossfadeCategory.AddCustomRow( LOCTEXT("TransitionCrossfadeSharingLabel", "Transition Crossfade Sharing") )
 			[
 				GetWidgetForInlineShareMenu(TEXT("Transition Crossfade Sharing"), TransNode->SharedCrossfadeName, TransNode->bSharedCrossfade,
 					FOnClicked::CreateSP(this, &FAnimTransitionNodeDetails::OnPromoteToSharedClick, false), 
@@ -124,11 +124,11 @@ void FAnimTransitionNodeDetails::CustomizeDetails( IDetailLayoutBuilder& DetailB
 		}
 
 		//@TODO: Gate editing these on shared non-authorative ones
-		CrossfadeCategory.AddProperty(GET_MEMBER_NAME_CHECKED(UAnimStateTransitionNode, CrossfadeDuration)).DisplayName( TEXT("Duration") );
-		CrossfadeCategory.AddProperty(GET_MEMBER_NAME_CHECKED(UAnimStateTransitionNode, CrossfadeMode)).DisplayName( TEXT("Mode") );
+		CrossfadeCategory.AddProperty(GET_MEMBER_NAME_CHECKED(UAnimStateTransitionNode, CrossfadeDuration)).DisplayName( LOCTEXT("DurationLabel", "Duration") );
+		CrossfadeCategory.AddProperty(GET_MEMBER_NAME_CHECKED(UAnimStateTransitionNode, CrossfadeMode)).DisplayName( LOCTEXT("ModeLabel", "Mode") );
 
 		// Add a button that is only visible when blend logic type is custom
-		CrossfadeCategory.AddCustomRow( LOCTEXT("EditBlendGraph", "Edit Blend Graph").ToString() )
+		CrossfadeCategory.AddCustomRow( LOCTEXT("EditBlendGraph", "Edit Blend Graph") )
 		[
 			SNew( SHorizontalBox )
 			+SHorizontalBox::Slot()
@@ -148,9 +148,9 @@ void FAnimTransitionNodeDetails::CustomizeDetails( IDetailLayoutBuilder& DetailB
 
 
 
-		IDetailCategoryBuilder& NotificationCategory = DetailBuilder.EditCategory("Notifications", TEXT("Notifications") );
+		IDetailCategoryBuilder& NotificationCategory = DetailBuilder.EditCategory("Notifications", LOCTEXT("NotificationsCategoryTitle", "Notifications") );
 
-		NotificationCategory.AddCustomRow( TEXT("Start Transition Event") )
+		NotificationCategory.AddCustomRow( LOCTEXT("StartTransitionEventPropertiesCategoryLabel", "Start Transition Event") )
 		[
 			SNew( STextBlock )
 			.Text( LOCTEXT("StartTransitionEventPropertiesCategoryLabel", "Start Transition Event") )
@@ -159,7 +159,7 @@ void FAnimTransitionNodeDetails::CustomizeDetails( IDetailLayoutBuilder& DetailB
 		CreateTransitionEventPropertyWidgets(NotificationCategory, TEXT("TransitionStart"));
 
 
-		NotificationCategory.AddCustomRow( TEXT("End Transition Event" ) ) 
+		NotificationCategory.AddCustomRow( LOCTEXT("EndTransitionEventPropertiesCategoryLabel", "End Transition Event" ) ) 
 		[
 			SNew( STextBlock )
 			.Text( LOCTEXT("EndTransitionEventPropertiesCategoryLabel", "End Transition Event" ) )
@@ -167,7 +167,7 @@ void FAnimTransitionNodeDetails::CustomizeDetails( IDetailLayoutBuilder& DetailB
 		];
 		CreateTransitionEventPropertyWidgets(NotificationCategory, TEXT("TransitionEnd"));
 
-		NotificationCategory.AddCustomRow( TEXT("Interrupt Transition Event") )
+		NotificationCategory.AddCustomRow( LOCTEXT("InterruptTransitionEventPropertiesCategoryLabel", "Interrupt Transition Event") )
 		[
 			SNew( STextBlock )
 			.Text( LOCTEXT("InterruptTransitionEventPropertiesCategoryLabel", "Interrupt Transition Event") )
@@ -369,7 +369,7 @@ void FAnimTransitionNodeDetails::CreateTransitionEventPropertyWidgets(IDetailCat
 	TSharedPtr<IPropertyHandle> NameProperty = TransitionCategory.GetParentLayout().GetProperty(*(TransitionName + TEXT(".NotifyName")));
 
 	TransitionCategory.AddProperty( NameProperty )
-		.DisplayName( LOCTEXT("CreateTransition_CustomBlueprintEvent", "Custom Blueprint Event").ToString() );
+		.DisplayName( LOCTEXT("CreateTransition_CustomBlueprintEvent", "Custom Blueprint Event") );
 }
 
 TSharedRef<SWidget> FAnimTransitionNodeDetails::GetWidgetForInlineShareMenu(FString TypeName, FString SharedName, bool bIsCurrentlyShared,  FOnClicked PromoteClick, FOnClicked DemoteClick, FOnGetContent GetContentMenu)
diff --git a/Engine/Source/Editor/DetailCustomizations/Private/AttenuationSettingsCustomizations.cpp b/Engine/Source/Editor/DetailCustomizations/Private/AttenuationSettingsCustomizations.cpp
--- a/Engine/Source/Editor/DetailCustomizations/Private/AttenuationSettingsCustomizations.cpp
+++ b/Engine/Source/Editor/DetailCustomizations/Private/AttenuationSettingsCustomizations.cpp
@@ -13,8 +13,8 @@ void FAttenuationSettingsCustomization::CustomizeHeader( TSharedRef<class IPrope
 {
 	// We'll set up reset to default ourselves
 	const bool bDisplayResetToDefault = false;
-	const FString DisplayNameOverride = TEXT("");
-	const FString DisplayToolTipOverride = TEXT("");
+	const FText DisplayNameOverride = FText::GetEmpty();
+	const FText DisplayToolTipOverride = FText::GetEmpty();
 
 	HeaderRow
 		.NameContent()
@@ -85,10 +85,10 @@ void FAttenuationSettingsCustomization::CustomizeChildren( TSharedRef<IPropertyH
 	
 	ChildBuilder.AddChildProperty(AttenuationExtentsHandle)
 		.Visibility(TAttribute<EVisibility>(this, &FAttenuationSettingsCustomization::IsBoxSelected))
-		.DisplayName(TEXT("Extents"))
-		.ToolTip(NSLOCTEXT("AttenuationSettings", "BoxExtents", "The dimensions of the of the box.").ToString());
+		.DisplayName(NSLOCTEXT("AttenuationSettings", "BoxExtentsLabel", "Extents"))
+		.ToolTip(NSLOCTEXT("AttenuationSettings", "BoxExtents", "The dimensions of the of the box."));
 
-	ChildBuilder.AddChildContent(TEXT("Radius"))
+	ChildBuilder.AddChildContent(NSLOCTEXT("AttenuationSettings", "RadiusLabel", "Radius"))
 		.NameContent()
 		[
 			SNew(STextBlock)
@@ -102,7 +102,7 @@ void FAttenuationSettingsCustomization::CustomizeChildren( TSharedRef<IPropertyH
 		]
 		.Visibility(TAttribute<EVisibility>(this, &FAttenuationSettingsCustomization::IsSphereSelected));
 
-	ChildBuilder.AddChildContent(TEXT("CapsuleHalfHeight"))
+	ChildBuilder.AddChildContent(NSLOCTEXT("AttenuationSettings", "CapsuleHalfHeightLabel", "Capsule Half Height"))
 			.NameContent()
 			[
 				SNew(STextBlock)
@@ -116,7 +116,7 @@ void FAttenuationSettingsCustomization::CustomizeChildren( TSharedRef<IPropertyH
 			]
 		.Visibility(TAttribute<EVisibility>(this, &FAttenuationSettingsCustomization::IsCapsuleSelected));
 
-	ChildBuilder.AddChildContent(TEXT("CapsuleRadius"))
+	ChildBuilder.AddChildContent(NSLOCTEXT("AttenuationSettings", "CapsuleRadiusLabel", "Capsule Radius"))
 		.NameContent()
 		[
 			SNew(STextBlock)
@@ -130,7 +130,7 @@ void FAttenuationSettingsCustomization::CustomizeChildren( TSharedRef<IPropertyH
 		]
 	.Visibility(TAttribute<EVisibility>(this, &FAttenuationSettingsCustomization::IsCapsuleSelected));
 
-	ChildBuilder.AddChildContent(TEXT("ConeRadius"))
+	ChildBuilder.AddChildContent(NSLOCTEXT("AttenuationSettings", "ConeRadiusLabel", "Cone Radius"))
 		.NameContent()
 		[
 			SNew(STextBlock)
@@ -144,7 +144,7 @@ void FAttenuationSettingsCustomization::CustomizeChildren( TSharedRef<IPropertyH
 		]
 	.Visibility(TAttribute<EVisibility>(this, &FAttenuationSettingsCustomization::IsConeSelected));
 
-	ChildBuilder.AddChildContent(TEXT("ConeAngle"))
+	ChildBuilder.AddChildContent(NSLOCTEXT("AttenuationSettings", "ConeAngleLabel", "Cone Angle"))
 		.NameContent()
 		[
 			SNew(STextBlock)
@@ -158,7 +158,7 @@ void FAttenuationSettingsCustomization::CustomizeChildren( TSharedRef<IPropertyH
 		]
 	.Visibility(TAttribute<EVisibility>(this, &FAttenuationSettingsCustomization::IsConeSelected));
 
-	ChildBuilder.AddChildContent(TEXT("ConeFalloffAngle"))
+	ChildBuilder.AddChildContent(NSLOCTEXT("AttenuationSettings", "ConeFalloffAngleLabel", "Cone Falloff Angle"))
 		.NameContent()
 		[
 			SNew(STextBlock)
diff --git a/Engine/Source/Editor/DetailCustomizations/Private/BodyInstanceCustomization.cpp b/Engine/Source/Editor/DetailCustomizations/Private/BodyInstanceCustomization.cpp
--- a/Engine/Source/Editor/DetailCustomizations/Private/BodyInstanceCustomization.cpp
+++ b/Engine/Source/Editor/DetailCustomizations/Private/BodyInstanceCustomization.cpp
@@ -74,7 +74,7 @@ void FBodyInstanceCustomization::CustomizeChildren( TSharedRef<class IPropertyHa
 	const FString PresetsDocLink = TEXT("Shared/Collision");
 	TSharedPtr<SToolTip> ProfileTooltip = IDocumentation::Get()->CreateToolTip(LOCTEXT("SelectCollisionPreset", "Select collision presets. You can set this data in Project settings."), NULL, PresetsDocLink, TEXT("PresetDetail"));
 
-	IDetailGroup& CollisionGroup = StructBuilder.AddChildGroup( TEXT("Collision"), LOCTEXT("CollisionPresetsLabel", "Collision Presets").ToString() );
+	IDetailGroup& CollisionGroup = StructBuilder.AddChildGroup( TEXT("Collision"), LOCTEXT("CollisionPresetsLabel", "Collision Presets") );
 	CollisionGroup.HeaderRow()
 	.NameContent()
 	[
diff --git a/Engine/Source/Editor/DetailCustomizations/Private/BodySetupDetails.cpp b/Engine/Source/Editor/DetailCustomizations/Private/BodySetupDetails.cpp
--- a/Engine/Source/Editor/DetailCustomizations/Private/BodySetupDetails.cpp
+++ b/Engine/Source/Editor/DetailCustomizations/Private/BodySetupDetails.cpp
@@ -53,7 +53,7 @@ void FBodySetupDetails::CustomizeDetails( IDetailLayoutBuilder& DetailBuilder )
 				}
 				else if (ChildProperty->GetProperty()->GetFName() == GET_MEMBER_NAME_CHECKED(FBodyInstance, MassInKg))
 				{
-					PhysicsCategory.AddCustomRow(TEXT("Mass"), false)
+					PhysicsCategory.AddCustomRow(LOCTEXT("MassLabel", "Mass"), false)
 						.IsEnabled(TAttribute<bool>(this, &FBodySetupDetails::IsBodyMassEnabled))
 						.NameContent()
 						[
diff --git a/Engine/Source/Editor/DetailCustomizations/Private/BrushDetails.cpp b/Engine/Source/Editor/DetailCustomizations/Private/BrushDetails.cpp
--- a/Engine/Source/Editor/DetailCustomizations/Private/BrushDetails.cpp
+++ b/Engine/Source/Editor/DetailCustomizations/Private/BrushDetails.cpp
@@ -187,10 +187,10 @@ void FBrushDetails::CustomizeDetails( IDetailLayoutBuilder& DetailLayout )
 		BrushBuilderObject->SetFlags( RF_Transactional );
 	}
 
-	IDetailCategoryBuilder& BrushBuilderCategory = DetailLayout.EditCategory( "BrushSettings", TEXT(""), ECategoryPriority::Important );
+	IDetailCategoryBuilder& BrushBuilderCategory = DetailLayout.EditCategory( "BrushSettings", FText::GetEmpty(), ECategoryPriority::Important );
 
 	BrushBuilderCategory.AddProperty( GET_MEMBER_NAME_CHECKED(ABrush, BrushType) );
-	BrushBuilderCategory.AddCustomRow( LOCTEXT("BrushShape", "Brush Shape").ToString() )
+	BrushBuilderCategory.AddCustomRow( LOCTEXT("BrushShape", "Brush Shape") )
 	.NameContent()
 	[
 		SNew( STextBlock )
@@ -213,7 +213,7 @@ void FBrushDetails::CustomizeDetails( IDetailLayoutBuilder& DetailLayout )
 		]
 	];
 
-	BrushBuilderCategory.AddCustomRow( TEXT(""), true )
+	BrushBuilderCategory.AddCustomRow( FText::GetEmpty(), true )
 	[
 		SNew(SHorizontalBox)
 		+ SHorizontalBox::Slot()
@@ -276,7 +276,7 @@ void FBrushDetails::CustomizeDetails( IDetailLayoutBuilder& DetailLayout )
 
 	TSharedPtr< SHorizontalBox > BrushHorizontalBox;
 
-	BrushBuilderCategory.AddCustomRow( TEXT(""), true)
+	BrushBuilderCategory.AddCustomRow( FText::GetEmpty(), true)
 	[
 		SAssignNew(BrushHorizontalBox, SHorizontalBox)
 		+SHorizontalBox::Slot()
diff --git a/Engine/Source/Editor/DetailCustomizations/Private/CameraDetails.cpp b/Engine/Source/Editor/DetailCustomizations/Private/CameraDetails.cpp
--- a/Engine/Source/Editor/DetailCustomizations/Private/CameraDetails.cpp
+++ b/Engine/Source/Editor/DetailCustomizations/Private/CameraDetails.cpp
@@ -32,7 +32,7 @@ void FCameraDetails::CustomizeDetails( IDetailLayoutBuilder& DetailLayout )
 	FSimpleDelegate OnAspectRatioChangedDelegate = FSimpleDelegate::CreateSP( this, &FCameraDetails::OnAspectRatioChanged );
 	AspectRatioProperty->SetOnPropertyValueChanged( OnAspectRatioChangedDelegate );
 
-	IDetailCategoryBuilder& CameraCategory = DetailLayout.EditCategory( "CameraSettings", FString(), ECategoryPriority::Important );
+	IDetailCategoryBuilder& CameraCategory = DetailLayout.EditCategory( "CameraSettings", FText::GetEmpty(), ECategoryPriority::Important );
 	
 	// Organize the properties
 	CameraCategory.AddProperty(ProjectionModeProperty);
diff --git a/Engine/Source/Editor/DetailCustomizations/Private/CollisionProfileDetails.cpp b/Engine/Source/Editor/DetailCustomizations/Private/CollisionProfileDetails.cpp
--- a/Engine/Source/Editor/DetailCustomizations/Private/CollisionProfileDetails.cpp
+++ b/Engine/Source/Editor/DetailCustomizations/Private/CollisionProfileDetails.cpp
@@ -1345,7 +1345,7 @@ void FCollisionProfileDetails::CustomizeDetails( IDetailLayoutBuilder& DetailBui
 	TSharedPtr<SToolTip> ProfileTooltip = IDocumentation::Get()->CreateToolTip(LOCTEXT("EditCollisionPreset", "Edit collision presets."), NULL, PresetsDocLink, TEXT("Preset"));
 
 	// Customize collision section
-	ObjectChannelCategory.AddCustomRow(LOCTEXT("CustomCollisionObjectChannels", "ObjectChannels").ToString())
+	ObjectChannelCategory.AddCustomRow(LOCTEXT("CustomCollisionObjectChannels", "ObjectChannels"))
 	[
 		SNew(SVerticalBox)
 
@@ -1439,7 +1439,7 @@ void FCollisionProfileDetails::CustomizeDetails( IDetailLayoutBuilder& DetailBui
 		]
 	];
 	
-	TraceChannelCategory.AddCustomRow(LOCTEXT("CustomCollisionTraceChannels", "TraceChannels").ToString())
+	TraceChannelCategory.AddCustomRow(LOCTEXT("CustomCollisionTraceChannels", "TraceChannels"))
 	[
 		SNew(SVerticalBox)
 
@@ -1533,7 +1533,7 @@ void FCollisionProfileDetails::CustomizeDetails( IDetailLayoutBuilder& DetailBui
 		]
 	];
 
-	PresetCategory.AddCustomRow(LOCTEXT("CustomCollisionProfiles", "Presets").ToString())
+	PresetCategory.AddCustomRow(LOCTEXT("CustomCollisionProfiles", "Presets"))
 	[
 		SNew(SVerticalBox)
 
diff --git a/Engine/Source/Editor/DetailCustomizations/Private/CollisionProfileNameCustomization.cpp b/Engine/Source/Editor/DetailCustomizations/Private/CollisionProfileNameCustomization.cpp
--- a/Engine/Source/Editor/DetailCustomizations/Private/CollisionProfileNameCustomization.cpp
+++ b/Engine/Source/Editor/DetailCustomizations/Private/CollisionProfileNameCustomization.cpp
@@ -37,7 +37,7 @@ void FCollisionProfileNameCustomization::CustomizeChildren(TSharedRef<class IPro
 		SetPropertyWithName(*InitialSelectedName.Get());
 	}
 
-	IDetailGroup& CollisionGroup = StructBuilder.AddChildGroup(TEXT("Collision"), LOCTEXT("CollisionPresetName", "Collision Preset").ToString());
+	IDetailGroup& CollisionGroup = StructBuilder.AddChildGroup(TEXT("Collision"), LOCTEXT("CollisionPresetName", "Collision Preset"));
 	CollisionGroup.HeaderRow()
 	.NameContent()
 	[
diff --git a/Engine/Source/Editor/DetailCustomizations/Private/ComponentTransformDetails.cpp b/Engine/Source/Editor/DetailCustomizations/Private/ComponentTransformDetails.cpp
--- a/Engine/Source/Editor/DetailCustomizations/Private/ComponentTransformDetails.cpp
+++ b/Engine/Source/Editor/DetailCustomizations/Private/ComponentTransformDetails.cpp
@@ -362,7 +362,7 @@ void FComponentTransformDetails::GenerateChildContent( IDetailChildrenBuilder& C
 	// Location
 	if(!bHideLocationField)
 	{
-		ChildrenBuilder.AddChildContent( LOCTEXT("LocationFilter", "Location").ToString() )
+		ChildrenBuilder.AddChildContent( LOCTEXT("LocationFilter", "Location") )
 		.CopyAction( CreateCopyAction( ETransformField::Location ) )
 		.PasteAction( CreatePasteAction( ETransformField::Location ) )
 		.NameContent()
@@ -420,7 +420,7 @@ void FComponentTransformDetails::GenerateChildContent( IDetailChildrenBuilder& C
 	// Rotation
 	if(!bHideRotationField)
 	{
-		ChildrenBuilder.AddChildContent( LOCTEXT("RotationFilter", "Rotation").ToString() )
+		ChildrenBuilder.AddChildContent( LOCTEXT("RotationFilter", "Rotation") )
 		.CopyAction( CreateCopyAction(ETransformField::Rotation) )
 		.PasteAction( CreatePasteAction(ETransformField::Rotation) )
 		.NameContent()
@@ -484,7 +484,7 @@ void FComponentTransformDetails::GenerateChildContent( IDetailChildrenBuilder& C
 	// Scale
 	if(!bHideScaleField)
 	{
-		ChildrenBuilder.AddChildContent( LOCTEXT("ScaleFilter", "Scale").ToString() )
+		ChildrenBuilder.AddChildContent( LOCTEXT("ScaleFilter", "Scale") )
 		.CopyAction( CreateCopyAction(ETransformField::Scale) )
 		.PasteAction( CreatePasteAction(ETransformField::Scale) )
 		.NameContent()
diff --git a/Engine/Source/Editor/DetailCustomizations/Private/CurveColorCustomization.cpp b/Engine/Source/Editor/DetailCustomizations/Private/CurveColorCustomization.cpp
--- a/Engine/Source/Editor/DetailCustomizations/Private/CurveColorCustomization.cpp
+++ b/Engine/Source/Editor/DetailCustomizations/Private/CurveColorCustomization.cpp
@@ -55,7 +55,7 @@ void FCurveColorCustomization::CustomizeHeader(TSharedRef<IPropertyHandle> InStr
 		HeaderRow
 			.NameContent()
 			[
-				InStructPropertyHandle->CreatePropertyNameWidget( TEXT( "" ), TEXT( "" ), false )
+				InStructPropertyHandle->CreatePropertyNameWidget( FText::GetEmpty(), FText::GetEmpty(), false )
 			]
 			.ValueContent()
 			.HAlign(HAlign_Fill)
@@ -90,7 +90,7 @@ void FCurveColorCustomization::CustomizeHeader(TSharedRef<IPropertyHandle> InStr
 		HeaderRow
 			.NameContent()
 			[
-				InStructPropertyHandle->CreatePropertyNameWidget( TEXT( "" ), TEXT( "" ), false )
+				InStructPropertyHandle->CreatePropertyNameWidget( FText::GetEmpty(), FText::GetEmpty(), false )
 			]
 			.ValueContent()
 			[
@@ -120,7 +120,7 @@ void FCurveColorCustomization::CustomizeChildren(TSharedRef<IPropertyHandle> InS
 			FSimpleDelegate OnCurveChangedDelegate = FSimpleDelegate::CreateSP( this, &FCurveColorCustomization::OnExternalCurveChanged, InStructPropertyHandle );
 			Child->SetOnPropertyValueChanged(OnCurveChangedDelegate);
 
-			StructBuilder.AddChildContent(TEXT("ExternalCurve"))
+			StructBuilder.AddChildContent(LOCTEXT("ExternalCurveLabel", "ExternalCurve"))
 				.NameContent()
 				[
 					Child->CreatePropertyNameWidget()
@@ -384,7 +384,7 @@ FReply FCurveColorCustomization::OnCurvePreviewDoubleClick(const FGeometry& InMy
 			FVector2D AdjustedSummonLocation = FSlateApplication::Get().CalculatePopupWindowPosition( Anchor, FCurveColorCustomization::DEFAULT_WINDOW_SIZE, Orient_Horizontal );
 
 			TSharedPtr<SWindow> Window = SNew(SWindow)
-				.Title( FText::Format( LOCTEXT("WindowHeader", "{0} - Internal Color Curve Editor"), FText::FromString(StructPropertyHandle->GetPropertyDisplayName())) )
+				.Title( FText::Format( LOCTEXT("WindowHeader", "{0} - Internal Color Curve Editor"), StructPropertyHandle->GetPropertyDisplayName()) )
 				.ClientSize( FCurveColorCustomization::DEFAULT_WINDOW_SIZE )
 				.ScreenPosition(AdjustedSummonLocation)
 				.AutoCenter(EAutoCenter::None)
diff --git a/Engine/Source/Editor/DetailCustomizations/Private/CurveStructCustomization.cpp b/Engine/Source/Editor/DetailCustomizations/Private/CurveStructCustomization.cpp
--- a/Engine/Source/Editor/DetailCustomizations/Private/CurveStructCustomization.cpp
+++ b/Engine/Source/Editor/DetailCustomizations/Private/CurveStructCustomization.cpp
@@ -70,7 +70,7 @@ void FCurveStructCustomization::CustomizeHeader( TSharedRef<IPropertyHandle> InS
 		HeaderRow
 			.NameContent()
 			[
-				InStructPropertyHandle->CreatePropertyNameWidget( TEXT( "" ), TEXT( "" ), false )
+				InStructPropertyHandle->CreatePropertyNameWidget( FText::GetEmpty(), FText::GetEmpty(), false )
 			]
 			.ValueContent()
 			.HAlign(HAlign_Fill)
@@ -107,7 +107,7 @@ void FCurveStructCustomization::CustomizeHeader( TSharedRef<IPropertyHandle> InS
 		HeaderRow
 			.NameContent()
 			[
-				InStructPropertyHandle->CreatePropertyNameWidget( TEXT( "" ), TEXT( "" ), false )
+				InStructPropertyHandle->CreatePropertyNameWidget( FText::GetEmpty(), FText::GetEmpty(), false )
 			]
 			.ValueContent()
 			[
@@ -137,7 +137,7 @@ void FCurveStructCustomization::CustomizeChildren( TSharedRef<IPropertyHandle> I
 			FSimpleDelegate OnCurveChangedDelegate = FSimpleDelegate::CreateSP( this, &FCurveStructCustomization::OnExternalCurveChanged, InStructPropertyHandle );
 			Child->SetOnPropertyValueChanged(OnCurveChangedDelegate);
 
-			StructBuilder.AddChildContent(TEXT("ExternalCurve"))
+			StructBuilder.AddChildContent(LOCTEXT("ExternalCurveLabel", "ExternalCurve"))
 				.NameContent()
 				[
 					Child->CreatePropertyNameWidget()
@@ -368,7 +368,7 @@ FReply FCurveStructCustomization::OnCurvePreviewDoubleClick( const FGeometry& In
 			FVector2D AdjustedSummonLocation = FSlateApplication::Get().CalculatePopupWindowPosition( Anchor, FCurveStructCustomization::DEFAULT_WINDOW_SIZE, Orient_Horizontal );
 
 			TSharedPtr<SWindow> Window = SNew(SWindow)
-				.Title( FText::Format( LOCTEXT("WindowHeader", "{0} - Internal Curve Editor"), FText::FromString(StructPropertyHandle->GetPropertyDisplayName())) )
+				.Title( FText::Format( LOCTEXT("WindowHeader", "{0} - Internal Curve Editor"), StructPropertyHandle->GetPropertyDisplayName()) )
 				.ClientSize( FCurveStructCustomization::DEFAULT_WINDOW_SIZE )
 				.ScreenPosition(AdjustedSummonLocation)
 				.AutoCenter(EAutoCenter::None)
diff --git a/Engine/Source/Editor/DetailCustomizations/Private/CurveTableCustomization.h b/Engine/Source/Editor/DetailCustomizations/Private/CurveTableCustomization.h
--- a/Engine/Source/Editor/DetailCustomizations/Private/CurveTableCustomization.h
+++ b/Engine/Source/Editor/DetailCustomizations/Private/CurveTableCustomization.h
@@ -22,7 +22,7 @@ public:
 		HeaderRow
 			.NameContent()
 			[
-				InStructPropertyHandle->CreatePropertyNameWidget( TEXT( "" ), TEXT( "" ), false )
+				InStructPropertyHandle->CreatePropertyNameWidget( FText::GetEmpty(), FText::GetEmpty(), false )
 			];
 	}
 
@@ -43,7 +43,7 @@ public:
 			CurveTablePropertyHandle->SetOnPropertyValueChanged( OnCurveTableChangedDelegate );
 
 			/** Construct a combo box widget to select from a list of valid options */
-			StructBuilder.AddChildContent( TEXT( "Row Name" ) )
+			StructBuilder.AddChildContent( NSLOCTEXT("CurveTable", "RowNameLabel", "Row Name" ) )
 			.NameContent()
 				[
 					SNew( STextBlock )
diff --git a/Engine/Source/Editor/DetailCustomizations/Private/DataTableCategoryCustomization.h b/Engine/Source/Editor/DetailCustomizations/Private/DataTableCategoryCustomization.h
--- a/Engine/Source/Editor/DetailCustomizations/Private/DataTableCategoryCustomization.h
+++ b/Engine/Source/Editor/DetailCustomizations/Private/DataTableCategoryCustomization.h
@@ -23,7 +23,7 @@ public:
 		HeaderRow
 			.NameContent()
 			[
-				InStructPropertyHandle->CreatePropertyNameWidget(TEXT(""), TEXT( "" ), false)
+				InStructPropertyHandle->CreatePropertyNameWidget(FText::GetEmpty(), FText::GetEmpty(), false)
 			];
 	}
 
@@ -47,7 +47,7 @@ public:
 			TSharedPtr<FString> InitialColumnValue = InitColumnWidgetContent();
 
 			/** Construct a combo box widget to select from a list of valid options */
-			StructBuilder.AddChildContent(LOCTEXT("DataTable_ColumnName", "Column Name").ToString())
+			StructBuilder.AddChildContent(LOCTEXT("DataTable_ColumnName", "Column Name"))
 			.NameContent()
 				[
 					SNew(STextBlock)
@@ -71,7 +71,7 @@ public:
 			TSharedPtr<FString> InitialRowValue = InitRowWidgetContent();
 
 			/** Construct a combo box widget to select from a list of valid options */
-			StructBuilder.AddChildContent(LOCTEXT("DataTable_RowContains", "Row Contains").ToString())
+			StructBuilder.AddChildContent(LOCTEXT("DataTable_RowContains", "Row Contains"))
 			.NameContent()
 				[
 					SNew(STextBlock)
diff --git a/Engine/Source/Editor/DetailCustomizations/Private/DataTableCustomization.h b/Engine/Source/Editor/DetailCustomizations/Private/DataTableCustomization.h
--- a/Engine/Source/Editor/DetailCustomizations/Private/DataTableCustomization.h
+++ b/Engine/Source/Editor/DetailCustomizations/Private/DataTableCustomization.h
@@ -24,7 +24,7 @@ public:
 		HeaderRow
 			.NameContent()
 			[
-				InStructPropertyHandle->CreatePropertyNameWidget( TEXT( "" ), TEXT( "" ), false )
+				InStructPropertyHandle->CreatePropertyNameWidget( FText::GetEmpty(), FText::GetEmpty(), false )
 			];
 	}
 
@@ -45,11 +45,11 @@ public:
 			DataTablePropertyHandle->SetOnPropertyValueChanged( OnDataTableChangedDelegate );
 
 			/** Construct a combo box widget to select from a list of valid options */
-			StructBuilder.AddChildContent( LOCTEXT( "DataTable_RowName", "Row Name" ).ToString() )
+			StructBuilder.AddChildContent( LOCTEXT( "DataTable_RowName", "Row Name" ) )
 			.NameContent()
 				[
 					SNew( STextBlock )
-					.Text( LOCTEXT( "DataTable_RowName", "Row Name" ).ToString() )
+					.Text( LOCTEXT( "DataTable_RowName", "Row Name" ) )
 					.Font( StructCustomizationUtils.GetRegularFont() )
 				]
 			.ValueContent()
diff --git a/Engine/Source/Editor/DetailCustomizations/Private/DestructibleMeshDetails.cpp b/Engine/Source/Editor/DetailCustomizations/Private/DestructibleMeshDetails.cpp
--- a/Engine/Source/Editor/DetailCustomizations/Private/DestructibleMeshDetails.cpp
+++ b/Engine/Source/Editor/DetailCustomizations/Private/DestructibleMeshDetails.cpp
@@ -16,7 +16,7 @@ TSharedRef<IDetailCustomization> FDestructibleMeshDetails::MakeInstance()
 
 void AddStructToDetails(FName CategoryName, FName PropertyName, IDetailLayoutBuilder& DetailBuilder, bool bInline = true, bool bAdvanced = false)
 {
-	IDetailCategoryBuilder& Category = DetailBuilder.EditCategory(CategoryName, TEXT(""), ECategoryPriority::Important);
+	IDetailCategoryBuilder& Category = DetailBuilder.EditCategory(CategoryName, FText::GetEmpty(), ECategoryPriority::Important);
 	TSharedPtr<IPropertyHandle> Params = DetailBuilder.GetProperty(PropertyName);
 	if (Params.IsValid())
 	{
diff --git a/Engine/Source/Editor/DetailCustomizations/Private/DeviceProfileDetails.cpp b/Engine/Source/Editor/DetailCustomizations/Private/DeviceProfileDetails.cpp
--- a/Engine/Source/Editor/DetailCustomizations/Private/DeviceProfileDetails.cpp
+++ b/Engine/Source/Editor/DetailCustomizations/Private/DeviceProfileDetails.cpp
@@ -428,7 +428,7 @@ void FDeviceProfileParentPropertyDetails::CreateParentPropertyView()
 	ParentPropertyNameHandle->GetValue(CurrentParentName);
 
 	IDetailCategoryBuilder& ParentDetailCategory = DetailBuilder->EditCategory("ParentDeviceProfile");
-	IDetailGroup& ParentNameGroup = ParentDetailCategory.AddGroup(TEXT("ParentProfileName"), LOCTEXT("ParentProfileOptionsGroupTitle", "Parent Profile Name").ToString());
+	IDetailGroup& ParentNameGroup = ParentDetailCategory.AddGroup(TEXT("ParentProfileName"), LOCTEXT("ParentProfileOptionsGroupTitle", "Parent Profile Name"));
 
 	ParentNameGroup.HeaderRow()
 	[
@@ -498,7 +498,7 @@ void FDeviceProfileParentPropertyDetails::CreateParentPropertyView()
 			{
 				if(ParentCVarsGroup == nullptr)
 				{
-					ParentCVarsGroup = &ParentDetailCategory.AddGroup(TEXT("ParentProfileOptions"), LOCTEXT("ParentProfileOptionsGroupTitle", "Parent Console Variables").ToString());
+					ParentCVarsGroup = &ParentDetailCategory.AddGroup(TEXT("ParentProfileOptions"), LOCTEXT("ParentProfileOptionsGroupTitle", "Parent Console Variables"));
 
 					ParentCVarsGroup->HeaderRow()
 					[
@@ -616,17 +616,17 @@ void FDeviceProfileConsoleVariablesPropertyDetails::CreateConsoleVariablesProper
 	for (auto& Current : CategoryPropertyMap)
 	{
 		const TArray<TSharedRef<IPropertyHandle>>& CurrentGroupsProperties = Current.Value;
-		const TSharedPtr<FString> GroupName = MakeShareable(new FString(Current.Key));
+		const FText GroupName = FText::FromString(Current.Key);
 
-		FString CVarPrefix = DeviceProfileCVarFormatHelper::CVarPrefixFromCategoryString(*GroupName.Get());
+		FString CVarPrefix = DeviceProfileCVarFormatHelper::CVarPrefixFromCategoryString(GroupName.ToString());
 		if (CVarPrefix.Len() > 0)
 		{
 			CVarPrefix += TEXT(".");
 		}
 
 		// Find the Property table UI Group for the current CVar
-		IDetailGroup& CVarGroup = CVarDetailCategory.AddGroup(**GroupName.Get(), *GroupName.Get());
-		CVarDetailGroups.Add(*GroupName.Get(), &CVarGroup);
+		IDetailGroup& CVarGroup = CVarDetailCategory.AddGroup(*GroupName.ToString(), GroupName);
+		CVarDetailGroups.Add(GroupName.ToString(), &CVarGroup);
 		CVarGroup.HeaderRow()
 			[
 				SNew(SHorizontalBox)
@@ -636,7 +636,7 @@ void FDeviceProfileConsoleVariablesPropertyDetails::CreateConsoleVariablesProper
 				.AutoWidth()
 				[
 					SNew(STextBlock)
-					.Text(FText::FromString(*GroupName.Get()))
+					.Text(GroupName)
 					.Font(IDetailLayoutBuilder::GetDetailFont())
 				]
 				+ SHorizontalBox::Slot()
@@ -803,14 +803,14 @@ FReply FDeviceProfileConsoleVariablesPropertyDetails::OnRemoveCVarProperty(TShar
 }
 
 
-FReply FDeviceProfileConsoleVariablesPropertyDetails::OnRemoveAllFromGroup(const TSharedPtr<FString> GroupName)
+FReply FDeviceProfileConsoleVariablesPropertyDetails::OnRemoveAllFromGroup(FText GroupName)
 {
 	const TSharedPtr<IPropertyHandleArray> CVarsArrayHandle = CVarsHandle->AsArray();
 
 	uint32 CVarCount = 0;
 	ensure(CVarsArrayHandle->GetNumElements(CVarCount) == FPropertyAccess::Success);
 
-	FString CVarPrefix = DeviceProfileCVarFormatHelper::CVarPrefixFromCategoryString(*GroupName);
+	FString CVarPrefix = DeviceProfileCVarFormatHelper::CVarPrefixFromCategoryString(GroupName.ToString());
 
 	for (int32 CVarPropertyIdx = CVarCount-1; CVarPropertyIdx >= 0; CVarPropertyIdx--)
 	{
diff --git a/Engine/Source/Editor/DetailCustomizations/Private/DeviceProfileDetails.h b/Engine/Source/Editor/DetailCustomizations/Private/DeviceProfileDetails.h
--- a/Engine/Source/Editor/DetailCustomizations/Private/DeviceProfileDetails.h
+++ b/Engine/Source/Editor/DetailCustomizations/Private/DeviceProfileDetails.h
@@ -138,7 +138,7 @@ private:
 	 * @param GroupName - The name of the group all Cvars have been removed for.
 	 * @return Whether the event was handled
 	 */
-	FReply OnRemoveAllFromGroup(const TSharedPtr<FString> GroupName);
+	FReply OnRemoveAllFromGroup(FText GroupName);
 
 
 	/**
diff --git a/Engine/Source/Editor/DetailCustomizations/Private/DialogueStructsCustomizations.cpp b/Engine/Source/Editor/DetailCustomizations/Private/DialogueStructsCustomizations.cpp
--- a/Engine/Source/Editor/DetailCustomizations/Private/DialogueStructsCustomizations.cpp
+++ b/Engine/Source/Editor/DetailCustomizations/Private/DialogueStructsCustomizations.cpp
@@ -751,7 +751,7 @@ void SValidatedDialogueContextHeaderWidget::Construct( const FArguments& InArgs,
 							+SHorizontalBox::Slot()
 							.AutoWidth()
 							[
-								TargetsPropertyHandle->CreatePropertyNameWidget( LOCTEXT("DirectedAt", "Directed At").ToString() )
+								TargetsPropertyHandle->CreatePropertyNameWidget( LOCTEXT("DirectedAt", "Directed At") )
 							]
 							+SHorizontalBox::Slot()
 							.FillWidth(1.0f)
diff --git a/Engine/Source/Editor/DetailCustomizations/Private/DialogueWaveDetails.cpp b/Engine/Source/Editor/DetailCustomizations/Private/DialogueWaveDetails.cpp
--- a/Engine/Source/Editor/DetailCustomizations/Private/DialogueWaveDetails.cpp
+++ b/Engine/Source/Editor/DetailCustomizations/Private/DialogueWaveDetails.cpp
@@ -97,7 +97,7 @@ void FDialogueWaveDetails::CustomizeDetails( class IDetailLayoutBuilder& DetailB
 	IDetailCategoryBuilder& ContextMappingsDetailCategoryBuilder = DetailBuilder.EditCategory("DialogueContexts");
 
 	// Add Context Button
-	ContextMappingsDetailCategoryBuilder.AddCustomRow( LOCTEXT("AddDialogueContext", "Add Dialogue Context").ToString() )
+	ContextMappingsDetailCategoryBuilder.AddCustomRow( LOCTEXT("AddDialogueContext", "Add Dialogue Context") )
 	[
 		SNew( SBox )
 		.Padding(2.0f)
diff --git a/Engine/Source/Editor/DetailCustomizations/Private/DialogueWaveWidgets.cpp b/Engine/Source/Editor/DetailCustomizations/Private/DialogueWaveWidgets.cpp
--- a/Engine/Source/Editor/DetailCustomizations/Private/DialogueWaveWidgets.cpp
+++ b/Engine/Source/Editor/DetailCustomizations/Private/DialogueWaveWidgets.cpp
@@ -852,7 +852,7 @@ void SDialogueContextHeaderWidget::Construct( const FArguments& InArgs, const TS
 						+SHorizontalBox::Slot()
 						.AutoWidth()
 						[
-							TargetsPropertyHandle->CreatePropertyNameWidget( LOCTEXT("DirectedAt", "Directed At").ToString() )
+							TargetsPropertyHandle->CreatePropertyNameWidget( LOCTEXT("DirectedAt", "Directed At") )
 						]
 						+SHorizontalBox::Slot()
 						.AutoWidth()
diff --git a/Engine/Source/Editor/DetailCustomizations/Private/DirectionalLightComponentDetails.cpp b/Engine/Source/Editor/DetailCustomizations/Private/DirectionalLightComponentDetails.cpp
--- a/Engine/Source/Editor/DetailCustomizations/Private/DirectionalLightComponentDetails.cpp
+++ b/Engine/Source/Editor/DetailCustomizations/Private/DirectionalLightComponentDetails.cpp
@@ -18,7 +18,7 @@ void FDirectionalLightComponentDetails::CustomizeDetails( IDetailLayoutBuilder&
 	MobilityProperty = DetailBuilder.GetProperty("Mobility", USceneComponent::StaticClass());
 
 	// Get cascaded shadow map category
-	IDetailCategoryBuilder& ShadowMapCategory = DetailBuilder.EditCategory("CascadedShadowMaps", TEXT(""), ECategoryPriority::Default );
+	IDetailCategoryBuilder& ShadowMapCategory = DetailBuilder.EditCategory("CascadedShadowMaps", FText::GetEmpty(), ECategoryPriority::Default );
 
 	// Add DynamicShadowDistanceMovableLight
 	TSharedPtr<IPropertyHandle> MovableShadowRadiusProperty = DetailBuilder.GetProperty("DynamicShadowDistanceMovableLight");
diff --git a/Engine/Source/Editor/DetailCustomizations/Private/DocumentationActorDetails.cpp b/Engine/Source/Editor/DetailCustomizations/Private/DocumentationActorDetails.cpp
--- a/Engine/Source/Editor/DetailCustomizations/Private/DocumentationActorDetails.cpp
+++ b/Engine/Source/Editor/DetailCustomizations/Private/DocumentationActorDetails.cpp
@@ -33,7 +33,7 @@ void FDocumentationActorDetails::CustomizeDetails(IDetailLayoutBuilder& DetailBu
 
 	// Add a button we can click on to open the documentation
 	IDetailCategoryBuilder& HelpCategory = DetailBuilder.EditCategory("Help Data");
-	HelpCategory.AddCustomRow("Help Documentation")
+	HelpCategory.AddCustomRow(LOCTEXT("HelpDocumentation", "Help Documentation"))
 		[
 			SNew(SButton)
 			.Text(this, &FDocumentationActorDetails::OnGetButtonText)
diff --git a/Engine/Source/Editor/DetailCustomizations/Private/EnvQueryParamInstanceCustomization.cpp b/Engine/Source/Editor/DetailCustomizations/Private/EnvQueryParamInstanceCustomization.cpp
--- a/Engine/Source/Editor/DetailCustomizations/Private/EnvQueryParamInstanceCustomization.cpp
+++ b/Engine/Source/Editor/DetailCustomizations/Private/EnvQueryParamInstanceCustomization.cpp
@@ -40,7 +40,7 @@ void FEnvQueryParamInstanceCustomization::CustomizeChildren( TSharedRef<class IP
 	StructBuilder.AddChildProperty(NameProp.ToSharedRef());
 	StructBuilder.AddChildProperty(TypeProp.ToSharedRef());
 	
-	StructBuilder.AddChildContent("Value")
+	StructBuilder.AddChildContent(LOCTEXT("ValueLabel", "Value"))
 		.NameContent()
 		[
 			ValueProp->CreatePropertyNameWidget()
diff --git a/Engine/Source/Editor/DetailCustomizations/Private/FbxImportUIDetails.cpp b/Engine/Source/Editor/DetailCustomizations/Private/FbxImportUIDetails.cpp
--- a/Engine/Source/Editor/DetailCustomizations/Private/FbxImportUIDetails.cpp
+++ b/Engine/Source/Editor/DetailCustomizations/Private/FbxImportUIDetails.cpp
@@ -36,7 +36,7 @@ void FFbxImportUIDetails::CustomizeDetails( IDetailLayoutBuilder& DetailBuilder
 	ImportUI = Cast<UFbxImportUI>(EditingObjects[0].Get());
 
 	// Handle mesh category
-	IDetailCategoryBuilder& MeshCategory = DetailBuilder.EditCategory("Mesh", TEXT(""), ECategoryPriority::Important);
+	IDetailCategoryBuilder& MeshCategory = DetailBuilder.EditCategory("Mesh", FText::GetEmpty(), ECategoryPriority::Important);
 	IDetailCategoryBuilder& TransformCategory = DetailBuilder.EditCategory("Transform");
 	TArray<TSharedRef<IPropertyHandle>> CategoryDefaultProperties;
 	TArray<TSharedPtr<IPropertyHandle>> ExtraProperties;
@@ -133,7 +133,7 @@ void FFbxImportUIDetails::CustomizeDetails( IDetailLayoutBuilder& DetailBuilder
 	}
 
 	// Animation Category
-	IDetailCategoryBuilder& AnimCategory = DetailBuilder.EditCategory("Animation", TEXT(""), ECategoryPriority::Important);
+	IDetailCategoryBuilder& AnimCategory = DetailBuilder.EditCategory("Animation", FText::GetEmpty(), ECategoryPriority::Important);
 
 	CategoryDefaultProperties.Empty();
 	AnimCategory.GetDefaultProperties(CategoryDefaultProperties);
diff --git a/Engine/Source/Editor/DetailCustomizations/Private/GeneralProjectSettingsDetails.cpp b/Engine/Source/Editor/DetailCustomizations/Private/GeneralProjectSettingsDetails.cpp
--- a/Engine/Source/Editor/DetailCustomizations/Private/GeneralProjectSettingsDetails.cpp
+++ b/Engine/Source/Editor/DetailCustomizations/Private/GeneralProjectSettingsDetails.cpp
@@ -23,7 +23,7 @@ void FGeneralProjectSettingsDetails::CustomizeDetails(IDetailLayoutBuilder& Deta
 	const FString ProjectThumbnail_TargetImagePath = FPaths::GetPath(FPaths::GetProjectFilePath()) / FString::Printf(TEXT("%s.png"), FApp::GetGameName());
 	const FString ProjectThumbnail_AutomaticImagePath = FPaths::GameSavedDir() / TEXT("AutoScreenshot.png");
 
-	ImageCategory.AddCustomRow(ProjectThumbnailDesc.ToString())
+	ImageCategory.AddCustomRow(ProjectThumbnailDesc)
 	.NameContent()
 	[
 		SNew(SHorizontalBox)
diff --git a/Engine/Source/Editor/DetailCustomizations/Private/HardwareTargetingSettingsDetails.cpp b/Engine/Source/Editor/DetailCustomizations/Private/HardwareTargetingSettingsDetails.cpp
--- a/Engine/Source/Editor/DetailCustomizations/Private/HardwareTargetingSettingsDetails.cpp
+++ b/Engine/Source/Editor/DetailCustomizations/Private/HardwareTargetingSettingsDetails.cpp
@@ -105,7 +105,7 @@ public:
 
 			SettingRegions.Add(Settings.SettingsObject, EditPropertiesBlock);
 
-			FDetailWidgetRow& CategoryRow = PendingChangesCategory.AddCustomRow(Settings.CategoryHeading.ToString())
+			FDetailWidgetRow& CategoryRow = PendingChangesCategory.AddCustomRow(Settings.CategoryHeading)
 				.Visibility(TAttribute<EVisibility>::Create(&SRequiredDefaultConfig::GetAnyPendingChangesVisibility))
 				.NameContent()
 				[
@@ -206,7 +206,7 @@ void FHardwareTargetingSettingsDetails::CustomizeDetails(IDetailLayoutBuilder& D
 
 	TSharedRef<SRequiredDefaultConfig> ConfigWidget = SNew(SRequiredDefaultConfig);
 	IHardwareTargetingModule& HardwareTargeting = IHardwareTargetingModule::Get();
-	PendingChangesCategory.AddCustomRow(FString())
+	PendingChangesCategory.AddCustomRow(FText::GetEmpty())
 	[
 		SNew(SVerticalBox)
 		+ SVerticalBox::Slot()
@@ -274,7 +274,7 @@ void FHardwareTargetingSettingsDetails::CustomizeDetails(IDetailLayoutBuilder& D
 		);
 	}
 
-	HardwareTargetingCategory.AddCustomRow(LOCTEXT("HardwareTargetingOption", "Targeted Hardware:").ToString())
+	HardwareTargetingCategory.AddCustomRow(LOCTEXT("HardwareTargetingOption", "Targeted Hardware:"))
 	.NameContent()
 	[
 		SNew(STextBlock)
diff --git a/Engine/Source/Editor/DetailCustomizations/Private/InputSettingsDetails.cpp b/Engine/Source/Editor/DetailCustomizations/Private/InputSettingsDetails.cpp
--- a/Engine/Source/Editor/DetailCustomizations/Private/InputSettingsDetails.cpp
+++ b/Engine/Source/Editor/DetailCustomizations/Private/InputSettingsDetails.cpp
@@ -79,7 +79,7 @@ void FActionMappingsNodeBuilder::GenerateChildContent( IDetailChildrenBuilder& C
 		FString GroupNameString(TEXT("ActionMappings."));
 		MappingSet.SharedName.AppendString(GroupNameString);
 		FName GroupName(*GroupNameString);
-		IDetailGroup& ActionMappingGroup = ChildrenBuilder.AddChildGroup(GroupName, MappingSet.SharedName.ToString());
+		IDetailGroup& ActionMappingGroup = ChildrenBuilder.AddChildGroup(GroupName, FText::FromName(MappingSet.SharedName));
 		MappingSet.DetailGroup = &ActionMappingGroup;
 
 		TSharedRef<SWidget> AddButton = PropertyCustomizationHelpers::MakeAddButton(FSimpleDelegate::CreateSP(this, &FActionMappingsNodeBuilder::AddActionMappingToGroupButton_OnClick, MappingSet),
@@ -386,7 +386,7 @@ void FAxisMappingsNodeBuilder::GenerateChildContent( IDetailChildrenBuilder& Chi
 		FString GroupNameString(TEXT("AxisMappings."));
 		MappingSet.SharedName.AppendString(GroupNameString);
 		FName GroupName(*GroupNameString);
-		IDetailGroup& AxisMappingGroup = ChildrenBuilder.AddChildGroup(GroupName, MappingSet.SharedName.ToString());
+		IDetailGroup& AxisMappingGroup = ChildrenBuilder.AddChildGroup(GroupName, FText::FromName(MappingSet.SharedName));
 		MappingSet.DetailGroup = &AxisMappingGroup;
 
 		TSharedRef<SWidget> AddButton = PropertyCustomizationHelpers::MakeAddButton(FSimpleDelegate::CreateSP(this, &FAxisMappingsNodeBuilder::AddAxisMappingToGroupButton_OnClick, MappingSet),
@@ -641,7 +641,7 @@ void FInputSettingsDetails::CustomizeDetails(class IDetailLayoutBuilder& DetailB
 
 	IDetailCategoryBuilder& MappingsDetailCategoryBuilder = DetailBuilder.EditCategory(BindingsCategory);
 
-	MappingsDetailCategoryBuilder.AddCustomRow(TEXT("Action Axis Mappings"))
+	MappingsDetailCategoryBuilder.AddCustomRow(LOCTEXT("Mappings_Title", "Action Axis Mappings"))
 	[
 		SNew(SHorizontalBox)
 		+SHorizontalBox::Slot()
diff --git a/Engine/Source/Editor/DetailCustomizations/Private/InputStructCustomization.cpp b/Engine/Source/Editor/DetailCustomizations/Private/InputStructCustomization.cpp
--- a/Engine/Source/Editor/DetailCustomizations/Private/InputStructCustomization.cpp
+++ b/Engine/Source/Editor/DetailCustomizations/Private/InputStructCustomization.cpp
@@ -23,7 +23,7 @@ void FInputAxisConfigCustomization::CustomizeHeader( TSharedRef<class IPropertyH
 
 	HeaderRow.NameContent()
 	[
-		InStructPropertyHandle->CreatePropertyNameWidget(AxisKeyName)
+		InStructPropertyHandle->CreatePropertyNameWidget(FText::FromString(AxisKeyName))
 	];
 }
 
@@ -65,7 +65,7 @@ void FInputActionMappingCustomization::CustomizeChildren( TSharedRef<class IProp
 	TSharedRef<SWidget> RemoveButton = PropertyCustomizationHelpers::MakeDeleteButton(FSimpleDelegate::CreateSP(this, &FInputActionMappingCustomization::RemoveActionMappingButton_OnClick),
 		LOCTEXT("RemoveActionMappingToolTip", "Removes Action Mapping"));
 
-	StructBuilder.AddChildContent( LOCTEXT("KeySearchStr", "Key").ToString() )
+	StructBuilder.AddChildContent( LOCTEXT("KeySearchStr", "Key") )
 	[
 		SNew(SHorizontalBox)
 		+ SHorizontalBox::Slot()
@@ -186,7 +186,7 @@ void FInputAxisMappingCustomization::CustomizeChildren( TSharedRef<class IProper
 	TSharedRef<SWidget> RemoveButton = PropertyCustomizationHelpers::MakeDeleteButton( FSimpleDelegate::CreateSP( this, &FInputAxisMappingCustomization::RemoveAxisMappingButton_OnClick), 
 		LOCTEXT("RemoveAxisMappingToolTip", "Removes Axis Mapping") );
 
-	StructBuilder.AddChildContent( LOCTEXT("KeySearchStr", "Key").ToString() )
+	StructBuilder.AddChildContent( LOCTEXT("KeySearchStr", "Key") )
 	[
 		SNew(SHorizontalBox)
 		+ SHorizontalBox::Slot()
diff --git a/Engine/Source/Editor/DetailCustomizations/Private/LandscapeSplineDetails.cpp b/Engine/Source/Editor/DetailCustomizations/Private/LandscapeSplineDetails.cpp
--- a/Engine/Source/Editor/DetailCustomizations/Private/LandscapeSplineDetails.cpp
+++ b/Engine/Source/Editor/DetailCustomizations/Private/LandscapeSplineDetails.cpp
@@ -14,9 +14,9 @@ TSharedRef<IDetailCustomization> FLandscapeSplineDetails::MakeInstance()
 
 void FLandscapeSplineDetails::CustomizeDetails( IDetailLayoutBuilder& DetailBuilder )
 {
-	IDetailCategoryBuilder& LandscapeSplineCategory = DetailBuilder.EditCategory( "LandscapeSpline", TEXT(""), ECategoryPriority::Transform );
+	IDetailCategoryBuilder& LandscapeSplineCategory = DetailBuilder.EditCategory( "LandscapeSpline", FText::GetEmpty(), ECategoryPriority::Transform );
 
-	LandscapeSplineCategory.AddCustomRow( TEXT("") )
+	LandscapeSplineCategory.AddCustomRow( FText::GetEmpty() )
 	[
 		SNew(SHorizontalBox)
 		+ SHorizontalBox::Slot()
@@ -45,7 +45,7 @@ void FLandscapeSplineDetails::CustomizeDetails( IDetailLayoutBuilder& DetailBuil
 		]
 	];
 
-	LandscapeSplineCategory.AddCustomRow( TEXT("") )
+	LandscapeSplineCategory.AddCustomRow( FText::GetEmpty() )
 	[
 		SNew(SHorizontalBox)
 		+ SHorizontalBox::Slot()
diff --git a/Engine/Source/Editor/DetailCustomizations/Private/LightComponentDetails.cpp b/Engine/Source/Editor/DetailCustomizations/Private/LightComponentDetails.cpp
--- a/Engine/Source/Editor/DetailCustomizations/Private/LightComponentDetails.cpp
+++ b/Engine/Source/Editor/DetailCustomizations/Private/LightComponentDetails.cpp
@@ -16,9 +16,9 @@ void FLightComponentDetails::CustomizeDetails( IDetailLayoutBuilder& DetailBuild
 	// Mobility property is on the scene component base class not the light component and that is why we have to use USceneComponent::StaticClass
 	TSharedRef<IPropertyHandle> MobilityHandle = DetailBuilder.GetProperty(GET_MEMBER_NAME_CHECKED(ULightComponent, Mobility), USceneComponent::StaticClass());
 	// Set a mobility tooltip specific to lights
-	MobilityHandle->SetToolTipText(LOCTEXT("LightMobilityTooltip", "Mobility for lights controls what the light is allowed to do at runtime and therefore what rendering methods are used.\n● A movable light uses fully dynamic lighting and anything can change in game, however it has a large performance cost, typically proportional to the light's influence size.\n● A stationary light will only have its shadowing and bounced lighting from static geometry baked by Lightmass, all other lighting will be dynamic.  It can change color and intensity in game. \n● A static light is fully baked into lightmaps and therefore has no performance cost, but also can't change in game.").ToString());
+	MobilityHandle->SetToolTipText(LOCTEXT("LightMobilityTooltip", "Mobility for lights controls what the light is allowed to do at runtime and therefore what rendering methods are used.\n● A movable light uses fully dynamic lighting and anything can change in game, however it has a large performance cost, typically proportional to the light's influence size.\n● A stationary light will only have its shadowing and bounced lighting from static geometry baked by Lightmass, all other lighting will be dynamic.  It can change color and intensity in game. \n● A static light is fully baked into lightmaps and therefore has no performance cost, but also can't change in game."));
 
-	IDetailCategoryBuilder& LightCategory = DetailBuilder.EditCategory( "Light", TEXT(""), ECategoryPriority::TypeSpecific );
+	IDetailCategoryBuilder& LightCategory = DetailBuilder.EditCategory( "Light", FText::GetEmpty(), ECategoryPriority::TypeSpecific );
 
 	LightIntensityProperty = DetailBuilder.GetProperty(GET_MEMBER_NAME_CHECKED(ULightComponent, Intensity), ULightComponentBase::StaticClass());
 	IESBrightnessTextureProperty = DetailBuilder.GetProperty(GET_MEMBER_NAME_CHECKED(ULightComponent, IESTexture));
@@ -33,7 +33,7 @@ void FLightComponentDetails::CustomizeDetails( IDetailLayoutBuilder& DetailBuild
 	}
 	else
 	{
-		IDetailCategoryBuilder& LightProfilesCategory = DetailBuilder.EditCategory( "Light Profiles", TEXT(""), ECategoryPriority::TypeSpecific );
+		IDetailCategoryBuilder& LightProfilesCategory = DetailBuilder.EditCategory( "Light Profiles", FText::GetEmpty(), ECategoryPriority::TypeSpecific );
 
 		LightCategory.AddProperty( LightIntensityProperty )
 			.IsEnabled( TAttribute<bool>( this, &FLightComponentDetails::IsLightBrightnessEnabled ) );
diff --git a/Engine/Source/Editor/DetailCustomizations/Private/MacTargetSettingsDetails.cpp b/Engine/Source/Editor/DetailCustomizations/Private/MacTargetSettingsDetails.cpp
--- a/Engine/Source/Editor/DetailCustomizations/Private/MacTargetSettingsDetails.cpp
+++ b/Engine/Source/Editor/DetailCustomizations/Private/MacTargetSettingsDetails.cpp
@@ -89,10 +89,10 @@ static FString GetIconFilename(EMacImageScope::Type Scope)
 void FMacTargetSettingsDetails::CustomizeDetails( IDetailLayoutBuilder& DetailBuilder )
 {
 	// Add the splash image customization
+	const FText EditorSplashDesc(LOCTEXT("EditorSplashLabel", "Editor Splash"));
 	IDetailCategoryBuilder& SplashCategoryBuilder = DetailBuilder.EditCategory(TEXT("Splash"));
-	FDetailWidgetRow& EditorSplashWidgetRow = SplashCategoryBuilder.AddCustomRow(TEXT("Editor Splash"));
+	FDetailWidgetRow& EditorSplashWidgetRow = SplashCategoryBuilder.AddCustomRow(EditorSplashDesc);
 
-	const FText EditorSplashDesc(LOCTEXT("EditorSplashLabel", "Editor Splash"));
 	const FString EditorSplash_TargetImagePath = GetSplashFilename(EMacImageScope::GameOverride, true);
 	const FString EditorSplash_DefaultImagePath = GetSplashFilename(EMacImageScope::Engine, true);
 
@@ -125,9 +125,9 @@ void FMacTargetSettingsDetails::CustomizeDetails( IDetailLayoutBuilder& DetailBu
 		]
 	];
 
-	FDetailWidgetRow& GameSplashWidgetRow = SplashCategoryBuilder.AddCustomRow(TEXT("Game Splash"));
-
 	const FText GameSplashDesc(LOCTEXT("GameSplashLabel", "Game Splash"));
+	FDetailWidgetRow& GameSplashWidgetRow = SplashCategoryBuilder.AddCustomRow(GameSplashDesc);
+
 	const FString GameSplash_TargetImagePath = GetSplashFilename(EMacImageScope::GameOverride, false);
 	const FString GameSplash_DefaultImagePath = GetSplashFilename(EMacImageScope::Engine, false);
 
@@ -161,7 +161,7 @@ void FMacTargetSettingsDetails::CustomizeDetails( IDetailLayoutBuilder& DetailBu
 	];
 
 	IDetailCategoryBuilder& IconsCategoryBuilder = DetailBuilder.EditCategory(TEXT("Icon"));	
-	FDetailWidgetRow& GameIconWidgetRow = IconsCategoryBuilder.AddCustomRow(TEXT("Game Icon"));
+	FDetailWidgetRow& GameIconWidgetRow = IconsCategoryBuilder.AddCustomRow(LOCTEXT("GameIconLabel", "Game Icon"));
 	GameIconWidgetRow
 	.NameContent()
 	[
diff --git a/Engine/Source/Editor/DetailCustomizations/Private/MarginCustomization.cpp b/Engine/Source/Editor/DetailCustomizations/Private/MarginCustomization.cpp
--- a/Engine/Source/Editor/DetailCustomizations/Private/MarginCustomization.cpp
+++ b/Engine/Source/Editor/DetailCustomizations/Private/MarginCustomization.cpp
@@ -198,7 +198,7 @@ void FMarginStructCustomization::OnMarginTextCommitted( const FText& InText, ETe
 				StructPropertyHandle->AccessRawData( RawData );
 
 				{
-					FScopedTransaction Transaction( FText::Format( NSLOCTEXT("FMarginStructCustomization", "SetMarginProperty", "Edit {0}"), FText::FromString( StructPropertyHandle->GetPropertyDisplayName() ) ) );
+					FScopedTransaction Transaction( FText::Format( NSLOCTEXT("FMarginStructCustomization", "SetMarginProperty", "Edit {0}"), StructPropertyHandle->GetPropertyDisplayName() ) );
 					StructPropertyHandle->NotifyPreChange();
 
 					for (void* Data : RawData)
@@ -279,7 +279,7 @@ void FMarginStructCustomization::OnBeginSliderMovement()
 {
 	bIsUsingSlider = true;
 
-	GEditor->BeginTransaction( FText::Format( NSLOCTEXT("FMarginStructCustomization", "SetMarginProperty", "Edit {0}"), FText::FromString( StructPropertyHandle->GetPropertyDisplayName() ) ) );
+	GEditor->BeginTransaction( FText::Format( NSLOCTEXT("FMarginStructCustomization", "SetMarginProperty", "Edit {0}"), StructPropertyHandle->GetPropertyDisplayName() ) );
 }
 
 void FMarginStructCustomization::OnEndSliderMovement( float NewValue )
diff --git a/Engine/Source/Editor/DetailCustomizations/Private/MathStructCustomizations.cpp b/Engine/Source/Editor/DetailCustomizations/Private/MathStructCustomizations.cpp
--- a/Engine/Source/Editor/DetailCustomizations/Private/MathStructCustomizations.cpp
+++ b/Engine/Source/Editor/DetailCustomizations/Private/MathStructCustomizations.cpp
@@ -38,8 +38,8 @@ void FMathStructCustomization::MakeHeaderRow( TSharedRef<class IPropertyHandle>&
 {
 	// We'll set up reset to default ourselves
 	const bool bDisplayResetToDefault = false;
-	const FString DisplayNameOverride = TEXT("");
-	const FString DisplayToolTipOverride = TEXT("");
+	const FText DisplayNameOverride = FText::GetEmpty();
+	const FText DisplayToolTipOverride = FText::GetEmpty();
 
 	TWeakPtr<IPropertyHandle> StructWeakHandlePtr = StructPropertyHandle;
 
@@ -433,8 +433,8 @@ void FColorStructCustomization::MakeHeaderRow( TSharedRef<class IPropertyHandle>
 {
 	// We'll set up reset to default ourselves
 	const bool bDisplayResetToDefault = false;
-	const FString DisplayNameOverride = TEXT("");
-	const FString DisplayToolTipOverride = TEXT("");
+	const FText DisplayNameOverride = FText::GetEmpty();
+	const FText DisplayToolTipOverride = FText::GetEmpty();
 
 	Row.NameContent()
 	[
@@ -671,7 +671,7 @@ void FColorStructCustomization::OnColorPickerInteractiveBegin()
 {
 	bIsInteractive = true;
 
-	GEditor->BeginTransaction( FText::Format( NSLOCTEXT("FColorStructCustomization", "SetColorProperty", "Edit {0}"), FText::FromString( StructPropertyHandle->GetPropertyDisplayName() ) ) );
+	GEditor->BeginTransaction( FText::Format( NSLOCTEXT("FColorStructCustomization", "SetColorProperty", "Edit {0}"), StructPropertyHandle->GetPropertyDisplayName() ) );
 }
 
 void FColorStructCustomization::OnColorPickerInteractiveEnd()
