Expose Actor primary tick properties bStartWithTickEnabled, bTickEvenWhenPaused, and bAllowTickOnDedicatedServer to blueprint defaults
Make SetActorTickEnabled non-virtual

diff --git a/Engine/Source/Editor/DetailCustomizations/Private/ActorDetails.cpp b/Engine/Source/Editor/DetailCustomizations/Private/ActorDetails.cpp
--- a/Engine/Source/Editor/DetailCustomizations/Private/ActorDetails.cpp
+++ b/Engine/Source/Editor/DetailCustomizations/Private/ActorDetails.cpp
@@ -136,6 +136,19 @@ void FActorDetails::CustomizeDetails( IDetailLayoutBuilder& DetailLayout )
 
 		//AddComponentsCategory( DetailLayout );
 	}
+
+	// Defaults only
+	if (DetailLayout.GetDetailsView().HasClassDefaultObject())
+	{
+		IDetailCategoryBuilder& TickCategory = DetailLayout.EditCategory("Tick");
+
+		TSharedPtr<IPropertyHandle> PrimaryTickProperty = DetailLayout.GetProperty(GET_MEMBER_NAME_CHECKED(AActor, PrimaryActorTick));
+		TickCategory.AddProperty(PrimaryTickProperty->GetChildHandle(GET_MEMBER_NAME_CHECKED(FTickFunction, bStartWithTickEnabled)));
+		TickCategory.AddProperty(PrimaryTickProperty->GetChildHandle(GET_MEMBER_NAME_CHECKED(FTickFunction, bTickEvenWhenPaused)), EPropertyLocation::Advanced);
+		TickCategory.AddProperty(PrimaryTickProperty->GetChildHandle(GET_MEMBER_NAME_CHECKED(FTickFunction, bAllowTickOnDedicatedServer)), EPropertyLocation::Advanced);
+
+		PrimaryTickProperty->MarkHiddenByCustomization();
+	}
 }
 
 void FActorDetails::OnConvertActor(UClass* ChosenClass)
diff --git a/Engine/Source/Runtime/Engine/Classes/Engine/EngineBaseTypes.h b/Engine/Source/Runtime/Engine/Classes/Engine/EngineBaseTypes.h
--- a/Engine/Source/Runtime/Engine/Classes/Engine/EngineBaseTypes.h
+++ b/Engine/Source/Runtime/Engine/Classes/Engine/EngineBaseTypes.h
@@ -142,7 +142,7 @@ public:
 	TEnumAsByte<enum ETickingGroup> TickGroup;
 
 	/** Bool indicating that this function should execute even if the game is paused. Pause ticks are very limited in capabilities. **/
-	UPROPERTY()
+	UPROPERTY(EditDefaultsOnly, Category="Tick", AdvancedDisplay)
 	uint32 bTickEvenWhenPaused:1;
 
 	/** If false, this tick function will never be registered and will never tick. Only settable in defaults. */
@@ -150,11 +150,11 @@ public:
 	uint32 bCanEverTick:1;
 
 	/** If true, this tick function will start enabled, but can be disabled later on. */
-	UPROPERTY()
+	UPROPERTY(EditDefaultsOnly, Category="Tick")
 	uint32 bStartWithTickEnabled:1;
 
 	/** If we allow this tick to run on a dedicated server */
-	UPROPERTY()
+	UPROPERTY(EditDefaultsOnly, Category="Tick", AdvancedDisplay)
 	uint32 bAllowTickOnDedicatedServer:1;
 
 	/** If false, this tick will run on the game thread, otherwise it will run on any thread in parallel with the game thread and in parallel with other "async ticks" **/
diff --git a/Engine/Source/Runtime/Engine/Classes/GameFramework/Actor.h b/Engine/Source/Runtime/Engine/Classes/GameFramework/Actor.h
--- a/Engine/Source/Runtime/Engine/Classes/GameFramework/Actor.h
+++ b/Engine/Source/Runtime/Engine/Classes/GameFramework/Actor.h
@@ -81,7 +81,7 @@ public:
 	 * @see https://docs.unrealengine.com/latest/INT/API/Runtime/Engine/Engine/FTickFunction/
 	 * @see AddTickPrerequisiteActor(), AddTickPrerequisiteComponent()
 	 */
-	UPROPERTY()
+	UPROPERTY(EditDefaultsOnly, Category="Tick")
 	struct FActorTickFunction PrimaryActorTick;
 
 	/** Allow each actor to run at a different time speed. The DeltaTime for a frame is multiplied by the global TimeDilation (in WorldSettings) and this CustomTimeDilation for this actor's tick.  */
@@ -1489,7 +1489,7 @@ public:
 	 * @param	bEnabled - Rather it should be enabled or not
 	 */
 	UFUNCTION(BlueprintCallable, Category="Utilities")
-	virtual void SetActorTickEnabled(bool bEnabled);
+	void SetActorTickEnabled(bool bEnabled);
 
 	/**  Returns whether this actor has tick enabled or not	 */
 	UFUNCTION(BlueprintCallable, Category="Utilities")
