Fixed code relying on SLATE_TEXT_ATTRIBUTE for SButton

Made sure everything was using FText rather than FString.

diff --git a/Engine/Source/Editor/Blutility/Private/BlutilityDetailsPanel.cpp b/Engine/Source/Editor/Blutility/Private/BlutilityDetailsPanel.cpp
--- a/Engine/Source/Editor/Blutility/Private/BlutilityDetailsPanel.cpp
+++ b/Engine/Source/Editor/Blutility/Private/BlutilityDetailsPanel.cpp
@@ -98,7 +98,7 @@ void FEditorUtilityInstanceDetails::CustomizeDetails(IDetailLayoutBuilder& Detai
 			{
 				++NumButtons;
 
-				const FString ButtonCaption = FName::NameToDisplayString(*Function->GetName(), false);
+				const FText ButtonCaption = FText::FromString(FName::NameToDisplayString(*Function->GetName(), false));
 
 				//@TODO: Expose the code in UK2Node_CallFunction::GetUserFacingFunctionName / etc...
 				FString Tooltip = Function->GetToolTipText().ToString();
diff --git a/Engine/Source/Editor/DetailCustomizations/Private/CollisionProfileDetails.cpp b/Engine/Source/Editor/DetailCustomizations/Private/CollisionProfileDetails.cpp
--- a/Engine/Source/Editor/DetailCustomizations/Private/CollisionProfileDetails.cpp
+++ b/Engine/Source/Editor/DetailCustomizations/Private/CollisionProfileDetails.cpp
@@ -586,7 +586,7 @@ void SProfileEditDialog::Construct(const FArguments& InArgs)
 			[
 				SNew(SButton)
 				.ContentPadding(FEditorStyle::GetMargin("StandardDialog.ContentPadding"))
-				.Text(LOCTEXT("SProfileEditDialog_Accept", "Accept").ToString())
+				.Text(LOCTEXT("SProfileEditDialog_Accept", "Accept"))
 				.OnClicked(this, &SProfileEditDialog::OnAccept)
 				.IsEnabled(this, &SProfileEditDialog::IsAcceptAvailable)
 			]
@@ -595,7 +595,7 @@ void SProfileEditDialog::Construct(const FArguments& InArgs)
 			[
 				SNew(SButton)
 				.ContentPadding(FEditorStyle::GetMargin("StandardDialog.ContentPadding"))
-				.Text(LOCTEXT("SProfileEditDialog_Cancel", "Cancel").ToString())
+				.Text(LOCTEXT("SProfileEditDialog_Cancel", "Cancel"))
 				.OnClicked(this, &SProfileEditDialog::OnCancel)
 			]
 		]
@@ -1163,7 +1163,7 @@ TSharedRef<SWidget> SChannelListItem::GenerateWidgetForColumn(const FName& Colum
 			.VAlign(VAlign_Center)
 			[
 				SNew(STextBlock)
-				.Text(FString::Printf(TEXT("%s"), *ChannelSetup->Name.ToString()))
+				.Text(FText::FromName(ChannelSetup->Name))
 				.Font(IDetailLayoutBuilder::GetDetailFont())
 			];
 	}
@@ -1247,7 +1247,7 @@ TSharedRef<SWidget> SProfileListItem::GenerateWidgetForColumn(const FName& Colum
 				[
 					SNew(SImage)
 					.Image(FEditorStyle::GetBrush("SettingsEditor.Collision_Engine"))
-					.ToolTipText(LOCTEXT("CantModify_Tooltip", "You can't modify the name of Engine profiles").ToString())
+					.ToolTipText(LOCTEXT("CantModify_Tooltip", "You can't modify the name of Engine profiles"))
 				];
 		}
 		else
@@ -1258,7 +1258,7 @@ TSharedRef<SWidget> SProfileListItem::GenerateWidgetForColumn(const FName& Colum
 				[
 					SNew(SImage)
 					.Image(FEditorStyle::GetBrush("SettingsEditor.Collision_Game"))
-					.ToolTipText(LOCTEXT("CanModify_Tooltip", "This is your custom project profie").ToString())
+					.ToolTipText(LOCTEXT("CanModify_Tooltip", "This is your custom project profie"))
 				];
 		}
 	}
@@ -1269,7 +1269,7 @@ TSharedRef<SWidget> SProfileListItem::GenerateWidgetForColumn(const FName& Colum
 			.VAlign(VAlign_Center)
 			[
 				SNew(STextBlock)
-				.Text(FString::Printf(TEXT("%s"), *ProfileTemplate->Name.ToString()))
+				.Text(FText::FromName(ProfileTemplate->Name))
 				.Font(IDetailLayoutBuilder::GetDetailFont())
 			];
 	}
@@ -1372,7 +1372,7 @@ void FCollisionProfileDetails::CustomizeDetails( IDetailLayoutBuilder& DetailBui
 			[
 				SNew(SButton)
 				.ContentPadding(FEditorStyle::GetMargin("StandardDialog.ContentPadding"))
-				.Text(LOCTEXT("ChannelMenu_NewObject", "New Object Channel...").ToString())
+				.Text(LOCTEXT("ChannelMenu_NewObject", "New Object Channel..."))
 				.OnClicked(this, &FCollisionProfileDetails::OnNewChannel, false)
 				.IsEnabled(this, &FCollisionProfileDetails::IsNewChannelAvailable)
 			]
@@ -1384,7 +1384,7 @@ void FCollisionProfileDetails::CustomizeDetails( IDetailLayoutBuilder& DetailBui
 			[
 				SNew(SButton)
 				.ContentPadding(FEditorStyle::GetMargin("StandardDialog.ContentPadding"))
-				.Text(LOCTEXT("ChannelMenu_Edit", "Edit...").ToString())
+				.Text(LOCTEXT("ChannelMenu_Edit", "Edit..."))
 				.OnClicked(this, &FCollisionProfileDetails::OnEditChannel, false)
 				.IsEnabled(this, &FCollisionProfileDetails::IsAnyChannelSelected, false)
 			]
@@ -1396,7 +1396,7 @@ void FCollisionProfileDetails::CustomizeDetails( IDetailLayoutBuilder& DetailBui
 			[
 				SNew(SButton)
 				.ContentPadding(FEditorStyle::GetMargin("StandardDialog.ContentPadding"))
-				.Text(LOCTEXT("ChannelMenu_Delete", "Delete...").ToString())
+				.Text(LOCTEXT("ChannelMenu_Delete", "Delete..."))
 				.OnClicked(this, &FCollisionProfileDetails::OnDeleteChannel, false)
 				.IsEnabled(this, &FCollisionProfileDetails::IsAnyChannelSelected, false)
 			]
@@ -1466,7 +1466,7 @@ void FCollisionProfileDetails::CustomizeDetails( IDetailLayoutBuilder& DetailBui
 			[
 				SNew(SButton)
 				.ContentPadding(FEditorStyle::GetMargin("StandardDialog.ContentPadding"))
-				.Text(LOCTEXT("ChannelMenu_NewTrace", "New Trace Channel...").ToString())
+				.Text(LOCTEXT("ChannelMenu_NewTrace", "New Trace Channel..."))
 				.OnClicked(this, &FCollisionProfileDetails::OnNewChannel, true)
 				.IsEnabled(this, &FCollisionProfileDetails::IsNewChannelAvailable)
 			]
@@ -1478,7 +1478,7 @@ void FCollisionProfileDetails::CustomizeDetails( IDetailLayoutBuilder& DetailBui
 			[
 				SNew(SButton)
 				.ContentPadding(FEditorStyle::GetMargin("StandardDialog.ContentPadding"))
-				.Text(LOCTEXT("ChannelMenu_Edit", "Edit...").ToString())
+				.Text(LOCTEXT("ChannelMenu_Edit", "Edit..."))
 				.OnClicked(this, &FCollisionProfileDetails::OnEditChannel, true)
 				.IsEnabled(this, &FCollisionProfileDetails::IsAnyChannelSelected, true)
 			]
@@ -1490,7 +1490,7 @@ void FCollisionProfileDetails::CustomizeDetails( IDetailLayoutBuilder& DetailBui
 			[
 				SNew(SButton)
 				.ContentPadding(FEditorStyle::GetMargin("StandardDialog.ContentPadding"))
-				.Text(LOCTEXT("ChannelMenu_Delete", "Delete...").ToString())
+				.Text(LOCTEXT("ChannelMenu_Delete", "Delete..."))
 				.OnClicked(this, &FCollisionProfileDetails::OnDeleteChannel, true)
 				.IsEnabled(this, &FCollisionProfileDetails::IsAnyChannelSelected, true)
 			]
@@ -1560,7 +1560,7 @@ void FCollisionProfileDetails::CustomizeDetails( IDetailLayoutBuilder& DetailBui
 			[
 				SNew(SButton)
 				.ContentPadding(FEditorStyle::GetMargin("StandardDialog.ContentPadding"))
-				.Text(LOCTEXT("ProfileMenu_New", "New...").ToString())
+				.Text(LOCTEXT("ProfileMenu_New", "New..."))
 				.OnClicked(this, &FCollisionProfileDetails::OnNewProfile)
 			]
 
@@ -1571,7 +1571,7 @@ void FCollisionProfileDetails::CustomizeDetails( IDetailLayoutBuilder& DetailBui
 			[
 				SNew(SButton)
 				.ContentPadding(FEditorStyle::GetMargin("StandardDialog.ContentPadding"))
-				.Text(LOCTEXT("ProfileMenu_Edit", "Edit...").ToString())
+				.Text(LOCTEXT("ProfileMenu_Edit", "Edit..."))
 				.OnClicked(this, &FCollisionProfileDetails::OnEditProfile)
 				.IsEnabled(this, &FCollisionProfileDetails::IsAnyProfileSelected)
 			]
@@ -1583,7 +1583,7 @@ void FCollisionProfileDetails::CustomizeDetails( IDetailLayoutBuilder& DetailBui
 			[
 				SNew(SButton)
 				.ContentPadding(FEditorStyle::GetMargin("StandardDialog.ContentPadding"))
-				.Text(LOCTEXT("ProfileMenu_Delete", "Delete...").ToString())
+				.Text(LOCTEXT("ProfileMenu_Delete", "Delete..."))
 				.OnClicked(this, &FCollisionProfileDetails::OnDeleteProfile)
 				.IsEnabled(this, &FCollisionProfileDetails::IsAnyProfileSelected)
 			]
diff --git a/Engine/Source/Editor/LevelEditor/Private/LevelEditorMenu.cpp b/Engine/Source/Editor/LevelEditor/Private/LevelEditorMenu.cpp
--- a/Engine/Source/Editor/LevelEditor/Private/LevelEditorMenu.cpp
+++ b/Engine/Source/Editor/LevelEditor/Private/LevelEditorMenu.cpp
@@ -32,15 +32,11 @@ class SFavouriteMenuEntry : public SCompoundWidget
 	 */
 	void Construct( const FArguments& InArgs )
 	{
-		FString AssetDisplayName("");
-		if ( !InArgs._LabelOverride.IsEmpty() )
-		{
-			AssetDisplayName = InArgs._LabelOverride.ToString();
-		}
+		const FText AssetDisplayName = InArgs._LabelOverride;
 		OnOpenClickedDelegate = InArgs._OnOpenClickedDelegate;
 		OnRemoveClickedDelegate = InArgs._OnRemoveClickedDelegate;
 
-		FString OpenDisplayName( *FString::Printf( *LOCTEXT("OpenLevelToolTip", "Open level: %s").ToString(), *AssetDisplayName ) );
+		const FText OpenDisplayName = FText::Format( LOCTEXT("OpenLevelToolTipFmt", "Open level: {0}"), AssetDisplayName );
 		FSlateFontInfo MenuEntryFont = FEditorStyle::GetFontStyle( "Menu.Label.Font" );
 
 		ChildSlot
@@ -76,7 +72,7 @@ class SFavouriteMenuEntry : public SCompoundWidget
 					SNew(SButton)
 					.ContentPadding( FMargin(4.0, 0.0) )
 					.ButtonStyle( FEditorStyle::Get(), "Docking.Tab.CloseButton" )
-					.ToolTipText( FString::Printf( *LOCTEXT("RemoveFavoriteToolTip", "Remove %s from Favorites").ToString(), *AssetDisplayName ) )
+					.ToolTipText( FText::Format( LOCTEXT("RemoveFavoriteToolTipFmt", "Remove {0} from Favorites"), AssetDisplayName ) )
 					.OnClicked(this, &SFavouriteMenuEntry::OnRemove)
 				]
 			]
diff --git a/Engine/Source/Editor/LevelEditor/Private/SSurfaceProperties.cpp b/Engine/Source/Editor/LevelEditor/Private/SSurfaceProperties.cpp
--- a/Engine/Source/Editor/LevelEditor/Private/SSurfaceProperties.cpp
+++ b/Engine/Source/Editor/LevelEditor/Private/SSurfaceProperties.cpp
@@ -175,9 +175,9 @@ TSharedRef<SWidget> SSurfaceProperties::ConstructTexturePan()
 		SAssignNew(GridBox,SUniformGridPanel)
 	];
 
-	const static FString ButtonFields[] = {"1/256", "1/64","1/16","1/4"};
+	const static FText ButtonFields[] = {FText::FromString("1/256"), FText::FromString("1/64"), FText::FromString("1/16"), FText::FromString("1/4")};
 	const static int32 ButtonIncriments[] = {1,4,16,64};
-	for (int i = 0; i < ARRAY_COUNT(ButtonFields); i++)
+	for (int32 i = 0; i < ARRAY_COUNT(ButtonFields); i++)
 	{
 		GridBox->AddSlot(i, 0)
 		.VAlign(VAlign_Fill)
diff --git a/Engine/Source/Editor/Persona/Private/SRigWindow.cpp b/Engine/Source/Editor/Persona/Private/SRigWindow.cpp
--- a/Engine/Source/Editor/Persona/Private/SRigWindow.cpp
+++ b/Engine/Source/Editor/Persona/Private/SRigWindow.cpp
@@ -432,14 +432,14 @@ FReply SRigWindow::OnToggleAdvanced()
 	return FReply::Handled();
 }
 
-FString SRigWindow::GetAdvancedButtonText() const
+FText SRigWindow::GetAdvancedButtonText() const
 {
 	if (bDisplayAdvanced)
 	{
-		return TEXT("Show Base");
+		return LOCTEXT("ShowBase", "Show Base");
 	}
 
-	return TEXT("Show Advanced");
+	return LOCTEXT("ShowAdvanced", "Show Advanced");
 }
 
 TSharedRef<SWidget> SRigWindow::MakeRigPickerWithMenu()
@@ -458,15 +458,15 @@ void SRigWindow::CloseComboButton()
 	AssetComboButton->SetIsOpen(false);
 }
 
-FString SRigWindow::GetAssetName() const
+FText SRigWindow::GetAssetName() const
 {
 	UObject* Rig = GetRigObject();
 	if (Rig)
 	{
-		return Rig->GetName();
+		return FText::FromString(Rig->GetName());
 	}
 
-	return TEXT("None");
+	return LOCTEXT("None", "None");
 }
 
 #undef LOCTEXT_NAMESPACE
diff --git a/Engine/Source/Editor/Persona/Private/SRigWindow.h b/Engine/Source/Editor/Persona/Private/SRigWindow.h
--- a/Engine/Source/Editor/Persona/Private/SRigWindow.h
+++ b/Engine/Source/Editor/Persona/Private/SRigWindow.h
@@ -159,7 +159,7 @@ private:
 	 */
 	/* Set rig set combo box*/
 	void OnAssetSelected(UObject* Object);
-	FString GetAssetName() const;
+	FText GetAssetName() const;
 	void CloseComboButton();
 	TSharedRef<SWidget> MakeRigPickerWithMenu();
 
@@ -172,5 +172,5 @@ private:
 	FName GetBoneMapping( FName NodeName );
 
 	FReply OnToggleAdvanced();
-	FString GetAdvancedButtonText() const;
+	FText GetAdvancedButtonText() const;
 };
\ No newline at end of file
diff --git a/Engine/Source/Editor/PhAT/Private/SPhATNewAssetDlg.cpp b/Engine/Source/Editor/PhAT/Private/SPhATNewAssetDlg.cpp
--- a/Engine/Source/Editor/PhAT/Private/SPhATNewAssetDlg.cpp
+++ b/Engine/Source/Editor/PhAT/Private/SPhATNewAssetDlg.cpp
@@ -60,7 +60,7 @@ void SPhATNewAssetDlg::Construct(const FArguments& InArgs)
 				.VAlign(VAlign_Center)
 				[
 					SNew(STextBlock)
-					.Text(NSLOCTEXT("PhAT", "MinimumBoneSizeLabel", "Minimum Bone Size:").ToString())
+					.Text(NSLOCTEXT("PhAT", "MinimumBoneSizeLabel", "Minimum Bone Size:"))
 				]
 				+SUniformGridPanel::Slot(1, 0)
 				.VAlign(VAlign_Center)
@@ -73,7 +73,7 @@ void SPhATNewAssetDlg::Construct(const FArguments& InArgs)
 				.VAlign(VAlign_Center)
 				[
 					SNew(STextBlock)
-					.Text(NSLOCTEXT("PhAT", "OrientAlongBoneLabel", "Orient Along Bone:").ToString())
+					.Text(NSLOCTEXT("PhAT", "OrientAlongBoneLabel", "Orient Along Bone:"))
 				]
 				+SUniformGridPanel::Slot(1, 1)
 				[
@@ -85,7 +85,7 @@ void SPhATNewAssetDlg::Construct(const FArguments& InArgs)
 				.VAlign(VAlign_Center)
 				[
 					SNew(STextBlock)
-					.Text(NSLOCTEXT("PhAT", "CollisionGeometryLabel", "Collision Geometry:").ToString())
+					.Text(NSLOCTEXT("PhAT", "CollisionGeometryLabel", "Collision Geometry:"))
 				]
 				+SUniformGridPanel::Slot(1, 2)
 				.VAlign(VAlign_Center)
@@ -99,7 +99,7 @@ void SPhATNewAssetDlg::Construct(const FArguments& InArgs)
 				.VAlign(VAlign_Center)
 				[
 					SNew(STextBlock)
-					.Text(NSLOCTEXT("PhAT", "UseVertsWithLabel", "Use Verts With:").ToString())
+					.Text(NSLOCTEXT("PhAT", "UseVertsWithLabel", "Use Verts With:"))
 				]
 				+SUniformGridPanel::Slot(1, 3)
 				.VAlign(VAlign_Center)
@@ -113,7 +113,7 @@ void SPhATNewAssetDlg::Construct(const FArguments& InArgs)
 				.VAlign(VAlign_Center)
 				[
 					SNew(STextBlock)
-					.Text(NSLOCTEXT("PhAT", "CreateJointsLabel", "Create Joints:").ToString())
+					.Text(NSLOCTEXT("PhAT", "CreateJointsLabel", "Create Joints:"))
 				]
 				+SUniformGridPanel::Slot(1, 4)
 				[
@@ -125,7 +125,7 @@ void SPhATNewAssetDlg::Construct(const FArguments& InArgs)
 				.VAlign(VAlign_Center)
 				[
 					SNew(STextBlock)
-					.Text(NSLOCTEXT("PhAT", "AngularConstraintMode", "Default Angular Constraint Mode:").ToString())
+					.Text(NSLOCTEXT("PhAT", "AngularConstraintMode", "Default Angular Constraint Mode:"))
 				]
 				+SUniformGridPanel::Slot(1, 5)
 				.VAlign(VAlign_Center)
@@ -139,7 +139,7 @@ void SPhATNewAssetDlg::Construct(const FArguments& InArgs)
 				.VAlign(VAlign_Center)
 				[
 					SNew(STextBlock)
-					.Text(NSLOCTEXT("PhAT", "WalkPastSmallBonesLabel", "Walk Past Small Bones:").ToString())
+					.Text(NSLOCTEXT("PhAT", "WalkPastSmallBonesLabel", "Walk Past Small Bones:"))
 				]
 				+SUniformGridPanel::Slot(1, 6)
 				[
@@ -151,7 +151,7 @@ void SPhATNewAssetDlg::Construct(const FArguments& InArgs)
 				.VAlign(VAlign_Center)
 				[
 					SNew(STextBlock)
-					.Text(NSLOCTEXT("PhAT", "CreateBodyForAllBonesLabel", "Create Body For All Bones:").ToString())
+					.Text(NSLOCTEXT("PhAT", "CreateBodyForAllBonesLabel", "Create Body For All Bones:"))
 				]
 				+SUniformGridPanel::Slot(1, 7)
 				[
@@ -165,7 +165,7 @@ void SPhATNewAssetDlg::Construct(const FArguments& InArgs)
 					[
 						SNew(STextBlock)
 						.Visibility(this, &SPhATNewAssetDlg::GetHullOptionsVisibility )
-						.Text( NSLOCTEXT("PhAT", "MaxNumHulls_ConvexDecomp", "Max Num Hulls").ToString() )
+						.Text( NSLOCTEXT("PhAT", "MaxNumHulls_ConvexDecomp", "Max Num Hulls") )
 					]
 				+SUniformGridPanel::Slot(1, 7)
 					[
@@ -181,7 +181,7 @@ void SPhATNewAssetDlg::Construct(const FArguments& InArgs)
 					[
 						SNew(STextBlock)
 						.Visibility(this, &SPhATNewAssetDlg::GetHullOptionsVisibility )
-						.Text( NSLOCTEXT("PhAT", "MaxHullVerts_ConvexDecomp", "Max Hull Verts").ToString() )
+						.Text( NSLOCTEXT("PhAT", "MaxHullVerts_ConvexDecomp", "Max Hull Verts") )
 					]
 				+SUniformGridPanel::Slot(1, 8)
 					[
@@ -205,7 +205,7 @@ void SPhATNewAssetDlg::Construct(const FArguments& InArgs)
 				+SUniformGridPanel::Slot(0,0)
 				[
 					SNew(SButton)
-					.Text(NSLOCTEXT("PhAT", "OkButtonText", "Ok").ToString())
+					.Text(NSLOCTEXT("PhAT", "OkButtonText", "Ok"))
 					.HAlign(HAlign_Center)
 					.ContentPadding(FEditorStyle::GetMargin("StandardDialog.ContentPadding"))
 					.OnClicked(FOnClicked::CreateSP<SPhATNewAssetDlg, EAppReturnType::Type>(this, &SPhATNewAssetDlg::OnClicked, EAppReturnType::Ok))
@@ -213,7 +213,7 @@ void SPhATNewAssetDlg::Construct(const FArguments& InArgs)
 				+SUniformGridPanel::Slot(1,0)
 				[
 					SNew(SButton)
-					.Text(NSLOCTEXT("PhAT", "CancelButtonText", "Cancel").ToString())
+					.Text(NSLOCTEXT("PhAT", "CancelButtonText", "Cancel"))
 					.HAlign(HAlign_Center)
 					.ContentPadding(FEditorStyle::GetMargin("StandardDialog.ContentPadding"))
 					.OnClicked(FOnClicked::CreateSP<SPhATNewAssetDlg, EAppReturnType::Type>(this, &SPhATNewAssetDlg::OnClicked, EAppReturnType::Cancel))
diff --git a/Engine/Source/Editor/UnrealEd/Classes/Preferences/UnrealEdOptions.h b/Engine/Source/Editor/UnrealEd/Classes/Preferences/UnrealEdOptions.h
--- a/Engine/Source/Editor/UnrealEd/Classes/Preferences/UnrealEdOptions.h
+++ b/Engine/Source/Editor/UnrealEd/Classes/Preferences/UnrealEdOptions.h
@@ -69,10 +69,10 @@ struct FClassPickerDefaults
 	FString AssetClass;
 
 	/** Gets the localized name text for the class */
-	FString GetName() const;
+	FText GetName() const;
 
 	/** Gets the localized descriptive text for the class */
-	FString GetDescription() const;
+	FText GetDescription() const;
 };
 
 UCLASS(Config=Editor)
diff --git a/Engine/Source/Editor/UnrealEd/Private/SClassPickerDialog.cpp b/Engine/Source/Editor/UnrealEd/Private/SClassPickerDialog.cpp
--- a/Engine/Source/Editor/UnrealEd/Private/SClassPickerDialog.cpp
+++ b/Engine/Source/Editor/UnrealEd/Private/SClassPickerDialog.cpp
@@ -92,7 +92,7 @@ void SClassPickerDialog::Construct(const FArguments& InArgs)
 					+SUniformGridPanel::Slot(0,0)
 					[
 						SNew(SButton)
-						.Text(NSLOCTEXT("SClassPickerDialog", "ClassPickerSelectButton", "Select").ToString())
+						.Text(NSLOCTEXT("SClassPickerDialog", "ClassPickerSelectButton", "Select"))
 						.HAlign(HAlign_Center)
 						.Visibility( this, &SClassPickerDialog::GetSelectButtonVisibility )
 						.ContentPadding(FEditorStyle::GetMargin("StandardDialog.ContentPadding"))
@@ -101,7 +101,7 @@ void SClassPickerDialog::Construct(const FArguments& InArgs)
 					+SUniformGridPanel::Slot(1,0)
 					[
 						SNew(SButton)
-						.Text(NSLOCTEXT("SClassPickerDialog", "ClassPickerCancelButton", "Cancel").ToString())
+						.Text(NSLOCTEXT("SClassPickerDialog", "ClassPickerCancelButton", "Cancel"))
 						.HAlign(HAlign_Center)
 						.ContentPadding(FEditorStyle::GetMargin("StandardDialog.ContentPadding"))
 						.OnClicked(this, &SClassPickerDialog::OnClassPickerCanceled)
@@ -195,7 +195,7 @@ TSharedRef<ITableRow> SClassPickerDialog::GenerateListRow(TSharedPtr<FClassPicke
 			.Padding(10.0f, 0.0f)
 			[
 				SNew(STextBlock)
-				.Text(FString::Printf(*NSLOCTEXT("EditorFactories", "FClassPickerDefaults", "%s").ToString(), *Obj->GetDescription()))
+				.Text(Obj->GetDescription())
 				.AutoWrapText(true)
 			]
 			+SHorizontalBox::Slot()
diff --git a/Engine/Source/Editor/UnrealEd/Private/SSocketManager.cpp b/Engine/Source/Editor/UnrealEd/Private/SSocketManager.cpp
--- a/Engine/Source/Editor/UnrealEd/Private/SSocketManager.cpp
+++ b/Engine/Source/Editor/UnrealEd/Private/SSocketManager.cpp
@@ -253,7 +253,7 @@ void SSocketManager::Construct(const FArguments& InArgs)
 						.Padding(5.0f, 0.0f, 4.0f, 0.0f)
 						[
 							SNew(SButton)
-							.Text(LOCTEXT("CreateSocket", "Create Socket").ToString())
+							.Text(LOCTEXT("CreateSocket", "Create Socket"))
 							.OnClicked(this, &SSocketManager::CreateSocket_Execute)
 						]
 						+SHorizontalBox::Slot()
@@ -262,7 +262,7 @@ void SSocketManager::Construct(const FArguments& InArgs)
 							.Padding(4.0f, 0.0f, 5.0f, 0.0f)
 							[
 								SNew(SButton)
-								.Text(LOCTEXT("DeleteSocket", "Delete Socket").ToString())
+								.Text(LOCTEXT("DeleteSocket", "Delete Socket"))
 								.OnClicked(this, &SSocketManager::DeleteSelectedSocket_Execute)
 							]
 					]
@@ -273,7 +273,7 @@ void SSocketManager::Construct(const FArguments& InArgs)
 					.Padding(0.0f, 12.0f, 0.0f, 4.0f)
 					[
 						SNew(STextBlock)
-						.Text(LOCTEXT("WorldSpaceRotation", "World Space Rotation").ToString())
+						.Text(LOCTEXT("WorldSpaceRotation", "World Space Rotation"))
 					]
 
 				+SVerticalBox::Slot()
@@ -285,7 +285,7 @@ void SSocketManager::Construct(const FArguments& InArgs)
 						.FillWidth(0.25f)
 						[
 							SNew(STextBlock)
-							.Text(LOCTEXT("Pitch", "Pitch").ToString())
+							.Text(LOCTEXT("Pitch", "Pitch"))
 						]
 						+SHorizontalBox::Slot()
 							.FillWidth(1.0f)
@@ -307,7 +307,7 @@ void SSocketManager::Construct(const FArguments& InArgs)
 						.FillWidth(0.25f)
 						[
 							SNew(STextBlock)
-							.Text(LOCTEXT("Yaw", "Yaw").ToString())
+							.Text(LOCTEXT("Yaw", "Yaw"))
 						]
 						+SHorizontalBox::Slot()
 							.FillWidth(1.0f)
@@ -329,7 +329,7 @@ void SSocketManager::Construct(const FArguments& InArgs)
 						.FillWidth(0.25f)
 						[
 							SNew(STextBlock)
-								.Text(LOCTEXT("Roll", "Roll").ToString())
+								.Text(LOCTEXT("Roll", "Roll"))
 						]
 						+SHorizontalBox::Slot()
 							.FillWidth(1.0f)
diff --git a/Engine/Source/Editor/UnrealEd/Private/UnrealEdEngine.cpp b/Engine/Source/Editor/UnrealEd/Private/UnrealEdEngine.cpp
--- a/Engine/Source/Editor/UnrealEd/Private/UnrealEdEngine.cpp
+++ b/Engine/Source/Editor/UnrealEd/Private/UnrealEdEngine.cpp
@@ -591,7 +591,7 @@ void UUnrealEdEngine::SetMapBuildCancelled( bool InCancelled )
 #define LOCTEXT_NAMESPACE "UnrealEd"
 
 
-FString FClassPickerDefaults::GetName() const
+FText FClassPickerDefaults::GetName() const
 {
 	static TMap< FString, FText > LocNames;
 
@@ -608,26 +608,26 @@ FString FClassPickerDefaults::GetName() const
 	{
 		UClass* ItemClass = LoadClass<UObject>(NULL, *ClassName, NULL, LOAD_None, NULL);
 		check( ItemClass );
-		return FName::NameToDisplayString(ItemClass->GetName(), false);
+		return FText::FromString(FName::NameToDisplayString(ItemClass->GetName(), false));
 	}
 	
 	const FText* PreExistingName = LocNames.Find( LocTextNameID );
 	if ( PreExistingName )
 	{
-		return PreExistingName->ToString();
+		return *PreExistingName;
 	}
 
 	FText OutName;
 	if ( FText::FindText(TEXT("UnrealEd"), LocTextNameID, OutName) )
 	{
-		return OutName.ToString();
+		return OutName;
 	}
 		
-	return LocTextNameID;
+	return FText::FromString(LocTextNameID);
 }
 
 
-FString FClassPickerDefaults::GetDescription() const
+FText FClassPickerDefaults::GetDescription() const
 {
 	static TMap< FString, FText > LocDescs;
 
@@ -642,22 +642,22 @@ FString FClassPickerDefaults::GetDescription() const
 
 	if ( LocTextDescriptionID.IsEmpty() )
 	{
-		return LOCTEXT("NoClassPickerDesc", "No Description.").ToString();
+		return LOCTEXT("NoClassPickerDesc", "No Description.");
 	}
 
 	const FText* PreExistingDesc = LocDescs.Find( LocTextDescriptionID );
 	if ( PreExistingDesc )
 	{
-		return PreExistingDesc->ToString();
+		return *PreExistingDesc;
 	}
 
 	FText OutDesc;
 	if ( FText::FindText(TEXT("UnrealEd"), LocTextDescriptionID, OutDesc) )
 	{
-		return OutDesc.ToString();
+		return OutDesc;
 	}
 		
-	return LocTextDescriptionID;
+	return FText::FromString(LocTextDescriptionID);
 }
 
 #undef LOCTEXT_NAMESPACE
