Fixed code relying on SLATE_TEXT_ARGUMENT for SNumericEntryBox

Made sure everything was using FText rather than FString.

diff --git a/Engine/Source/Editor/DetailCustomizations/Private/MarginCustomization.cpp b/Engine/Source/Editor/DetailCustomizations/Private/MarginCustomization.cpp
--- a/Engine/Source/Editor/DetailCustomizations/Private/MarginCustomization.cpp
+++ b/Engine/Source/Editor/DetailCustomizations/Private/MarginCustomization.cpp
@@ -70,7 +70,7 @@ TSharedRef<SEditableTextBox> FMarginStructCustomization::MakePropertyWidget()
 	return
 		SAssignNew( MarginEditableTextBox,SEditableTextBox )
 		.Text( this, &FMarginStructCustomization::GetMarginText )
-		.ToolTipText( NSLOCTEXT( "UnrealEd", "MarginPropertyToolTip", "Margin values" ).ToString() )
+		.ToolTipText( NSLOCTEXT( "UnrealEd", "MarginPropertyToolTip", "Margin values" ) )
 		.OnTextCommitted( this, &FMarginStructCustomization::OnMarginTextCommitted )
 		.Font( IDetailLayoutBuilder::GetDetailFont() )
 		.SelectAllTextWhenFocused( true )
@@ -83,7 +83,7 @@ TSharedRef<SWidget> FMarginStructCustomization::MakeChildPropertyWidget( int32 P
 		SNew( SNumericEntryBox<float> )
 		.Value( this, &FMarginStructCustomization::OnGetValue, PropertyIndex )
 		.Font( IDetailLayoutBuilder::GetDetailFont() )
-		.UndeterminedString( NSLOCTEXT( "PropertyEditor", "MultipleValues", "Multiple Values").ToString() )
+		.UndeterminedString( NSLOCTEXT( "PropertyEditor", "MultipleValues", "Multiple Values") )
 		.OnValueCommitted( this, &FMarginStructCustomization::OnValueCommitted, PropertyIndex )
 		.OnValueChanged( this, &FMarginStructCustomization::OnValueChanged, PropertyIndex )
 		.OnBeginSliderMovement( this, &FMarginStructCustomization::OnBeginSliderMovement )
diff --git a/Engine/Source/Editor/DetailCustomizations/Private/MathStructCustomizations.cpp b/Engine/Source/Editor/DetailCustomizations/Private/MathStructCustomizations.cpp
--- a/Engine/Source/Editor/DetailCustomizations/Private/MathStructCustomizations.cpp
+++ b/Engine/Source/Editor/DetailCustomizations/Private/MathStructCustomizations.cpp
@@ -213,7 +213,7 @@ TSharedRef<SWidget> FMathStructCustomization::MakeNumericWidget(
 		.IsEnabled( this, &FMathStructCustomization::IsValueEnabled, WeakHandlePtr )
 		.Value( this, &FMathStructCustomization::OnGetValue, WeakHandlePtr )
 		.Font( IDetailLayoutBuilder::GetDetailFont() )
-		.UndeterminedString( NSLOCTEXT("PropertyEditor", "MultipleValues", "Multiple Values").ToString() )
+		.UndeterminedString( NSLOCTEXT("PropertyEditor", "MultipleValues", "Multiple Values") )
 		.OnValueCommitted( this, &FMathStructCustomization::OnValueCommitted<NumericType>, WeakHandlePtr )
 		.OnValueChanged( this, &FMathStructCustomization::OnValueChanged<NumericType>, WeakHandlePtr )
 		.OnBeginSliderMovement( this, &FMathStructCustomization::OnBeginSliderMovement )
diff --git a/Engine/Source/Editor/DetailCustomizations/Private/MathStructProxyCustomizations.cpp b/Engine/Source/Editor/DetailCustomizations/Private/MathStructProxyCustomizations.cpp
--- a/Engine/Source/Editor/DetailCustomizations/Private/MathStructProxyCustomizations.cpp
+++ b/Engine/Source/Editor/DetailCustomizations/Private/MathStructProxyCustomizations.cpp
@@ -25,7 +25,7 @@ TSharedRef<SWidget> FMathStructProxyCustomization::MakeNumericProxyWidget(TShare
 		SNew( SNumericEntryBox<NumericType> )
 		.Value( this, &FMathStructProxyCustomization::OnGetValue<ProxyType, NumericType>, WeakHandlePtr, ProxyValue )
 		.Font( IDetailLayoutBuilder::GetDetailFont() )
-		.UndeterminedString( NSLOCTEXT("PropertyEditor", "MultipleValues", "Multiple Values").ToString() )
+		.UndeterminedString( NSLOCTEXT("PropertyEditor", "MultipleValues", "Multiple Values") )
 		.OnValueCommitted( this, &FMathStructProxyCustomization::OnValueCommitted<ProxyType, NumericType>, WeakHandlePtr, ProxyValue )
 		.OnValueChanged( this, &FMathStructProxyCustomization::OnValueChanged<ProxyType, NumericType>, WeakHandlePtr, ProxyValue )
 		.OnBeginSliderMovement( this, &FMathStructProxyCustomization::OnBeginSliderMovement )
diff --git a/Engine/Source/Editor/PropertyEditor/Private/UserInterface/PropertyEditor/SPropertyEditorNumeric.h b/Engine/Source/Editor/PropertyEditor/Private/UserInterface/PropertyEditor/SPropertyEditorNumeric.h
--- a/Engine/Source/Editor/PropertyEditor/Private/UserInterface/PropertyEditor/SPropertyEditorNumeric.h
+++ b/Engine/Source/Editor/PropertyEditor/Private/UserInterface/PropertyEditor/SPropertyEditorNumeric.h
@@ -102,7 +102,7 @@ public:
 				.MaxSliderValue(SliderMaxValue)
 				.SliderExponent(SliderExponent)
 				.Delta(Delta)
-				.UndeterminedString(LOCTEXT("MultipleValues", "Multiple Values").ToString())
+				.UndeterminedString(LOCTEXT("MultipleValues", "Multiple Values"))
 				.OnValueChanged(this, &SPropertyEditorNumeric<NumericType>::OnValueChanged)
 				.OnValueCommitted(this, &SPropertyEditorNumeric<NumericType>::OnValueCommitted)
 				.OnBeginSliderMovement(this, &SPropertyEditorNumeric<NumericType>::OnBeginSliderMovement)
diff --git a/Engine/Source/Editor/UnrealEd/Private/SCurveEditor.cpp b/Engine/Source/Editor/UnrealEd/Private/SCurveEditor.cpp
--- a/Engine/Source/Editor/UnrealEd/Private/SCurveEditor.cpp
+++ b/Engine/Source/Editor/UnrealEd/Private/SCurveEditor.cpp
@@ -213,7 +213,7 @@ void SCurveEditor::Construct(const FArguments& InArgs)
 							.IsEnabled(this, &SCurveEditor::GetInputEditEnabled)
 							.Font(FEditorStyle::GetFontStyle("CurveEd.InfoFont"))
 							.Value(this, &SCurveEditor::OnGetTime)
-							.UndeterminedString(LOCTEXT("MultipleValues", "Multiple Values").ToString())
+							.UndeterminedString(LOCTEXT("MultipleValues", "Multiple Values"))
 							.OnValueCommitted(this, &SCurveEditor::OnTimeComitted)
 							.OnValueChanged(this, &SCurveEditor::OnTimeChanged)
 							.OnBeginSliderMovement(this, &SCurveEditor::OnBeginSliderMovement, LOCTEXT("SetValue", "Set New Time"))
@@ -245,7 +245,7 @@ void SCurveEditor::Construct(const FArguments& InArgs)
 							SNew(SNumericEntryBox<float>)
 							.Font(FEditorStyle::GetFontStyle("CurveEd.InfoFont"))
 							.Value(this, &SCurveEditor::OnGetValue)
-							.UndeterminedString(LOCTEXT("MultipleValues", "Multiple Values").ToString())
+							.UndeterminedString(LOCTEXT("MultipleValues", "Multiple Values"))
 							.OnValueCommitted(this, &SCurveEditor::OnValueComitted)
 							.OnValueChanged(this, &SCurveEditor::OnValueChanged)
 							.OnBeginSliderMovement(this, &SCurveEditor::OnBeginSliderMovement, LOCTEXT("SetValue", "Set New Value"))
