Fixed "Revert to Default" in Project thumbnail image.

diff --git a/Engine/Source/Editor/DetailCustomizations/Private/GeneralProjectSettingsDetails.cpp b/Engine/Source/Editor/DetailCustomizations/Private/GeneralProjectSettingsDetails.cpp
--- a/Engine/Source/Editor/DetailCustomizations/Private/GeneralProjectSettingsDetails.cpp
+++ b/Engine/Source/Editor/DetailCustomizations/Private/GeneralProjectSettingsDetails.cpp
@@ -21,7 +21,11 @@ void FGeneralProjectSettingsDetails::CustomizeDetails(IDetailLayoutBuilder& Deta
 	const FText ProjectThumbnailDesc(LOCTEXT("ProjectThumbnailLabel", "Project Thumbnail"));
 
 	const FString ProjectThumbnail_TargetImagePath = FPaths::GetPath(FPaths::GetProjectFilePath()) / FString::Printf(TEXT("%s.png"), FApp::GetGameName());
-	const FString ProjectThumbnail_AutomaticImagePath = FPaths::GameSavedDir() / TEXT("AutoScreenshot.png");
+	FString ProjectThumbnail_AutomaticImagePath = FPaths::GameSavedDir() / TEXT("AutoScreenshot.png");
+	if (!FPaths::FileExists(ProjectThumbnail_AutomaticImagePath))
+	{
+		ProjectThumbnail_AutomaticImagePath = FPaths::EngineContentDir() / TEXT("Editor") / TEXT("Slate") / TEXT("GameProjectDialog") / TEXT("default_game_thumbnail_192x.png");
+	}
 
 	ImageCategory.AddCustomRow(ProjectThumbnailDesc)
 	.NameContent()
