Fixed code relying on the version of SWidget::SetToolTipText that took an FString

Made sure everything was using FText rather than FString, and also updated the property editor API to return tooltips as FText rather than FString.

diff --git a/Engine/Source/Developer/LogVisualizer/Private/SLogBar.cpp b/Engine/Source/Developer/LogVisualizer/Private/SLogBar.cpp
--- a/Engine/Source/Developer/LogVisualizer/Private/SLogBar.cpp
+++ b/Engine/Source/Developer/LogVisualizer/Private/SLogBar.cpp
@@ -399,12 +399,11 @@ void SLogBar::SetEntries(const TArray<TSharedPtr<FVisualLogEntry> >& InEntries,
 
 	if (bDrawHistogramCursor)
 	{
-		//const FText ToolTip(TEXT("Use mouse wheel and left shift to modify graph window"));
-		SetToolTipText(TAttribute<FString>(TEXT("Use mouse wheel and left shift to modify 2d graph cursor size")));
+		SetToolTipText(LOCTEXT("ModifyHistogramCursorToolTip", "Use mouse wheel and left shift to modify 2d graph cursor size"));
 	}
 	else
 	{
-		SetToolTipText(FText());
+		SetToolTipText(FText::GetEmpty());
 	}
 }
 
diff --git a/Engine/Source/Developer/TaskGraph/Private/SGraphBar.cpp b/Engine/Source/Developer/TaskGraph/Private/SGraphBar.cpp
--- a/Engine/Source/Developer/TaskGraph/Private/SGraphBar.cpp
+++ b/Engine/Source/Developer/TaskGraph/Private/SGraphBar.cpp
@@ -199,11 +199,11 @@ FReply SGraphBar::OnMouseMove( const FGeometry& MyGeometry, const FPointerEvent&
 	{
 		if( HoveredEventIndex != INDEX_NONE )
 		{
-			this->SetToolTipText(TAttribute<FString>(Events[HoveredEventIndex]->EventName));
+			this->SetToolTipText( FText::FromString(Events[HoveredEventIndex]->EventName) );
 		}
 		else
 		{
-			this->SetToolTipText( FText() );
+			this->SetToolTipText( FText::GetEmpty() );
 		}
 		LastHoveredEvent = HoveredEventIndex;
 	}
diff --git a/Engine/Source/Editor/ClassViewer/Private/SClassViewer.cpp b/Engine/Source/Editor/ClassViewer/Private/SClassViewer.cpp
--- a/Engine/Source/Editor/ClassViewer/Private/SClassViewer.cpp
+++ b/Engine/Source/Editor/ClassViewer/Private/SClassViewer.cpp
@@ -2470,7 +2470,7 @@ TSharedRef< ITableRow > SClassViewer::OnGenerateRowForClassViewer( TSharedPtr<FC
 
 	if ( !Item->GeneratedClassPackage.IsEmpty() )
 	{
-		ReturnRow->SetToolTipText(TAttribute<FString>(Item->GeneratedClassPackage));
+		ReturnRow->SetToolTipText(FText::FromString(Item->GeneratedClassPackage));
 	}
 
 	// Expand the item if needed.
diff --git a/Engine/Source/Editor/DetailCustomizations/Private/ParticleSysParamStructCustomization.cpp b/Engine/Source/Editor/DetailCustomizations/Private/ParticleSysParamStructCustomization.cpp
--- a/Engine/Source/Editor/DetailCustomizations/Private/ParticleSysParamStructCustomization.cpp
+++ b/Engine/Source/Editor/DetailCustomizations/Private/ParticleSysParamStructCustomization.cpp
@@ -153,14 +153,14 @@ EVisibility FParticleSysParamStructCustomization::GetMaterialVisibility() const
 TSharedRef<SWidget> FParticleSysParamStructCustomization::OnGenerateComboWidget(TSharedPtr<FString> InComboString)
 {
 	// Find ToolTip which corresponds to string
-	FString ToolTip;
+	FText ToolTip;
 	check(ParameterTypeNames.Num() == ParameterTypeToolTips.Num());
 	if (ParameterTypeToolTips.Num() > 0)
 	{
 		int32 Index = ParameterTypeNames.IndexOfByKey(InComboString);
 		if (ensure(Index >= 0))
 		{
-			ToolTip = *ParameterTypeToolTips[Index];
+			ToolTip = ParameterTypeToolTips[Index];
 		}
 	}
 
diff --git a/Engine/Source/Editor/DetailCustomizations/Private/ParticleSysParamStructCustomization.h b/Engine/Source/Editor/DetailCustomizations/Private/ParticleSysParamStructCustomization.h
--- a/Engine/Source/Editor/DetailCustomizations/Private/ParticleSysParamStructCustomization.h
+++ b/Engine/Source/Editor/DetailCustomizations/Private/ParticleSysParamStructCustomization.h
@@ -40,7 +40,7 @@ private:
 	TArray<TSharedPtr<FString>> ParameterTypeNames;
 
 	/** A list of parameter type tooltips */
-	TArray<TSharedPtr<FString>> ParameterTypeToolTips;
+	TArray<FText> ParameterTypeToolTips;
 
 	/** Delegates called by each property widget to determine their visibility */
 	EVisibility GetScalarVisibility() const;
diff --git a/Engine/Source/Editor/DetailCustomizations/Private/PhysicsConstraintComponentDetails.cpp b/Engine/Source/Editor/DetailCustomizations/Private/PhysicsConstraintComponentDetails.cpp
--- a/Engine/Source/Editor/DetailCustomizations/Private/PhysicsConstraintComponentDetails.cpp
+++ b/Engine/Source/Editor/DetailCustomizations/Private/PhysicsConstraintComponentDetails.cpp
@@ -106,7 +106,7 @@ void FPhysicsConstraintComponentDetails::CustomizeDetails( IDetailLayoutBuilder&
 		LinearZMotionProperty = ConstraintInstance->GetChildHandle("LinearZMotion");
 		
 		TArray<TSharedPtr<FString>> LinearLimitOptionNames;
-		TArray<TSharedPtr<FString>> LinearLimitOptionTooltips;
+		TArray<FText> LinearLimitOptionTooltips;
 		TArray<bool> LinearLimitOptionRestrictItems;
 
 		const int32 ExpectedLinearLimitOptionCount = 3;
@@ -144,7 +144,7 @@ void FPhysicsConstraintComponentDetails::CustomizeDetails( IDetailLayoutBuilder&
 						.Style(FEditorStyle::Get(), "RadioButton")
 						.IsChecked( this, &FPhysicsConstraintComponentDetails::IsLimitRadioChecked, CurProperty, LinearLimitEnum[0])
 						.OnCheckStateChanged( this, &FPhysicsConstraintComponentDetails::OnLimitRadioChanged, CurProperty, LinearLimitEnum[0] )
-						.ToolTipText(*LinearLimitOptionTooltips[0].Get())
+						.ToolTipText(LinearLimitOptionTooltips[0])
 						[
 							SNew(STextBlock).Text(*LinearLimitOptionNames[0].Get())
 						]						
@@ -157,7 +157,7 @@ void FPhysicsConstraintComponentDetails::CustomizeDetails( IDetailLayoutBuilder&
 							.Style(FEditorStyle::Get(), "RadioButton")
 							.IsChecked( this, &FPhysicsConstraintComponentDetails::IsLimitRadioChecked, CurProperty, LinearLimitEnum[1])
 							.OnCheckStateChanged( this, &FPhysicsConstraintComponentDetails::OnLimitRadioChanged, CurProperty, LinearLimitEnum[1] )
-							.ToolTipText(*LinearLimitOptionTooltips[1].Get())
+							.ToolTipText(LinearLimitOptionTooltips[1])
 							[
 								SNew(STextBlock).Text(*LinearLimitOptionNames[1].Get())
 							]
@@ -170,7 +170,7 @@ void FPhysicsConstraintComponentDetails::CustomizeDetails( IDetailLayoutBuilder&
 							.Style(FEditorStyle::Get(), "RadioButton")
 							.IsChecked( this, &FPhysicsConstraintComponentDetails::IsLimitRadioChecked, CurProperty, LinearLimitEnum[2])
 							.OnCheckStateChanged( this, &FPhysicsConstraintComponentDetails::OnLimitRadioChanged, CurProperty, LinearLimitEnum[2] )
-							.ToolTipText(*LinearLimitOptionTooltips[2].Get())
+							.ToolTipText(LinearLimitOptionTooltips[2])
 							[
 								SNew(STextBlock).Text(*LinearLimitOptionNames[2].Get())
 							]
@@ -200,7 +200,7 @@ void FPhysicsConstraintComponentDetails::CustomizeDetails( IDetailLayoutBuilder&
 		AngularTwistMotionProperty = ConstraintInstance->GetChildHandle(GET_MEMBER_NAME_CHECKED(FConstraintInstance, AngularTwistMotion));
 
 		TArray<TSharedPtr<FString>> AngularLimitOptionNames;
-		TArray<TSharedPtr<FString>> AngularLimitOptionTooltips;
+		TArray<FText> AngularLimitOptionTooltips;
 		TArray<bool> AngularLimitOptionRestrictItems;
 
 		const int32 ExpectedAngularLimitOptionCount = 3;
@@ -237,7 +237,7 @@ void FPhysicsConstraintComponentDetails::CustomizeDetails( IDetailLayoutBuilder&
 						.Style(FEditorStyle::Get(), "RadioButton")
 						.IsChecked( this, &FPhysicsConstraintComponentDetails::IsLimitRadioChecked, CurProperty, AngularLimitEnum[0])
 						.OnCheckStateChanged( this, &FPhysicsConstraintComponentDetails::OnLimitRadioChanged, CurProperty, AngularLimitEnum[0] )
-						.ToolTipText(*AngularLimitOptionTooltips[0].Get())
+						.ToolTipText(AngularLimitOptionTooltips[0])
 						[
 							SNew(STextBlock).Text(*AngularLimitOptionNames[0].Get())
 						]						
@@ -250,7 +250,7 @@ void FPhysicsConstraintComponentDetails::CustomizeDetails( IDetailLayoutBuilder&
 							.Style(FEditorStyle::Get(), "RadioButton")
 							.IsChecked( this, &FPhysicsConstraintComponentDetails::IsLimitRadioChecked, CurProperty, AngularLimitEnum[1])
 							.OnCheckStateChanged( this, &FPhysicsConstraintComponentDetails::OnLimitRadioChanged, CurProperty, AngularLimitEnum[1] )
-							.ToolTipText(*AngularLimitOptionTooltips[1].Get())
+							.ToolTipText(AngularLimitOptionTooltips[1])
 							[
 								SNew(STextBlock).Text(*AngularLimitOptionNames[1].Get())
 							]
@@ -263,7 +263,7 @@ void FPhysicsConstraintComponentDetails::CustomizeDetails( IDetailLayoutBuilder&
 							.Style(FEditorStyle::Get(), "RadioButton")
 							.IsChecked( this, &FPhysicsConstraintComponentDetails::IsLimitRadioChecked, CurProperty, AngularLimitEnum[2])
 							.OnCheckStateChanged( this, &FPhysicsConstraintComponentDetails::OnLimitRadioChanged, CurProperty, AngularLimitEnum[2] )
-							.ToolTipText(*AngularLimitOptionTooltips[2].Get())
+							.ToolTipText(AngularLimitOptionTooltips[2])
 							[
 								SNew(STextBlock).Text(*AngularLimitOptionNames[2].Get())
 							]
diff --git a/Engine/Source/Editor/DetailCustomizations/Private/RangeStructCustomization.cpp b/Engine/Source/Editor/DetailCustomizations/Private/RangeStructCustomization.cpp
--- a/Engine/Source/Editor/DetailCustomizations/Private/RangeStructCustomization.cpp
+++ b/Engine/Source/Editor/DetailCustomizations/Private/RangeStructCustomization.cpp
@@ -339,7 +339,7 @@ bool FRangeStructCustomization<NumericType>::ShouldAllowSpin() const
 template <typename NumericType>
 TSharedRef<SWidget> FRangeStructCustomization<NumericType>::OnGenerateComboWidget(TSharedPtr<FString> InComboString)
 {
-	FString ToolTip;
+	FText ToolTip;
 
 	// A list of tool tips should have been populated in a 1 to 1 correspondence
 	check(ComboBoxList.Num() == ComboBoxToolTips.Num());
@@ -349,7 +349,7 @@ TSharedRef<SWidget> FRangeStructCustomization<NumericType>::OnGenerateComboWidge
 		int32 Index = ComboBoxList.IndexOfByKey(InComboString);
 		if (ensure(Index >= 0))
 		{
-			ToolTip = *ComboBoxToolTips[Index];
+			ToolTip = ComboBoxToolTips[Index];
 		}
 	}
 
diff --git a/Engine/Source/Editor/DetailCustomizations/Private/RangeStructCustomization.h b/Engine/Source/Editor/DetailCustomizations/Private/RangeStructCustomization.h
--- a/Engine/Source/Editor/DetailCustomizations/Private/RangeStructCustomization.h
+++ b/Engine/Source/Editor/DetailCustomizations/Private/RangeStructCustomization.h
@@ -115,7 +115,7 @@ protected:
 
 	// Arrays of combo box data: list items, and tooltips
 	TArray< TSharedPtr<FString> > ComboBoxList;
-	TArray< TSharedPtr<FString> > ComboBoxToolTips;
+	TArray< FText > ComboBoxToolTips;
 
 	// Flags whether the slider is being moved at the moment on any of our widgets
 	bool bIsUsingSlider;
diff --git a/Engine/Source/Editor/MaterialEditor/Private/MaterialEditorDetailCustomization.cpp b/Engine/Source/Editor/MaterialEditor/Private/MaterialEditorDetailCustomization.cpp
--- a/Engine/Source/Editor/MaterialEditor/Private/MaterialEditorDetailCustomization.cpp
+++ b/Engine/Source/Editor/MaterialEditor/Private/MaterialEditorDetailCustomization.cpp
@@ -139,15 +139,15 @@ FMaterialExpressionCollectionParameterDetails::FMaterialExpressionCollectionPara
 {
 }
 
-FString FMaterialExpressionCollectionParameterDetails::GetToolTipText() const
+FText FMaterialExpressionCollectionParameterDetails::GetToolTipText() const
 {
 	if (ParametersSource.Num() == 1)
 	{
-		return LOCTEXT("SpecifyCollection", "Specify a Collection to get parameter options").ToString();
+		return LOCTEXT("SpecifyCollection", "Specify a Collection to get parameter options");
 	}
 	else
 	{
-		return LOCTEXT("ChooseParameter", "Choose a parameter from the collection").ToString();
+		return LOCTEXT("ChooseParameter", "Choose a parameter from the collection");
 	}
 }
 
@@ -246,7 +246,7 @@ void FMaterialExpressionCollectionParameterDetails::CustomizeDetails( IDetailLay
 	ParameterComboButton = NewComboButton;
 	ParameterListView = NewListView;
 
-	NewComboButton->SetToolTipText(TAttribute<FString>(GetToolTipText()));
+	NewComboButton->SetToolTipText(GetToolTipText());
 }
 
 void FMaterialExpressionCollectionParameterDetails::PopulateParameters()
diff --git a/Engine/Source/Editor/MaterialEditor/Private/MaterialEditorDetailCustomization.h b/Engine/Source/Editor/MaterialEditor/Private/MaterialEditorDetailCustomization.h
--- a/Engine/Source/Editor/MaterialEditor/Private/MaterialEditorDetailCustomization.h
+++ b/Engine/Source/Editor/MaterialEditor/Private/MaterialEditorDetailCustomization.h
@@ -70,7 +70,7 @@ private:
 	TSharedRef< ITableRow > MakeDetailsGroupViewWidget( TSharedPtr<FString> Item, const TSharedRef< STableViewBase >& OwnerTable );
 	void OnSelectionChanged( TSharedPtr<FString> ProposedSelection, ESelectInfo::Type SelectInfo );
 
-	FString GetToolTipText() const;
+	FText GetToolTipText() const;
 	FString GetParameterNameString() const;
 	bool IsParameterNameComboEnabled() const;
 	void OnCollectionChanged();
diff --git a/Engine/Source/Editor/PropertyEditor/Private/PropertyHandleImpl.cpp b/Engine/Source/Editor/PropertyEditor/Private/PropertyHandleImpl.cpp
--- a/Engine/Source/Editor/PropertyEditor/Private/PropertyHandleImpl.cpp
+++ b/Engine/Source/Editor/PropertyEditor/Private/PropertyHandleImpl.cpp
@@ -1840,7 +1840,7 @@ FPropertyAccess::Result FPropertyHandleBase::GetPerObjectValues( TArray<FString>
 	return Result;
 }
 
-bool FPropertyHandleBase::GeneratePossibleValues(TArray< TSharedPtr<FString> >& OutOptionStrings, TArray< TSharedPtr<FString> >& OutToolTips, TArray<bool>& OutRestrictedItems)
+bool FPropertyHandleBase::GeneratePossibleValues(TArray< TSharedPtr<FString> >& OutOptionStrings, TArray< FText >& OutToolTips, TArray<bool>& OutRestrictedItems)
 {
 	UProperty* Property = GetProperty();
 
@@ -1895,9 +1895,8 @@ bool FPropertyHandleBase::GeneratePossibleValues(TArray< TSharedPtr<FString> >&
 				TSharedPtr< FString > EnumStr( new FString( EnumValueName ) );
 				OutOptionStrings.Add( EnumStr );
 
-				FString EnumValueToolTip = bIsRestricted ? RestrictionTooltip.ToString() : Enum->GetToolTipText(EnumIndex).ToString();
-				TSharedPtr< FString > ToolTipStr( new FString( EnumValueToolTip ) );
-				OutToolTips.Add( ToolTipStr );
+				FText EnumValueToolTip = bIsRestricted ? RestrictionTooltip : Enum->GetToolTipText(EnumIndex);
+				OutToolTips.Add( EnumValueToolTip );
 			}
 		}
 	}
diff --git a/Engine/Source/Editor/PropertyEditor/Private/PropertyHandleImpl.h b/Engine/Source/Editor/PropertyEditor/Private/PropertyHandleImpl.h
--- a/Engine/Source/Editor/PropertyEditor/Private/PropertyHandleImpl.h
+++ b/Engine/Source/Editor/PropertyEditor/Private/PropertyHandleImpl.h
@@ -406,7 +406,7 @@ public:
 	virtual void SetToolTipText(const FText& ToolTip) override;
 	virtual FPropertyAccess::Result SetPerObjectValues( const TArray<FString>& InPerObjectValues,  EPropertyValueSetFlags::Type Flags = EPropertyValueSetFlags::DefaultFlags ) override;
 	virtual FPropertyAccess::Result GetPerObjectValues( TArray<FString>& OutPerObjectValues ) override;
-	virtual bool GeneratePossibleValues(TArray< TSharedPtr<FString> >& OutOptionStrings, TArray< TSharedPtr<FString> >& OutToolTips, TArray<bool>& OutRestrictedItems) override;
+	virtual bool GeneratePossibleValues(TArray< TSharedPtr<FString> >& OutOptionStrings, TArray< FText >& OutToolTips, TArray<bool>& OutRestrictedItems) override;
 	virtual FPropertyAccess::Result SetObjectValueFromSelection() override;
 	virtual void NotifyPreChange() override;
 	virtual void NotifyPostChange() override;
diff --git a/Engine/Source/Editor/PropertyEditor/Private/UserInterface/PropertyEditor/SPropertyComboBox.cpp b/Engine/Source/Editor/PropertyEditor/Private/UserInterface/PropertyEditor/SPropertyComboBox.cpp
--- a/Engine/Source/Editor/PropertyEditor/Private/UserInterface/PropertyEditor/SPropertyComboBox.cpp
+++ b/Engine/Source/Editor/PropertyEditor/Private/UserInterface/PropertyEditor/SPropertyComboBox.cpp
@@ -20,7 +20,7 @@ void SPropertyComboBox::Construct( const FArguments& InArgs )
 	{
 		if(*ComboItemList[ItemIndex].Get() == VisibleText)
 		{
-			SetToolTipText(TAttribute<FString>(*ToolTipList[ItemIndex]));
+			SetToolTipText(ToolTipList[ItemIndex]);
 			InitiallySelectedItem = ComboItemList[ItemIndex];
 			break;
 		}
@@ -56,7 +56,7 @@ void SPropertyComboBox::SetSelectedItem( const FString& InSelectedItem )
 	{
 		if(*ComboItemList[ItemIndex].Get() == InSelectedItem)
 		{
-			SetToolTipText(TAttribute<FString>(*ToolTipList[ItemIndex]));
+			SetToolTipText(ToolTipList[ItemIndex]);
 
 			SComboBox< TSharedPtr<FString> >::SetSelectedItem(ComboItemList[ItemIndex]);
 			return;
@@ -64,7 +64,7 @@ void SPropertyComboBox::SetSelectedItem( const FString& InSelectedItem )
 	}
 }
 
-void SPropertyComboBox::SetItemList(TArray< TSharedPtr< FString > >& InItemList, TArray< TSharedPtr< FString > >& InTooltipList, TArray<bool>& InRestrictedList)
+void SPropertyComboBox::SetItemList(TArray< TSharedPtr< FString > >& InItemList, TArray< FText >& InTooltipList, TArray<bool>& InRestrictedList)
 {
 	ComboItemList = InItemList;
 	ToolTipList = InTooltipList;
@@ -99,7 +99,7 @@ void SPropertyComboBox::OnSelectionChangedInternal( TSharedPtr<FString> InSelect
 TSharedRef<SWidget> SPropertyComboBox::OnGenerateComboWidget( TSharedPtr<FString> InComboString )
 {
 	//Find the corresponding tool tip for this combo entry if any
-	FString ToolTip;
+	FText ToolTip;
 	bool bEnabled = true;
 	if (ToolTipList.Num() > 0)
 	{
@@ -108,7 +108,7 @@ TSharedRef<SWidget> SPropertyComboBox::OnGenerateComboWidget( TSharedPtr<FString
 		{
 			//A list of tool tips should have been populated in a 1 to 1 correspondance
 			check(ComboItemList.Num() == ToolTipList.Num());
-			ToolTip = *ToolTipList[Index];
+			ToolTip = ToolTipList[Index];
 
 			if( RestrictedList.Num() > 0 )
 			{
diff --git a/Engine/Source/Editor/PropertyEditor/Private/UserInterface/PropertyEditor/SPropertyComboBox.h b/Engine/Source/Editor/PropertyEditor/Private/UserInterface/PropertyEditor/SPropertyComboBox.h
--- a/Engine/Source/Editor/PropertyEditor/Private/UserInterface/PropertyEditor/SPropertyComboBox.h
+++ b/Engine/Source/Editor/PropertyEditor/Private/UserInterface/PropertyEditor/SPropertyComboBox.h
@@ -13,7 +13,7 @@ public:
 		SLATE_ATTRIBUTE( TArray< TSharedPtr< FString > >, ComboItemList )
 		SLATE_ATTRIBUTE( TArray< bool >, RestrictedList )
 		SLATE_ATTRIBUTE( FString, VisibleText )
-		SLATE_ARGUMENT( TArray< TSharedPtr< FString > >, ToolTipList )
+		SLATE_ARGUMENT( TArray< FText >, ToolTipList )
 		SLATE_EVENT( FOnSelectionChanged, OnSelectionChanged )
 		SLATE_EVENT( FOnComboBoxOpening, OnComboBoxOpening )
 		SLATE_ARGUMENT( FSlateFontInfo, Font )
@@ -36,7 +36,7 @@ public:
 	 *	@param InTooltipList		The list of tooltips to use.
 	 *	@param InRestrictedList		The list of restricted items.
 	 */
-	void SetItemList(TArray< TSharedPtr< FString > >& InItemList, TArray< TSharedPtr< FString > >& InTooltipList, TArray<bool>& InRestrictedList);
+	void SetItemList(TArray< TSharedPtr< FString > >& InItemList, TArray< FText >& InTooltipList, TArray<bool>& InRestrictedList);
 
 
 private:
@@ -52,7 +52,7 @@ private:
 
 	/** List of items in our combo box. Only generated once as combo items dont change at runtime */
 	TArray< TSharedPtr<FString> > ComboItemList;
-	TArray< TSharedPtr<FString> > ToolTipList;
+	TArray< FText > ToolTipList;
 	FOnSelectionChanged OnSelectionChanged;
 	FSlateFontInfo Font;
 	TArray< bool > RestrictedList;
diff --git a/Engine/Source/Editor/PropertyEditor/Private/UserInterface/PropertyEditor/SPropertyEditorCombo.cpp b/Engine/Source/Editor/PropertyEditor/Private/UserInterface/PropertyEditor/SPropertyEditorCombo.cpp
--- a/Engine/Source/Editor/PropertyEditor/Private/UserInterface/PropertyEditor/SPropertyEditorCombo.cpp
+++ b/Engine/Source/Editor/PropertyEditor/Private/UserInterface/PropertyEditor/SPropertyEditorCombo.cpp
@@ -68,7 +68,7 @@ void SPropertyEditorCombo::Construct( const FArguments& InArgs, const TSharedRef
 	
 	// Important to find certain info out about the combo box.
 	TArray< TSharedPtr<FString> > ComboItems;
-	TArray< TSharedPtr<FString> > ToolTips;
+	TArray< FText > ToolTips;
 	TArray< bool > Restrictions;
 	GenerateComboBoxStrings( ComboItems, ToolTips, Restrictions );
 
@@ -110,7 +110,7 @@ FString SPropertyEditorCombo::GetDisplayValueAsString() const
 	return (bUsesAlternateDisplayValues) ? PropertyEditor->GetValueAsDisplayString() : PropertyEditor->GetValueAsString();
 }
 
-void SPropertyEditorCombo::GenerateComboBoxStrings( TArray< TSharedPtr<FString> >& OutComboBoxStrings, TArray< TSharedPtr<FString> >& OutToolTips, TArray<bool>& OutRestrictedItems )
+void SPropertyEditorCombo::GenerateComboBoxStrings( TArray< TSharedPtr<FString> >& OutComboBoxStrings, TArray< FText >& OutToolTips, TArray<bool>& OutRestrictedItems )
 {
 	const TSharedRef< IPropertyHandle > PropertyHandle = PropertyEditor->GetPropertyHandle();
 	bUsesAlternateDisplayValues = PropertyHandle->GeneratePossibleValues(OutComboBoxStrings, OutToolTips, OutRestrictedItems);
@@ -138,7 +138,6 @@ void SPropertyEditorCombo::SendToObjects( const FString& NewValue )
 	UProperty* Property = PropertyNode->GetProperty();
 
 	FString Value;
-	FString ToolTipValue;
 	if ( bUsesAlternateDisplayValues && !Property->IsA(UStrProperty::StaticClass()))
 	{
 		// currently only enum properties can use alternate display values; this 
@@ -152,13 +151,13 @@ void SPropertyEditorCombo::SendToObjects( const FString& NewValue )
 
 		Value = Enum->GetEnumName(Index);
 
-		ToolTipValue = Enum->GetToolTipText(Index).ToString();
-		FString ToolTipText = Property->GetToolTipText().ToString();
-		if (ToolTipValue.Len() > 0)
+		FText ToolTipValue = Enum->GetToolTipText(Index);
+		FText ToolTipText = Property->GetToolTipText();
+		if (!ToolTipValue.IsEmpty())
 		{
-			ToolTipText = FString::Printf(TEXT("%s\n\n%s"), *ToolTipText, *ToolTipValue);
+			ToolTipText = FText::Format(FText::FromString(TEXT("{0}\n\n{1}")), ToolTipText, ToolTipValue);
 		}
-		SetToolTipText(TAttribute<FString>(ToolTipText));
+		SetToolTipText(ToolTipText);
 	}
 	else
 	{
diff --git a/Engine/Source/Editor/PropertyEditor/Private/UserInterface/PropertyEditor/SPropertyEditorCombo.h b/Engine/Source/Editor/PropertyEditor/Private/UserInterface/PropertyEditor/SPropertyEditorCombo.h
--- a/Engine/Source/Editor/PropertyEditor/Private/UserInterface/PropertyEditor/SPropertyEditorCombo.h
+++ b/Engine/Source/Editor/PropertyEditor/Private/UserInterface/PropertyEditor/SPropertyEditorCombo.h
@@ -20,7 +20,7 @@ public:
 	void GetDesiredWidth( float& OutMinDesiredWidth, float& OutMaxDesiredWidth );
 private:
 
-	void GenerateComboBoxStrings( TArray< TSharedPtr<FString> >& OutComboBoxStrings, TArray< TSharedPtr<FString> >& OutToolTips, TArray<bool>& OutRestrictedItems );
+	void GenerateComboBoxStrings( TArray< TSharedPtr<FString> >& OutComboBoxStrings, TArray< FText >& OutToolTips, TArray<bool>& OutRestrictedItems );
 	void OnComboSelectionChanged( TSharedPtr<FString> NewValue, ESelectInfo::Type SelectInfo );
 	void OnComboOpening();
 
diff --git a/Engine/Source/Editor/PropertyEditor/Private/UserInterface/PropertyEditor/SPropertyEditorDateTime.cpp b/Engine/Source/Editor/PropertyEditor/Private/UserInterface/PropertyEditor/SPropertyEditorDateTime.cpp
--- a/Engine/Source/Editor/PropertyEditor/Private/UserInterface/PropertyEditor/SPropertyEditorDateTime.cpp
+++ b/Engine/Source/Editor/PropertyEditor/Private/UserInterface/PropertyEditor/SPropertyEditorDateTime.cpp
@@ -27,7 +27,7 @@ void SPropertyEditorDateTime::Construct( const FArguments& InArgs, const TShared
 	if( InPropertyEditor->PropertyIsA( UObjectPropertyBase::StaticClass() ) )
 	{
 		// Object properties should display their entire text in a tooltip
-		PrimaryWidget->SetToolTipText( TAttribute<FString>( InPropertyEditor, &FPropertyEditor::GetValueAsString ) );
+		PrimaryWidget->SetToolTipText( TAttribute<FText>( InPropertyEditor, &FPropertyEditor::GetValueAsText ) );
 	}
 }
 
diff --git a/Engine/Source/Editor/PropertyEditor/Private/UserInterface/PropertyEditor/SPropertyEditorText.cpp b/Engine/Source/Editor/PropertyEditor/Private/UserInterface/PropertyEditor/SPropertyEditorText.cpp
--- a/Engine/Source/Editor/PropertyEditor/Private/UserInterface/PropertyEditor/SPropertyEditorText.cpp
+++ b/Engine/Source/Editor/PropertyEditor/Private/UserInterface/PropertyEditor/SPropertyEditorText.cpp
@@ -48,7 +48,7 @@ void SPropertyEditorText::Construct( const FArguments& InArgs, const TSharedRef<
 	if( InPropertyEditor->PropertyIsA( UObjectPropertyBase::StaticClass() ) )
 	{
 		// Object properties should display their entire text in a tooltip
-		PrimaryWidget->SetToolTipText( TAttribute<FString>( InPropertyEditor, &FPropertyEditor::GetValueAsString ) );
+		PrimaryWidget->SetToolTipText( TAttribute<FText>( InPropertyEditor, &FPropertyEditor::GetValueAsText ) );
 	}
 
 	SetEnabled( TAttribute<bool>( this, &SPropertyEditorText::CanEdit ) );
diff --git a/Engine/Source/Editor/PropertyEditor/Public/PropertyHandle.h b/Engine/Source/Editor/PropertyEditor/Public/PropertyHandle.h
--- a/Engine/Source/Editor/PropertyEditor/Public/PropertyHandle.h
+++ b/Engine/Source/Editor/PropertyEditor/Public/PropertyHandle.h
@@ -314,7 +314,7 @@ public:
 	/**
 	 * Generates a list of possible enum/class options for the property
 	 */
-	virtual bool GeneratePossibleValues(TArray< TSharedPtr<FString> >& OutOptionStrings, TArray< TSharedPtr<FString> >& OutToolTips, TArray<bool>& OutRestrictedItems) = 0;
+	virtual bool GeneratePossibleValues(TArray< TSharedPtr<FString> >& OutOptionStrings, TArray< FText >& OutToolTips, TArray<bool>& OutRestrictedItems) = 0;
 
 	/**
 	 * Marks this property has hidden by customizaton (will not show up in the default place)
