UE-6068 The Default Value for a Timespan variable resets whenever its Blueprint is compiled

WithIdenticalViaEquality = true in TStructOpsTypeTraits<FDateTime> and TStructOpsTypeTraits<FTimespan>

diff --git a/Engine/Source/Editor/DetailCustomizations/Private/DateTimeStructCustomization.cpp b/Engine/Source/Editor/DetailCustomizations/Private/DateTimeStructCustomization.cpp
--- a/Engine/Source/Editor/DetailCustomizations/Private/DateTimeStructCustomization.cpp
+++ b/Engine/Source/Editor/DetailCustomizations/Private/DateTimeStructCustomization.cpp
@@ -80,16 +80,19 @@ void FDateTimeStructCustomization::HandleTextBoxTextChanged( const FText& NewTex
 void FDateTimeStructCustomization::HandleTextBoxTextCommited( const FText& NewText, ETextCommit::Type CommitInfo )
 {
 	FDateTime ParsedDateTime;
-								
-	if (FDateTime::Parse(NewText.ToString(), ParsedDateTime))
+		
+	InputValid = FDateTime::Parse(NewText.ToString(), ParsedDateTime);
+	if (InputValid && PropertyHandle.IsValid())
 	{
 		TArray<void*> RawData;
 		PropertyHandle->AccessRawData(RawData);
 
+		PropertyHandle->NotifyPreChange();
 		for (auto RawDataInstance : RawData)
 		{
 			*(FDateTime*)RawDataInstance = ParsedDateTime;
 		}
+		PropertyHandle->NotifyPostChange();
 	}
 }
 
diff --git a/Engine/Source/Editor/DetailCustomizations/Private/TimespanStructCustomization.cpp b/Engine/Source/Editor/DetailCustomizations/Private/TimespanStructCustomization.cpp
--- a/Engine/Source/Editor/DetailCustomizations/Private/TimespanStructCustomization.cpp
+++ b/Engine/Source/Editor/DetailCustomizations/Private/TimespanStructCustomization.cpp
@@ -80,16 +80,19 @@ void FTimespanStructCustomization::HandleTextBoxTextChanged( const FText& NewTex
 void FTimespanStructCustomization::HandleTextBoxTextCommited( const FText& NewText, ETextCommit::Type CommitInfo )
 {
 	FTimespan ParsedTimespan;
-								
-	if (FTimespan::Parse(NewText.ToString(), ParsedTimespan))
+			
+	InputValid = FTimespan::Parse(NewText.ToString(), ParsedTimespan);
+	if (InputValid && PropertyHandle.IsValid())
 	{
 		TArray<void*> RawData;
 		PropertyHandle->AccessRawData(RawData);
 
+		PropertyHandle->NotifyPreChange();
 		for (auto RawDataInstance : RawData)
 		{
 			*(FTimespan*)RawDataInstance = ParsedTimespan;
 		}
+		PropertyHandle->NotifyPostChange();
 	}
 }
 
diff --git a/Engine/Source/Runtime/CoreUObject/Private/UObject/Property.cpp b/Engine/Source/Runtime/CoreUObject/Private/UObject/Property.cpp
--- a/Engine/Source/Runtime/CoreUObject/Private/UObject/Property.cpp
+++ b/Engine/Source/Runtime/CoreUObject/Private/UObject/Property.cpp
@@ -236,6 +236,7 @@ struct TStructOpsTypeTraits<FDateTime> : public TStructOpsTypeTraitsBase
 		WithImportTextItem = true,
 		WithSerializer = true,
 		WithZeroConstructor = true,
+		WithIdenticalViaEquality = true,
 	};
 };
 IMPLEMENT_STRUCT(DateTime);
@@ -250,6 +251,7 @@ struct TStructOpsTypeTraits<FTimespan> : public TStructOpsTypeTraitsBase
 		WithImportTextItem = false, // @todo gmp: implement FTimespan::ImportTextItem
 		WithSerializer = true,
 		WithZeroConstructor = true,
+		WithIdenticalViaEquality = true,
 	};
 };
 IMPLEMENT_STRUCT(Timespan);
