Speculative fix for a crash when a nested widget is changed

UE-5795 - UMG: Crash when changing font size on TextBlock in a widget that was referenced by another widget (both UMG editors opened)

We were getting a crash using the cached values of SlateFontInfoStructs after updating a nested widget in UMG. This has been changed to get the value of SlateFontInfoStructs as required, first ensuring that the struct property handle is still valid.

diff --git a/Engine/Source/Editor/DetailCustomizations/Private/SlateFontInfoCustomization.cpp b/Engine/Source/Editor/DetailCustomizations/Private/SlateFontInfoCustomization.cpp
--- a/Engine/Source/Editor/DetailCustomizations/Private/SlateFontInfoCustomization.cpp
+++ b/Engine/Source/Editor/DetailCustomizations/Private/SlateFontInfoCustomization.cpp
@@ -12,29 +12,23 @@ TSharedRef<IPropertyTypeCustomization> FSlateFontInfoStructCustomization::MakeIn
 	return MakeShareable(new FSlateFontInfoStructCustomization());
 }
 
-void FSlateFontInfoStructCustomization::CustomizeHeader(TSharedRef<IPropertyHandle> StructPropertyHandle, FDetailWidgetRow& HeaderRow, IPropertyTypeCustomizationUtils& StructCustomizationUtils)
+void FSlateFontInfoStructCustomization::CustomizeHeader(TSharedRef<IPropertyHandle> InStructPropertyHandle, FDetailWidgetRow& InHeaderRow, IPropertyTypeCustomizationUtils& InStructCustomizationUtils)
 {
-	FontObjectProperty = StructPropertyHandle->GetChildHandle(TEXT("FontObject"));
+	StructPropertyHandle = InStructPropertyHandle;
+
+	FontObjectProperty = InStructPropertyHandle->GetChildHandle(FName("FontObject"));
 	check(FontObjectProperty.IsValid());
 
-	TypefaceFontNameProperty = StructPropertyHandle->GetChildHandle(TEXT("TypefaceFontName"));
+	TypefaceFontNameProperty = InStructPropertyHandle->GetChildHandle(FName("TypefaceFontName"));
 	check(TypefaceFontNameProperty.IsValid());
 
-	FontSizeProperty = StructPropertyHandle->GetChildHandle(TEXT("Size"));
+	FontSizeProperty = InStructPropertyHandle->GetChildHandle(FName("Size"));
 	check(FontSizeProperty.IsValid());
 
-	TArray<void*> StructPtrs;
-	StructPropertyHandle->AccessRawData(StructPtrs);
-	for(auto It = StructPtrs.CreateConstIterator(); It; ++It)
-	{
-		FSlateFontInfo* const SlateFontInfo = reinterpret_cast<FSlateFontInfo*>(*It);
-		SlateFontInfoStructs.Add(SlateFontInfo);
-	}
-
-	HeaderRow
+	InHeaderRow
 	.NameContent()
 	[
-		StructPropertyHandle->CreatePropertyNameWidget()
+		InStructPropertyHandle->CreatePropertyNameWidget()
 	]
 	.ValueContent()
 	.MinDesiredWidth(250.0f)
@@ -84,7 +78,7 @@ void FSlateFontInfoStructCustomization::CustomizeHeader(TSharedRef<IPropertyHand
 	];
 }
 
-void FSlateFontInfoStructCustomization::CustomizeChildren(TSharedRef<IPropertyHandle> StructPropertyHandle, IDetailChildrenBuilder& StructBuilder, IPropertyTypeCustomizationUtils& StructCustomizationUtils)
+void FSlateFontInfoStructCustomization::CustomizeChildren(TSharedRef<IPropertyHandle> InStructPropertyHandle, IDetailChildrenBuilder& InStructBuilder, IPropertyTypeCustomizationUtils& InStructCustomizationUtils)
 {
 }
 
@@ -99,6 +93,7 @@ void FSlateFontInfoStructCustomization::OnFontChanged(const FAssetData& InAssetD
 	const UFont* const FontAsset = Cast<const UFont>(InAssetData.GetAsset());
 	const FName FirstFontName = (FontAsset && FontAsset->CompositeFont.DefaultTypeface.Fonts.Num()) ? FontAsset->CompositeFont.DefaultTypeface.Fonts[0].Name : NAME_None;
 
+	TArray<FSlateFontInfo*> SlateFontInfoStructs = GetFontInfoBeingEdited();
 	for(FSlateFontInfo* FontInfo : SlateFontInfoStructs)
 	{
 		// The font has been updated in the editor, so clear the non-UObject pointer so that the two don't conflict
@@ -111,6 +106,12 @@ void FSlateFontInfoStructCustomization::OnFontChanged(const FAssetData& InAssetD
 
 bool FSlateFontInfoStructCustomization::IsFontEntryComboEnabled() const
 {
+	TArray<const FSlateFontInfo*> SlateFontInfoStructs = GetFontInfoBeingEdited();
+	if(SlateFontInfoStructs.Num() == 0)
+	{
+		return false;
+	}
+
 	const FSlateFontInfo* const FirstSlateFontInfo = SlateFontInfoStructs[0];
 	const UFont* const FontObject = Cast<const UFont>(FirstSlateFontInfo->FontObject);
 	if(!FontObject)
@@ -134,43 +135,58 @@ bool FSlateFontInfoStructCustomization::IsFontEntryComboEnabled() const
 
 void FSlateFontInfoStructCustomization::OnFontEntryComboOpening()
 {
-	const FSlateFontInfo* const FirstSlateFontInfo = SlateFontInfoStructs[0];
-	const UFont* const FontObject = Cast<const UFont>(FirstSlateFontInfo->FontObject);
-	check(FontObject);
-
-	const FName ActiveFontEntry = GetActiveFontEntry();
-	TSharedPtr<FName> SelectedNamePtr;
+	TArray<FSlateFontInfo*> SlateFontInfoStructs = GetFontInfoBeingEdited();
 
 	FontEntryComboData.Empty();
-	for(const FTypefaceEntry& TypefaceEntry : FontObject->CompositeFont.DefaultTypeface.Fonts)
+
+	if(SlateFontInfoStructs.Num() > 0)
 	{
-		TSharedPtr<FName> NameEntryPtr = MakeShareable(new FName(TypefaceEntry.Name));
-		FontEntryComboData.Add(NameEntryPtr);
+		const FSlateFontInfo* const FirstSlateFontInfo = SlateFontInfoStructs[0];
+		const UFont* const FontObject = Cast<const UFont>(FirstSlateFontInfo->FontObject);
+		check(FontObject);
+
+		const FName ActiveFontEntry = GetActiveFontEntry();
+		TSharedPtr<FName> SelectedNamePtr;
 
-		if(!TypefaceEntry.Name.IsNone() && TypefaceEntry.Name == ActiveFontEntry)
+		for(const FTypefaceEntry& TypefaceEntry : FontObject->CompositeFont.DefaultTypeface.Fonts)
 		{
-			SelectedNamePtr = NameEntryPtr;
+			TSharedPtr<FName> NameEntryPtr = MakeShareable(new FName(TypefaceEntry.Name));
+			FontEntryComboData.Add(NameEntryPtr);
+
+			if(!TypefaceEntry.Name.IsNone() && TypefaceEntry.Name == ActiveFontEntry)
+			{
+				SelectedNamePtr = NameEntryPtr;
+			}
 		}
-	}
 
-	FontEntryComboData.Sort([](const TSharedPtr<FName>& One, const TSharedPtr<FName>& Two) -> bool
-	{
-		return One->ToString() < Two->ToString();
-	});
+		FontEntryComboData.Sort([](const TSharedPtr<FName>& One, const TSharedPtr<FName>& Two) -> bool
+		{
+			return One->ToString() < Two->ToString();
+		});
 
-	FontEntryCombo->ClearSelection();
-	FontEntryCombo->RefreshOptions();
-	FontEntryCombo->SetSelectedItem(SelectedNamePtr);
+		FontEntryCombo->ClearSelection();
+		FontEntryCombo->RefreshOptions();
+		FontEntryCombo->SetSelectedItem(SelectedNamePtr);
+	}
+	else
+	{
+		FontEntryCombo->ClearSelection();
+		FontEntryCombo->RefreshOptions();
+	}
 }
 
 void FSlateFontInfoStructCustomization::OnFontEntrySelectionChanged(TSharedPtr<FName> InNewSelection, ESelectInfo::Type)
 {
 	if(InNewSelection.IsValid())
 	{
-		const FSlateFontInfo* const FirstSlateFontInfo = SlateFontInfoStructs[0];
-		if(FirstSlateFontInfo->TypefaceFontName != *InNewSelection)
+		TArray<FSlateFontInfo*> SlateFontInfoStructs = GetFontInfoBeingEdited();
+		if(SlateFontInfoStructs.Num() > 0)
 		{
-			TypefaceFontNameProperty->SetValue(*InNewSelection);
+			const FSlateFontInfo* const FirstSlateFontInfo = SlateFontInfoStructs[0];
+			if(FirstSlateFontInfo->TypefaceFontName != *InNewSelection)
+			{
+				TypefaceFontNameProperty->SetValue(*InNewSelection);
+			}
 		}
 	}
 }
@@ -190,16 +206,60 @@ FText FSlateFontInfoStructCustomization::GetFontEntryComboText() const
 
 FName FSlateFontInfoStructCustomization::GetActiveFontEntry() const
 {
-	const FSlateFontInfo* const FirstSlateFontInfo = SlateFontInfoStructs[0];
-	const UFont* const FontObject = Cast<const UFont>(FirstSlateFontInfo->FontObject);
-	if(FontObject)
+	TArray<const FSlateFontInfo*> SlateFontInfoStructs = GetFontInfoBeingEdited();
+	if(SlateFontInfoStructs.Num() > 0)
 	{
-		return (FirstSlateFontInfo->TypefaceFontName.IsNone() && FontObject->CompositeFont.DefaultTypeface.Fonts.Num())
-			? FontObject->CompositeFont.DefaultTypeface.Fonts[0].Name
-			: FirstSlateFontInfo->TypefaceFontName;
+		const FSlateFontInfo* const FirstSlateFontInfo = SlateFontInfoStructs[0];
+		const UFont* const FontObject = Cast<const UFont>(FirstSlateFontInfo->FontObject);
+		if(FontObject)
+		{
+			return (FirstSlateFontInfo->TypefaceFontName.IsNone() && FontObject->CompositeFont.DefaultTypeface.Fonts.Num())
+				? FontObject->CompositeFont.DefaultTypeface.Fonts[0].Name
+				: FirstSlateFontInfo->TypefaceFontName;
+		}
 	}
 
 	return NAME_None;
 }
 
+TArray<FSlateFontInfo*> FSlateFontInfoStructCustomization::GetFontInfoBeingEdited()
+{
+	TArray<FSlateFontInfo*> SlateFontInfoStructs;
+
+	if(StructPropertyHandle->IsValidHandle())
+	{
+		TArray<void*> StructPtrs;
+		StructPropertyHandle->AccessRawData(StructPtrs);
+		SlateFontInfoStructs.Reserve(StructPtrs.Num());
+
+		for(auto It = StructPtrs.CreateConstIterator(); It; ++It)
+		{
+			FSlateFontInfo* const SlateFontInfo = reinterpret_cast<FSlateFontInfo*>(*It);
+			SlateFontInfoStructs.Add(SlateFontInfo);
+		}
+	}
+
+	return SlateFontInfoStructs;
+}
+
+TArray<const FSlateFontInfo*> FSlateFontInfoStructCustomization::GetFontInfoBeingEdited() const
+{
+	TArray<const FSlateFontInfo*> SlateFontInfoStructs;
+
+	if(StructPropertyHandle->IsValidHandle())
+	{
+		TArray<const void*> StructPtrs;
+		StructPropertyHandle->AccessRawData(StructPtrs);
+		SlateFontInfoStructs.Reserve(StructPtrs.Num());
+
+		for(auto It = StructPtrs.CreateConstIterator(); It; ++It)
+		{
+			const FSlateFontInfo* const SlateFontInfo = reinterpret_cast<const FSlateFontInfo*>(*It);
+			SlateFontInfoStructs.Add(SlateFontInfo);
+		}
+	}
+
+	return SlateFontInfoStructs;
+}
+
 #undef LOCTEXT_NAMESPACE
diff --git a/Engine/Source/Editor/DetailCustomizations/Public/Customizations/SlateFontInfoCustomization.h b/Engine/Source/Editor/DetailCustomizations/Public/Customizations/SlateFontInfoCustomization.h
--- a/Engine/Source/Editor/DetailCustomizations/Public/Customizations/SlateFontInfoCustomization.h
+++ b/Engine/Source/Editor/DetailCustomizations/Public/Customizations/SlateFontInfoCustomization.h
@@ -9,8 +9,8 @@ public:
 	static TSharedRef<IPropertyTypeCustomization> MakeInstance();
 
 	/** IPropertyTypeCustomization interface */
-	virtual void CustomizeHeader(TSharedRef<class IPropertyHandle> StructPropertyHandle, class FDetailWidgetRow& HeaderRow, IPropertyTypeCustomizationUtils& StructCustomizationUtils) override;
-	virtual void CustomizeChildren(TSharedRef<class IPropertyHandle> StructPropertyHandle, class IDetailChildrenBuilder& StructBuilder, IPropertyTypeCustomizationUtils& StructCustomizationUtils) override;
+	virtual void CustomizeHeader(TSharedRef<class IPropertyHandle> InStructPropertyHandle, class FDetailWidgetRow& InHeaderRow, IPropertyTypeCustomizationUtils& InStructCustomizationUtils) override;
+	virtual void CustomizeChildren(TSharedRef<class IPropertyHandle> InStructPropertyHandle, class IDetailChildrenBuilder& InStructBuilder, IPropertyTypeCustomizationUtils& InStructCustomizationUtils) override;
 
 protected:
 	/** Called to filter out invalid font assets */
@@ -37,8 +37,12 @@ protected:
 	/** Get the name of the currently active font entry (may not be the selected entry if the entry is set to use "None") */
 	FName GetActiveFontEntry() const;
 
-	/** Array of FSlateFontInfo instances this customization is currently editing */
-	TArray<FSlateFontInfo*> SlateFontInfoStructs;
+	/** Get the array of FSlateFontInfo instances this customization is currently editing */
+	TArray<FSlateFontInfo*> GetFontInfoBeingEdited();
+	TArray<const FSlateFontInfo*> GetFontInfoBeingEdited() const;
+
+	/** Handle to the struct property being edited */
+	TSharedPtr<IPropertyHandle> StructPropertyHandle;
 
 	/** Handle to the "FontObject" property being edited */
 	TSharedPtr<IPropertyHandle> FontObjectProperty;
