Fixed crash getting the font info being edited while transitioning to PIE

UE-6203 - Crash when PIE after changing value of wrap text in UMG

IPropertyHandle::AccessRawData can return null values while we're transitioning to PIE, and these should not be added to the array of structs returned by FSlateFontInfoStructCustomization::GetFontInfoBeingEdited as all of those entries are assumed to be valid, non-null, pointers.

diff --git a/Engine/Source/Editor/DetailCustomizations/Private/SlateFontInfoCustomization.cpp b/Engine/Source/Editor/DetailCustomizations/Private/SlateFontInfoCustomization.cpp
--- a/Engine/Source/Editor/DetailCustomizations/Private/SlateFontInfoCustomization.cpp
+++ b/Engine/Source/Editor/DetailCustomizations/Private/SlateFontInfoCustomization.cpp
@@ -234,8 +234,12 @@ TArray<FSlateFontInfo*> FSlateFontInfoStructCustomization::GetFontInfoBeingEdite
 
 		for(auto It = StructPtrs.CreateConstIterator(); It; ++It)
 		{
-			FSlateFontInfo* const SlateFontInfo = reinterpret_cast<FSlateFontInfo*>(*It);
-			SlateFontInfoStructs.Add(SlateFontInfo);
+			void* RawPtr = *It;
+			if(RawPtr)
+			{
+				FSlateFontInfo* const SlateFontInfo = reinterpret_cast<FSlateFontInfo*>(RawPtr);
+				SlateFontInfoStructs.Add(SlateFontInfo);
+			}
 		}
 	}
 
@@ -254,8 +258,12 @@ TArray<const FSlateFontInfo*> FSlateFontInfoStructCustomization::GetFontInfoBein
 
 		for(auto It = StructPtrs.CreateConstIterator(); It; ++It)
 		{
-			const FSlateFontInfo* const SlateFontInfo = reinterpret_cast<const FSlateFontInfo*>(*It);
-			SlateFontInfoStructs.Add(SlateFontInfo);
+			const void* RawPtr = *It;
+			if(RawPtr)
+			{
+				const FSlateFontInfo* const SlateFontInfo = reinterpret_cast<const FSlateFontInfo*>(RawPtr);
+				SlateFontInfoStructs.Add(SlateFontInfo);
+			}
 		}
 	}
 
