#ue4 Converting "Composited Blueprint" text to "Convert to Class BLueprint"

diff --git a/Engine/Source/Editor/DetailCustomizations/Private/ActorDetails.cpp b/Engine/Source/Editor/DetailCustomizations/Private/ActorDetails.cpp
--- a/Engine/Source/Editor/DetailCustomizations/Private/ActorDetails.cpp
+++ b/Engine/Source/Editor/DetailCustomizations/Private/ActorDetails.cpp
@@ -789,14 +789,14 @@ void FActorDetails::AddUtilityBlueprintRows( IDetailCategoryBuilder& BlueprintCa
 			.WidthOverride( 200 )
 			[
 				SNew( SButton )
-				.ToolTipText( LOCTEXT( "CreateHarvestBlueprint_ToolTip", "Harvest Components from Selected Actors and create Blueprint" ) )
+				.ToolTipText( LOCTEXT( "CreateHarvestBlueprint_ToolTip", "Copy the components from the selected actors and create a new Class Blueprint" ) )
 				.HAlign( HAlign_Fill )
 				//.ButtonColorAndOpacity( FLinearColor( 0.2f, 0.4f, 0.6f, 1.0f ) )
 				.OnClicked( FOnClicked::CreateRaw( this, &FActorDetails::OnPickBlueprintPathClicked, true ) )
 				.ContentPadding( 2 )
 				[
 					SNew( STextBlock )
-					.Text( LOCTEXT( "ReplaceWithAmalgamBlueprint", "Replace With Composited Blueprint" ) )
+					.Text( LOCTEXT( "ReplaceWithAmalgamBlueprint", "Convert to Class Blueprint" ) )
 					.Font(IDetailLayoutBuilder::GetDetailFont())
 				]
 			]
