Properly hide Tick properties on instances of Actors.

diff --git a/Engine/Source/Editor/DetailCustomizations/Private/ActorDetails.cpp b/Engine/Source/Editor/DetailCustomizations/Private/ActorDetails.cpp
--- a/Engine/Source/Editor/DetailCustomizations/Private/ActorDetails.cpp
+++ b/Engine/Source/Editor/DetailCustomizations/Private/ActorDetails.cpp
@@ -137,18 +137,19 @@ void FActorDetails::CustomizeDetails( IDetailLayoutBuilder& DetailLayout )
 		//AddComponentsCategory( DetailLayout );
 	}
 
-	// Defaults only
+	TSharedPtr<IPropertyHandle> PrimaryTickProperty = DetailLayout.GetProperty(GET_MEMBER_NAME_CHECKED(AActor, PrimaryActorTick));
+
+	// Defaults only show tick properties
 	if (DetailLayout.GetDetailsView().HasClassDefaultObject())
 	{
 		IDetailCategoryBuilder& TickCategory = DetailLayout.EditCategory("Tick");
 
-		TSharedPtr<IPropertyHandle> PrimaryTickProperty = DetailLayout.GetProperty(GET_MEMBER_NAME_CHECKED(AActor, PrimaryActorTick));
 		TickCategory.AddProperty(PrimaryTickProperty->GetChildHandle(GET_MEMBER_NAME_CHECKED(FTickFunction, bStartWithTickEnabled)));
 		TickCategory.AddProperty(PrimaryTickProperty->GetChildHandle(GET_MEMBER_NAME_CHECKED(FTickFunction, bTickEvenWhenPaused)), EPropertyLocation::Advanced);
 		TickCategory.AddProperty(PrimaryTickProperty->GetChildHandle(GET_MEMBER_NAME_CHECKED(FTickFunction, bAllowTickOnDedicatedServer)), EPropertyLocation::Advanced);
-
-		PrimaryTickProperty->MarkHiddenByCustomization();
 	}
+
+	PrimaryTickProperty->MarkHiddenByCustomization();
 }
 
 void FActorDetails::OnConvertActor(UClass* ChosenClass)
