FStringClassReference's customization was accessing internal property of FStringAssetReference and was setting it without checking if path is "None". FStringClassReference is inheriting FStringAssetReference::ImportTextItem, which is checking for that special case. This method is used when customization treats the property as a whole (without accessing internal property). Fixed UE-4538.

Also fixed assert+crash when user tried to use FStringClassReference, without specifying "MetaClass" in UPROPERTY macro. UHT allowed it, but the editor was crashing. Now when it's missing it is assumed it is UClass.

diff --git a/Engine/Source/Editor/DetailCustomizations/Private/StringClassReferenceCustomization.cpp b/Engine/Source/Editor/DetailCustomizations/Private/StringClassReferenceCustomization.cpp
--- a/Engine/Source/Editor/DetailCustomizations/Private/StringClassReferenceCustomization.cpp
+++ b/Engine/Source/Editor/DetailCustomizations/Private/StringClassReferenceCustomization.cpp
@@ -4,27 +4,25 @@
 #include "StringClassReferenceCustomization.h"
 #include "EditorClassUtils.h"
 
-void FStringClassReferenceCustomization::CustomizeHeader(TSharedRef<IPropertyHandle> InStructPropertyHandle, FDetailWidgetRow& HeaderRow, IPropertyTypeCustomizationUtils& StructCustomizationUtils)
+void FStringClassReferenceCustomization::CustomizeHeader(TSharedRef<IPropertyHandle> InPropertyHandle, FDetailWidgetRow& HeaderRow, IPropertyTypeCustomizationUtils& StructCustomizationUtils)
 {
-	StructPropertyHandle = InStructPropertyHandle;
+	PropertyHandle = InPropertyHandle;
 
-	ClassNamePropertyHandle = InStructPropertyHandle->GetChildHandle("AssetLongPathname");
-	check(ClassNamePropertyHandle.IsValid());
+	const FString& MetaClassName = PropertyHandle->GetMetaData("MetaClass");
+	const FString& RequiredInterfaceName = PropertyHandle->GetMetaData("RequiredInterface");
+	const bool bAllowAbstract = PropertyHandle->GetBoolMetaData("AllowAbstract");
+	const bool bIsBlueprintBaseOnly = PropertyHandle->GetBoolMetaData("IsBlueprintBaseOnly");
+	const bool bAllowNone = !(PropertyHandle->GetMetaDataProperty()->PropertyFlags & CPF_NoClear);
 
-	const FString& MetaClassName = StructPropertyHandle->GetMetaData("MetaClass");
-	const FString& RequiredInterfaceName = StructPropertyHandle->GetMetaData("RequiredInterface");
-	const bool bAllowAbstract = StructPropertyHandle->GetBoolMetaData("AllowAbstract");
-	const bool bIsBlueprintBaseOnly = StructPropertyHandle->GetBoolMetaData("IsBlueprintBaseOnly");
-	const bool bAllowNone = !(StructPropertyHandle->GetMetaDataProperty()->PropertyFlags & CPF_NoClear);
-
-	check(!MetaClassName.IsEmpty());
-	const UClass* const MetaClass = FEditorClassUtils::GetClassFromString(MetaClassName);
+	const UClass* const MetaClass = !MetaClassName.IsEmpty()
+		? FEditorClassUtils::GetClassFromString(MetaClassName)
+		: UClass::StaticClass();
 	const UClass* const RequiredInterface = FEditorClassUtils::GetClassFromString(RequiredInterfaceName);
 
 	HeaderRow
 	.NameContent()
 	[
-		InStructPropertyHandle->CreatePropertyNameWidget()
+		InPropertyHandle->CreatePropertyNameWidget()
 	]
 	.ValueContent()
 	.MinDesiredWidth(250.0f)
@@ -49,7 +47,7 @@ void FStringClassReferenceCustomization::CustomizeChildren(TSharedRef<IPropertyH
 const UClass* FStringClassReferenceCustomization::OnGetClass() const
 {
 	FString ClassName;
-	ClassNamePropertyHandle->GetValueAsFormattedString(ClassName);
+	PropertyHandle->GetValueAsFormattedString(ClassName);
 
 	// Do we have a valid cached class pointer?
 	const UClass* Class = CachedClassPtr.Get();
@@ -63,7 +61,7 @@ const UClass* FStringClassReferenceCustomization::OnGetClass() const
 
 void FStringClassReferenceCustomization::OnSetClass(const UClass* NewClass)
 {
-	if(ClassNamePropertyHandle->SetValueFromFormattedString((NewClass) ? NewClass->GetPathName() : "None") == FPropertyAccess::Result::Success)
+	if (PropertyHandle->SetValueFromFormattedString((NewClass) ? NewClass->GetPathName() : "None") == FPropertyAccess::Result::Success)
 	{
 		CachedClassPtr = NewClass;
 	}
diff --git a/Engine/Source/Editor/DetailCustomizations/Private/StringClassReferenceCustomization.h b/Engine/Source/Editor/DetailCustomizations/Private/StringClassReferenceCustomization.h
--- a/Engine/Source/Editor/DetailCustomizations/Private/StringClassReferenceCustomization.h
+++ b/Engine/Source/Editor/DetailCustomizations/Private/StringClassReferenceCustomization.h
@@ -14,8 +14,8 @@ public:
 	}
 
 	/** IPropertyTypeCustomization interface */
-	virtual void CustomizeHeader(TSharedRef<class IPropertyHandle> InStructPropertyHandle, class FDetailWidgetRow& HeaderRow, IPropertyTypeCustomizationUtils& StructCustomizationUtils) override;
-	virtual void CustomizeChildren(TSharedRef<class IPropertyHandle> InStructPropertyHandle, class IDetailChildrenBuilder& StructBuilder, IPropertyTypeCustomizationUtils& StructCustomizationUtils) override;
+	virtual void CustomizeHeader(TSharedRef<class IPropertyHandle> InPropertyHandle, class FDetailWidgetRow& HeaderRow, IPropertyTypeCustomizationUtils& StructCustomizationUtils) override;
+	virtual void CustomizeChildren(TSharedRef<class IPropertyHandle> InPropertyHandle, class IDetailChildrenBuilder& StructBuilder, IPropertyTypeCustomizationUtils& StructCustomizationUtils) override;
 
 private:
 	/** @return The class currently set on this reference */
@@ -26,10 +26,8 @@ private:
 	/** Find or load the class associated with the given string */
 	static const UClass* StringToClass(const FString& ClassName);
 
-	/** Handle to the struct property being customized */
-	TSharedPtr<IPropertyHandle> StructPropertyHandle;
-	/** Handle to the class name property we're editing */
-	TSharedPtr<IPropertyHandle> ClassNamePropertyHandle;
+	/** Handle to the property being customized */
+	TSharedPtr<IPropertyHandle> PropertyHandle;
 	/** A cache of the currently resolved value for the class name */
 	mutable TWeakObjectPtr<UClass> CachedClassPtr;
 };
