Added button to restart editor when changing hardware settings.

diff --git a/Engine/Source/Editor/DetailCustomizations/Private/HardwareTargetingSettingsDetails.cpp b/Engine/Source/Editor/DetailCustomizations/Private/HardwareTargetingSettingsDetails.cpp
--- a/Engine/Source/Editor/DetailCustomizations/Private/HardwareTargetingSettingsDetails.cpp
+++ b/Engine/Source/Editor/DetailCustomizations/Private/HardwareTargetingSettingsDetails.cpp
@@ -26,7 +26,16 @@ public:
 	{
 		GetMutableDefault<UHardwareTargetingSettings>()->OnSettingChanged().RemoveAll(this);
 	}
-	
+
+private:
+
+	static FReply Apply()
+	{
+		IHardwareTargetingModule& Module = IHardwareTargetingModule::Get();
+		Module.ApplyHardwareTargetingSettings();
+		return FReply::Handled();
+	}
+
 public:
 
 	TMap<TWeakObjectPtr<UObject>, TSharedPtr<SRichTextBlock>> SettingRegions;
@@ -37,9 +46,8 @@ public:
 		GetMutableDefault<UHardwareTargetingSettings>()->OnSettingChanged().AddRaw(this, &SRequiredDefaultConfig::Update);
 
 		auto ApplyNow = []{
-			IHardwareTargetingModule& Module = IHardwareTargetingModule::Get();
-			Module.ApplyHardwareTargetingSettings();
-
+			Apply();
+			FUnrealEdMisc::Get().RestartEditor(false);
 			return FReply::Handled();
 		};
 
@@ -65,10 +73,20 @@ public:
 				.Padding(6)
 				[
 					SNew(SButton)
-					.Text(LOCTEXT("ApplyNow", "Apply Now"))
+					.Text(LOCTEXT("RestartEditor", "Restart Editor"))
 					.IsEnabled(this, &SRequiredDefaultConfig::CanApply)
 					.OnClicked_Static(ApplyNow)
 				]
+
+				+ SHorizontalBox::Slot()
+				.AutoWidth()
+				.Padding(6)
+				[
+					SNew(SButton)
+					.Text(LOCTEXT("ApplyLater", "Apply Later"))
+					.IsEnabled(this, &SRequiredDefaultConfig::CanApply)
+					.OnClicked_Static(&SRequiredDefaultConfig::Apply)
+				]
 			]
 			
 			+ SVerticalBox::Slot()
