Added units to the component transform details

diff --git a/Engine/Source/Editor/DetailCustomizations/Private/ComponentTransformDetails.cpp b/Engine/Source/Editor/DetailCustomizations/Private/ComponentTransformDetails.cpp
--- a/Engine/Source/Editor/DetailCustomizations/Private/ComponentTransformDetails.cpp
+++ b/Engine/Source/Editor/DetailCustomizations/Private/ComponentTransformDetails.cpp
@@ -10,6 +10,9 @@
 #include "ScopedTransaction.h"
 #include "IPropertyUtilities.h"
 
+#include "UnitConversion.h"
+#include "NumericUnitTypeInterface.inl"
+
 #define LOCTEXT_NAMESPACE "FComponentTransformDetails"
 
 class FScopedSwitchWorldForObject
@@ -64,7 +67,8 @@ static void PropagateTransformPropertyChange(UObject* InObject, UProperty* InPro
 
 
 FComponentTransformDetails::FComponentTransformDetails( const TArray< TWeakObjectPtr<UObject> >& InSelectedObjects, const FSelectedActorInfo& InSelectedActorInfo, IDetailLayoutBuilder& DetailBuilder )
-	: SelectedActorInfo( InSelectedActorInfo )
+	: TNumericUnitTypeInterface(EUnit::Centimeters)
+	, SelectedActorInfo( InSelectedActorInfo )
 	, SelectedObjects( InSelectedObjects )
 	, NotifyHook( DetailBuilder.GetPropertyUtilities()->GetNotifyHook() )
 	, bPreserveScaleRatio( false )
@@ -359,9 +363,17 @@ void FComponentTransformDetails::GenerateChildContent( IDetailChildrenBuilder& C
 	const bool bHideRotationField = ( HiddenFieldMask & ( 1 << ETransformField::Rotation ) ) != 0;
 	const bool bHideScaleField = ( HiddenFieldMask & ( 1 << ETransformField::Scale ) ) != 0;
 
+	VectorUnits = EUnit::Centimeters;
+
 	// Location
 	if(!bHideLocationField)
 	{
+		TSharedPtr<INumericTypeInterface<float>> TypeInterface;
+		if( FUnitConversion::bIsUnitDisplayEnabled )
+		{
+			TypeInterface = SharedThis(this);
+		}
+
 		ChildrenBuilder.AddChildContent( LOCTEXT("LocationFilter", "Location") )
 		.CopyAction( CreateCopyAction( ETransformField::Location ) )
 		.PasteAction( CreatePasteAction( ETransformField::Location ) )
@@ -390,6 +402,7 @@ void FComponentTransformDetails::GenerateChildContent( IDetailChildrenBuilder& C
 				.OnYCommitted( this, &FComponentTransformDetails::OnSetLocation, 1 )
 				.OnZCommitted( this, &FComponentTransformDetails::OnSetLocation, 2 )
 				.Font( FontInfo )
+				.TypeInterface( TypeInterface )
 			]
 			+SHorizontalBox::Slot()
 			.AutoWidth()
@@ -420,6 +433,12 @@ void FComponentTransformDetails::GenerateChildContent( IDetailChildrenBuilder& C
 	// Rotation
 	if(!bHideRotationField)
 	{
+		TSharedPtr<INumericTypeInterface<float>> TypeInterface;
+		if( FUnitConversion::bIsUnitDisplayEnabled )
+		{
+			TypeInterface = MakeShareable( new TNumericUnitTypeInterface<float>(EUnit::Degrees) );
+		}
+
 		ChildrenBuilder.AddChildContent( LOCTEXT("RotationFilter", "Rotation") )
 		.CopyAction( CreateCopyAction(ETransformField::Rotation) )
 		.PasteAction( CreatePasteAction(ETransformField::Rotation) )
@@ -453,6 +472,7 @@ void FComponentTransformDetails::GenerateChildContent( IDetailChildrenBuilder& C
 				.OnRollCommitted( this, &FComponentTransformDetails::OnRotationCommitted, 0 )
 				.OnPitchCommitted( this, &FComponentTransformDetails::OnRotationCommitted, 1 )
 				.OnYawCommitted( this, &FComponentTransformDetails::OnRotationCommitted, 2 )
+				.TypeInterface(TypeInterface)
 				.Font( FontInfo )
 			]
 			+SHorizontalBox::Slot()
@@ -557,6 +577,47 @@ END_SLATE_FUNCTION_BUILD_OPTIMIZATION
 void FComponentTransformDetails::Tick( float DeltaTime ) 
 {
 	CacheTransform();
+	CacheCommonLocationUnits();
+}
+
+void FComponentTransformDetails::CacheCommonLocationUnits()
+{
+	// Calculate the largest units
+	TOptional<EUnit> LargestUnit;
+	auto SetLargestUnit = [&](EUnit NewUnit)
+	{
+		if (!LargestUnit.IsSet() || (uint8)NewUnit > (uint8)LargestUnit.GetValue())
+		{
+			LargestUnit = NewUnit;
+		}
+	};
+
+	if (CachedLocation.X.IsSet())
+	{
+		SetLargestUnit(FUnitConversion::QuantizeUnitsToBestFit(CachedLocation.X.GetValue(), EUnit::Centimeters).Units);
+	}
+	if (CachedLocation.Y.IsSet())
+	{
+		SetLargestUnit(FUnitConversion::QuantizeUnitsToBestFit(CachedLocation.Y.GetValue(), EUnit::Centimeters).Units);
+	}
+	if (CachedLocation.Z.IsSet())
+	{
+		SetLargestUnit(FUnitConversion::QuantizeUnitsToBestFit(CachedLocation.Z.GetValue(), EUnit::Centimeters).Units);
+	}
+
+	if (LargestUnit.IsSet())
+	{
+		VectorUnits = LargestUnit.GetValue();
+	}
+	else
+	{
+		VectorUnits	= EUnit::Centimeters;
+	}
+}
+
+FNumericUnit<float> FComponentTransformDetails::QuantizeUnitsToBestFit(const float& InValue, EUnit InUnits) const
+{
+	return FNumericUnit<float>(FUnitConversion::Convert(InValue, InUnits, VectorUnits), VectorUnits);
 }
 
 bool FComponentTransformDetails::GetIsEnabled() const
@@ -1041,6 +1102,7 @@ void FComponentTransformDetails::OnSetLocation( float NewValue, ETextCommit::Typ
 	}
 
 	CacheTransform();
+	CacheCommonLocationUnits();
 
 	GUnrealEd->RedrawLevelEditingViewports();
 }
diff --git a/Engine/Source/Runtime/Slate/Private/Widgets/Input/SRotatorInputBox.cpp b/Engine/Source/Runtime/Slate/Private/Widgets/Input/SRotatorInputBox.cpp
--- a/Engine/Source/Runtime/Slate/Private/Widgets/Input/SRotatorInputBox.cpp
+++ b/Engine/Source/Runtime/Slate/Private/Widgets/Input/SRotatorInputBox.cpp
@@ -46,7 +46,8 @@ void SRotatorInputBox::Construct( const SRotatorInputBox::FArguments& InArgs )
 			.OnBeginSliderMovement( InArgs._OnBeginSliderMovement )
 			.OnEndSliderMovement( InArgs._OnEndSliderMovement )
 			.UndeterminedString( LOCTEXT("MultipleValues", "Multiple Values") )
-			.ToolTipText( LOCTEXT("Roll_ToolTip", "Roll Value") ) 
+			.ToolTipText( LOCTEXT("Roll_ToolTip", "Roll Value") )
+			.TypeInterface(InArgs._TypeInterface)
 		]
 		+SHorizontalBox::Slot()
 		.VAlign( VAlign_Center )
@@ -69,7 +70,8 @@ void SRotatorInputBox::Construct( const SRotatorInputBox::FArguments& InArgs )
 			.OnBeginSliderMovement( InArgs._OnBeginSliderMovement )
 			.OnEndSliderMovement( InArgs._OnEndSliderMovement )
 			.UndeterminedString( LOCTEXT("MultipleValues", "Multiple Values") )
-			.ToolTipText( LOCTEXT("Pitch_ToolTip", "Pitch Value") ) 
+			.ToolTipText( LOCTEXT("Pitch_ToolTip", "Pitch Value") )
+			.TypeInterface(InArgs._TypeInterface)
 		]
 		+SHorizontalBox::Slot()
 		.VAlign( VAlign_Center )
@@ -92,7 +94,8 @@ void SRotatorInputBox::Construct( const SRotatorInputBox::FArguments& InArgs )
 			.OnBeginSliderMovement( InArgs._OnBeginSliderMovement )
 			.OnEndSliderMovement( InArgs._OnEndSliderMovement )
 			.UndeterminedString( LOCTEXT("MultipleValues", "Multiple Values") )
-			.ToolTipText( LOCTEXT("Yaw_ToolTip", "Yaw Value") ) 			
+			.ToolTipText( LOCTEXT("Yaw_ToolTip", "Yaw Value") )
+			.TypeInterface(InArgs._TypeInterface)
 		]
 	];
 
diff --git a/Engine/Source/Runtime/Slate/Private/Widgets/Input/SVectorInputBox.cpp b/Engine/Source/Runtime/Slate/Private/Widgets/Input/SVectorInputBox.cpp
--- a/Engine/Source/Runtime/Slate/Private/Widgets/Input/SVectorInputBox.cpp
+++ b/Engine/Source/Runtime/Slate/Private/Widgets/Input/SVectorInputBox.cpp
@@ -41,6 +41,7 @@ void SVectorInputBox::ConstructX( const FArguments& InArgs, TSharedRef<SHorizont
 		.UndeterminedString( LOCTEXT("MultipleValues", "Multiple Values") )
 		.LabelPadding(0)
 		.ContextMenuExtender( InArgs._ContextMenuExtenderX )
+		.TypeInterface(InArgs._TypeInterface)
 		.Label()
 		[
 			SNumericEntryBox<float>::BuildLabel( LOCTEXT("X_Label", "X"), FLinearColor::White, LabelColor )
@@ -69,6 +70,7 @@ void SVectorInputBox::ConstructY( const FArguments& InArgs, TSharedRef<SHorizont
 		.UndeterminedString( LOCTEXT("MultipleValues", "Multiple Values") )
 		.LabelPadding(0)
 		.ContextMenuExtender( InArgs._ContextMenuExtenderY )
+		.TypeInterface(InArgs._TypeInterface)
 		.Label()
 		[
 			SNumericEntryBox<float>::BuildLabel( LOCTEXT("Y_Label", "Y"), FLinearColor::White, LabelColor )
@@ -96,6 +98,7 @@ void SVectorInputBox::ConstructZ( const FArguments& InArgs, TSharedRef<SHorizont
 		.UndeterminedString( LOCTEXT("MultipleValues", "Multiple Values") )
 		.LabelPadding(0)
 		.ContextMenuExtender( InArgs._ContextMenuExtenderZ )
+		.TypeInterface(InArgs._TypeInterface)
 		.Label()
 		[
 			SNumericEntryBox<float>::BuildLabel( LOCTEXT("Z_Label", "Z"), FLinearColor::White, LabelColor )
diff --git a/Engine/Source/Runtime/Slate/Public/Widgets/Input/SRotatorInputBox.h b/Engine/Source/Runtime/Slate/Public/Widgets/Input/SRotatorInputBox.h
--- a/Engine/Source/Runtime/Slate/Public/Widgets/Input/SRotatorInputBox.h
+++ b/Engine/Source/Runtime/Slate/Public/Widgets/Input/SRotatorInputBox.h
@@ -56,6 +56,9 @@ public:
 		/** Called when the slider for any axis is released */
 		SLATE_EVENT( FOnFloatValueChanged, OnEndSliderMovement )
 
+		/** Provide custom type functionality for the rotator */
+		SLATE_ARGUMENT( TSharedPtr< INumericTypeInterface<float> >, TypeInterface )
+
 	SLATE_END_ARGS()
 
 	/**
diff --git a/Engine/Source/Runtime/Slate/Public/Widgets/Input/SVectorInputBox.h b/Engine/Source/Runtime/Slate/Public/Widgets/Input/SVectorInputBox.h
--- a/Engine/Source/Runtime/Slate/Public/Widgets/Input/SVectorInputBox.h
+++ b/Engine/Source/Runtime/Slate/Public/Widgets/Input/SVectorInputBox.h
@@ -57,6 +57,9 @@ public:
 		/** Menu extender delegate for the Z value */
 		SLATE_EVENT( FMenuExtensionDelegate, ContextMenuExtenderZ )
 
+		/** Provide custom type functionality for the vector */
+		SLATE_ARGUMENT( TSharedPtr< INumericTypeInterface<float> >, TypeInterface )
+
 	SLATE_END_ARGS()
 
 	/**
