Component transform now displays the same units for all (x,y,z) location components

diff --git a/Engine/Source/Editor/DetailCustomizations/Private/ComponentTransformDetails.h b/Engine/Source/Editor/DetailCustomizations/Private/ComponentTransformDetails.h
--- a/Engine/Source/Editor/DetailCustomizations/Private/ComponentTransformDetails.h
+++ b/Engine/Source/Editor/DetailCustomizations/Private/ComponentTransformDetails.h
@@ -16,7 +16,7 @@ namespace ETransformField
 /**
  * Manages the Transform section of a details view                    
  */
-class FComponentTransformDetails : public TSharedFromThis<FComponentTransformDetails>, public IDetailCustomNodeBuilder
+class FComponentTransformDetails : public TSharedFromThis<FComponentTransformDetails>, public IDetailCustomNodeBuilder, public TNumericUnitTypeInterface<float>
 {
 public:
 	FComponentTransformDetails( const TArray< TWeakObjectPtr<UObject> >& InSelectedObjects, const FSelectedActorInfo& InSelectedActorInfo, IDetailLayoutBuilder& DetailBuilder );
@@ -34,6 +34,8 @@ public:
 	virtual void Tick( float DeltaTime ) override;
 	virtual void SetOnRebuildChildren( FSimpleDelegate OnRebuildChildren ) override{}
 
+	virtual FNumericUnit<float> QuantizeUnitsToBestFit(const float& InValue, EUnit InUnits) const override;
+
 	void HideTransformField(const ETransformField::Type InTransformField)
 	{
 		HiddenFieldMask |= (1 << InTransformField);
@@ -264,6 +266,9 @@ private:
 		return CachedScale.IsSet() && CachedScale.X.GetValue() == 1.0f && CachedScale.Y.GetValue() == 1.0f && CachedScale.Z.GetValue() == 1.0f ? EVisibility::Hidden : EVisibility::Visible;
 	}
 
+	/** Cache a single unit to display all location comonents in */
+	void CacheCommonLocationUnits();
+
 private:
 	/** A vector where it may optionally be unset */
 	struct FOptionalVector
@@ -325,4 +330,6 @@ private:
 	bool bEditingRotationInUI;
 	/** Bitmask to indicate which fields should be hidden (if any) */
 	uint8 HiddenFieldMask;
+	/** Enum specifying the most appropriate units for the whole vector to be displayed in */
+	EUnit VectorUnits;
 };
