Added color property for Vector / Vector Random parameter types in particle instance parameters.

diff --git a/Engine/Source/Editor/DetailCustomizations/Private/ParticleSysParamStructCustomization.cpp b/Engine/Source/Editor/DetailCustomizations/Private/ParticleSysParamStructCustomization.cpp
--- a/Engine/Source/Editor/DetailCustomizations/Private/ParticleSysParamStructCustomization.cpp
+++ b/Engine/Source/Editor/DetailCustomizations/Private/ParticleSysParamStructCustomization.cpp
@@ -136,7 +136,7 @@ EVisibility FParticleSysParamStructCustomization::GetVectorLowVisibility() const
 
 EVisibility FParticleSysParamStructCustomization::GetColorVisibility() const
 {
-	return (ParameterType == PSPT_Color) ? EVisibility::Visible : EVisibility::Collapsed;
+	return (ParameterType == PSPT_Vector || ParameterType == PSPT_VectorRand || ParameterType == PSPT_Color) ? EVisibility::Visible : EVisibility::Collapsed;
 }
 
 EVisibility FParticleSysParamStructCustomization::GetActorVisibility() const
