Fix particle system emitter action buttons stretching out the entire with of the details panel

diff --git a/Engine/Source/Editor/DetailCustomizations/Private/ParticleSystemComponentDetails.cpp b/Engine/Source/Editor/DetailCustomizations/Private/ParticleSystemComponentDetails.cpp
--- a/Engine/Source/Editor/DetailCustomizations/Private/ParticleSystemComponentDetails.cpp
+++ b/Engine/Source/Editor/DetailCustomizations/Private/ParticleSystemComponentDetails.cpp
@@ -21,21 +21,26 @@ void FParticleSystemComponentDetails::CustomizeDetails(IDetailLayoutBuilder& InD
 	IDetailCategoryBuilder& CustomCategory = InDetailLayout.EditCategory("EmitterActions", NSLOCTEXT("ParticleSystemComponentDetails", "EmitterActionCategoryName", "Emitter Actions"), ECategoryPriority::Important);
 
 	CustomCategory.AddCustomRow(FText::GetEmpty())
+		.WholeRowContent()
+		.HAlign(HAlign_Left)
 		[
-			SNew(SUniformGridPanel)
-			.SlotPadding(2.0f)
-			+ SUniformGridPanel::Slot(0, 0)
+			SNew( SBox )
+			.MaxDesiredWidth(300.f)
 			[
-				SNew(SButton)
-				.OnClicked(this, &FParticleSystemComponentDetails::OnAutoPopulateClicked)
-				.ToolTipText(NSLOCTEXT("ParticleSystemComponentDetails", "AutoPopulateButtonTooltip", "Copies properties from the source particle system into the instanced parameters of this system"))
-				.HAlign(HAlign_Center)
+				SNew(SUniformGridPanel)
+				.SlotPadding(2.0f)
+				+ SUniformGridPanel::Slot(0, 0)
 				[
-					SNew(STextBlock)
-					.Text(NSLOCTEXT("ParticleSystemComponentDetails", "AutoPopulateButton", "Expose Parameter"))
+					SNew(SButton)
+					.OnClicked(this, &FParticleSystemComponentDetails::OnAutoPopulateClicked)
+					.ToolTipText(NSLOCTEXT("ParticleSystemComponentDetails", "AutoPopulateButtonTooltip", "Copies properties from the source particle system into the instanced parameters of this system"))
+					.HAlign(HAlign_Center)
+					[
+						SNew(STextBlock)
+						.Text(NSLOCTEXT("ParticleSystemComponentDetails", "AutoPopulateButton", "Expose Parameter"))
+					]
 				]
-			]
-			+ SUniformGridPanel::Slot(1, 0)
+				+ SUniformGridPanel::Slot(1, 0)
 				[
 					SNew(SButton)
 					.OnClicked(this, &FParticleSystemComponentDetails::OnResetEmitter)
@@ -46,6 +51,7 @@ void FParticleSystemComponentDetails::CustomizeDetails(IDetailLayoutBuilder& InD
 						.Text(NSLOCTEXT("ParticleSystemComponentDetails", "ResetEmitterButton", "Reset Emitter"))
 					]
 				]
+			]
 		];
 }
 
