Editor: Fix crash when a non-uniform scale is used (failure to find FVector type inside UObject as its outer has changed)

diff --git a/Engine/Source/Editor/DetailCustomizations/Private/ComponentTransformDetails.cpp b/Engine/Source/Editor/DetailCustomizations/Private/ComponentTransformDetails.cpp
--- a/Engine/Source/Editor/DetailCustomizations/Private/ComponentTransformDetails.cpp
+++ b/Engine/Source/Editor/DetailCustomizations/Private/ComponentTransformDetails.cpp
@@ -1415,7 +1415,7 @@ void FComponentTransformDetails::ScaleObject( float NewValue, int32 Axis, bool b
 
 					if (!bPreserveScaleRatio)
 					{
-						UStruct* VectorStruct = FindObjectChecked<UStruct>(UObject::StaticClass(), TEXT("Vector"), false);
+						UStruct* VectorStruct = FindObjectChecked<UScriptStruct>(UObject::StaticClass()->GetOutermost(), TEXT("Vector"), false);
 
 						UProperty* VectorValueProperty = NULL;
 						switch( Axis )
