[INTEGRATE] Change 2414015 by [REDACTED_PERSON][REDACTED_HANDLE] on 2015/01/21 12:54:11

	Resetting an Actor instance to its Blueprint defaults will remove Instance Components

diff --git a/Engine/Source/Editor/DetailCustomizations/Private/ActorDetails.cpp b/Engine/Source/Editor/DetailCustomizations/Private/ActorDetails.cpp
--- a/Engine/Source/Editor/DetailCustomizations/Private/ActorDetails.cpp
+++ b/Engine/Source/Editor/DetailCustomizations/Private/ActorDetails.cpp
@@ -1359,11 +1359,18 @@ bool FActorDetails::ResetToBlueprintDefaults_IsEnabled( TWeakObjectPtr<UBlueprin
 			&& Blueprint != NULL
 			&& Actor->GetClass()->ClassGeneratedBy == Blueprint)
 		{
-			AActor* BlueprintCDO = Actor->GetClass()->GetDefaultObject<AActor>();
-			if(BlueprintCDO != NULL)
+			if (Actor->InstanceComponents.Num() > 0)
 			{
-				const auto CopyOptions = (EditorUtilities::ECopyOptions::Type)(EditorUtilities::ECopyOptions::PreviewOnly|EditorUtilities::ECopyOptions::OnlyCopyEditOrInterpProperties);
-				bIsEnabled = EditorUtilities::CopyActorProperties(BlueprintCDO, Actor, CopyOptions) > 0;
+				bIsEnabled = true;
+			}
+			else
+			{
+				AActor* BlueprintCDO = Actor->GetClass()->GetDefaultObject<AActor>();
+				if(BlueprintCDO != NULL)
+				{
+					const auto CopyOptions = (EditorUtilities::ECopyOptions::Type)(EditorUtilities::ECopyOptions::PreviewOnly|EditorUtilities::ECopyOptions::OnlyCopyEditOrInterpProperties);
+					bIsEnabled = EditorUtilities::CopyActorProperties(BlueprintCDO, Actor, CopyOptions) > 0;
+				}
 			}
 		}
 	}
@@ -1389,6 +1396,7 @@ FText FActorDetails::ResetToBlueprintDefaults_ToolTipText( TWeakObjectPtr<UBluep
 				const auto CopyOptions = (EditorUtilities::ECopyOptions::Type)(EditorUtilities::ECopyOptions::PreviewOnly|EditorUtilities::ECopyOptions::OnlyCopyEditOrInterpProperties);
 				NumChangedProperties += EditorUtilities::CopyActorProperties(BlueprintCDO, Actor, CopyOptions);
 			}
+			NumChangedProperties += Actor->InstanceComponents.Num();
 		}
 	}
 
@@ -1427,6 +1435,15 @@ FReply FActorDetails::ResetToBlueprintDefaults_OnClicked( TWeakObjectPtr<UBluepr
 				const auto CopyOptions = (EditorUtilities::ECopyOptions::Type)(EditorUtilities::ECopyOptions::OnlyCopyEditOrInterpProperties|EditorUtilities::ECopyOptions::CallPostEditChangeProperty);
 				NumChangedProperties = EditorUtilities::CopyActorProperties(BlueprintCDO, Actor, CopyOptions);
 			}
+			if (Actor->InstanceComponents.Num() > 0)
+			{
+				TArray<UActorComponent*> InstanceComponents = Actor->InstanceComponents;
+				NumChangedProperties += InstanceComponents.Num();
+				for (UActorComponent* ActorComponent : InstanceComponents)
+				{
+					ActorComponent->DestroyComponent();
+				}
+			}
 		}
 
 		// Set up a notification record to indicate success/failure
