[INTEGRATE] Change 2415910 by [REDACTED_PERSON][REDACTED_HANDLE] on 2015/01/22 17:17:50

	Add bInstanceComponent to UActorComponent and encapsulate AActor::InstanceComponents

diff --git a/Engine/Source/Editor/DetailCustomizations/Private/ActorDetails.cpp b/Engine/Source/Editor/DetailCustomizations/Private/ActorDetails.cpp
--- a/Engine/Source/Editor/DetailCustomizations/Private/ActorDetails.cpp
+++ b/Engine/Source/Editor/DetailCustomizations/Private/ActorDetails.cpp
@@ -1220,7 +1220,7 @@ bool FActorDetails::PushToBlueprintDefaults_IsEnabled( TWeakObjectPtr<UBlueprint
 			&& Blueprint != NULL
 			&& Actor->GetClass()->ClassGeneratedBy == Blueprint)
 		{
-			if (Actor->InstanceComponents.Num() > 0)
+			if (Actor->GetInstanceComponents().Num() > 0)
 			{
 				bIsEnabled = true;
 			}
@@ -1257,7 +1257,7 @@ FText FActorDetails::PushToBlueprintDefaults_ToolTipText( TWeakObjectPtr<UBluepr
 				const auto CopyOptions = (EditorUtilities::ECopyOptions::Type)(EditorUtilities::ECopyOptions::PreviewOnly|EditorUtilities::ECopyOptions::OnlyCopyEditOrInterpProperties);
 				NumChangedProperties += EditorUtilities::CopyActorProperties(Actor, BlueprintCDO, CopyOptions);
 			}
-			NumChangedProperties += Actor->InstanceComponents.Num();
+			NumChangedProperties += Actor->GetInstanceComponents().Num();
 		}
 	}
 
@@ -1295,11 +1295,11 @@ FReply FActorDetails::PushToBlueprintDefaults_OnClicked( TWeakObjectPtr<UBluepri
 			{
 				const auto CopyOptions = (EditorUtilities::ECopyOptions::Type)(EditorUtilities::ECopyOptions::OnlyCopyEditOrInterpProperties|EditorUtilities::ECopyOptions::PropagateChangesToArcheypeInstances);
 				NumChangedProperties = EditorUtilities::CopyActorProperties(Actor, BlueprintCDO, CopyOptions);
-				if (Actor->InstanceComponents.Num() > 0)
+				if (Actor->GetInstanceComponents().Num() > 0)
 				{
-					FKismetEditorUtilities::AddComponentsToBlueprint(Blueprint, Actor->InstanceComponents);
-					NumChangedProperties += Actor->InstanceComponents.Num();
-					Actor->InstanceComponents.Empty();
+					FKismetEditorUtilities::AddComponentsToBlueprint(Blueprint, Actor->GetInstanceComponents());
+					NumChangedProperties += Actor->GetInstanceComponents().Num();
+					Actor->ClearInstanceComponents();
 				}
 				if(NumChangedProperties > 0)
 				{
@@ -1359,7 +1359,7 @@ bool FActorDetails::ResetToBlueprintDefaults_IsEnabled( TWeakObjectPtr<UBlueprin
 			&& Blueprint != NULL
 			&& Actor->GetClass()->ClassGeneratedBy == Blueprint)
 		{
-			if (Actor->InstanceComponents.Num() > 0)
+			if (Actor->GetInstanceComponents().Num() > 0)
 			{
 				bIsEnabled = true;
 			}
@@ -1396,7 +1396,7 @@ FText FActorDetails::ResetToBlueprintDefaults_ToolTipText( TWeakObjectPtr<UBluep
 				const auto CopyOptions = (EditorUtilities::ECopyOptions::Type)(EditorUtilities::ECopyOptions::PreviewOnly|EditorUtilities::ECopyOptions::OnlyCopyEditOrInterpProperties);
 				NumChangedProperties += EditorUtilities::CopyActorProperties(BlueprintCDO, Actor, CopyOptions);
 			}
-			NumChangedProperties += Actor->InstanceComponents.Num();
+			NumChangedProperties += Actor->GetInstanceComponents().Num();
 		}
 	}
 
@@ -1435,9 +1435,9 @@ FReply FActorDetails::ResetToBlueprintDefaults_OnClicked( TWeakObjectPtr<UBluepr
 				const auto CopyOptions = (EditorUtilities::ECopyOptions::Type)(EditorUtilities::ECopyOptions::OnlyCopyEditOrInterpProperties|EditorUtilities::ECopyOptions::CallPostEditChangeProperty);
 				NumChangedProperties = EditorUtilities::CopyActorProperties(BlueprintCDO, Actor, CopyOptions);
 			}
-			if (Actor->InstanceComponents.Num() > 0)
+			if (Actor->GetInstanceComponents().Num() > 0)
 			{
-				TArray<UActorComponent*> InstanceComponents = Actor->InstanceComponents;
+				TArray<UActorComponent*> InstanceComponents = Actor->GetInstanceComponents();
 				NumChangedProperties += InstanceComponents.Num();
 				for (UActorComponent* ActorComponent : InstanceComponents)
 				{
diff --git a/Engine/Source/Editor/Kismet/Private/SSCSEditor.cpp b/Engine/Source/Editor/Kismet/Private/SSCSEditor.cpp
--- a/Engine/Source/Editor/Kismet/Private/SSCSEditor.cpp
+++ b/Engine/Source/Editor/Kismet/Private/SSCSEditor.cpp
@@ -3618,7 +3618,7 @@ UActorComponent* SSCSEditor::AddNewComponent( UClass* NewComponentClass, UObject
 		}
 
 		// Add to SerializedComponents array so it gets saved
-		ActorInstance->InstanceComponents.Add(NewComponentInstance);
+		ActorInstance->AddInstanceComponent(NewComponentInstance);
 
 		return AddNewNode(NewComponentInstance, Asset, true);
 	}
@@ -3926,7 +3926,7 @@ void SSCSEditor::PasteNodes()
 			NewActorComponent->Rename(*NewComponentName, ActorInstance, REN_DontCreateRedirectors | REN_DoNotDirty);
 
 			// Add to SerializedComponents array so it gets saved
-			ActorInstance->InstanceComponents.Add(NewActorComponent);
+			ActorInstance->AddInstanceComponent(NewActorComponent);
 
 			// Create a new node to contain the new component instance and add it to the tree
 			NewActorComponent = AddNewNode(NewActorComponent, NULL, false);
diff --git a/Engine/Source/Editor/LevelEditor/Private/ComponentsTree.cpp b/Engine/Source/Editor/LevelEditor/Private/ComponentsTree.cpp
--- a/Engine/Source/Editor/LevelEditor/Private/ComponentsTree.cpp
+++ b/Engine/Source/Editor/LevelEditor/Private/ComponentsTree.cpp
@@ -253,7 +253,7 @@ void SComponentsTree::OnSelectedCompClass(TSubclassOf<UActorComponent> CompClass
 		check(NewComp);
 
 		// Add to InstanceComponents array so it gets saved
-		Actor->InstanceComponents.Add(NewComp);
+		Actor->AddInstanceComponent(NewComp);
 
 		// If a scene component, and actor has a root component, attach it to root
 		USceneComponent* NewSceneComp = Cast<USceneComponent>(NewComp);
diff --git a/Engine/Source/Editor/UnrealEd/Private/Factories/ActorFactory.cpp b/Engine/Source/Editor/UnrealEd/Private/Factories/ActorFactory.cpp
--- a/Engine/Source/Editor/UnrealEd/Private/Factories/ActorFactory.cpp
+++ b/Engine/Source/Editor/UnrealEd/Private/Factories/ActorFactory.cpp
@@ -1158,8 +1158,8 @@ AActor* UActorFactoryEmptyActor::SpawnActor( UObject* Asset, ULevel* InLevel, co
 
 		UBillboardComponent* BillboardComponent	= CreateEditorOnlyBillboardComponent(NewActor, RootComponent);
 
-	NewActor->InstanceComponents.Add(RootComponent);
-	NewActor->InstanceComponents.Add(BillboardComponent);
+	NewActor->AddInstanceComponent(RootComponent);
+	NewActor->AddInstanceComponent(BillboardComponent);
 
 	return NewActor;
 }
diff --git a/Engine/Source/Editor/UnrealEd/Private/Kismet2/Kismet2.cpp b/Engine/Source/Editor/UnrealEd/Private/Kismet2/Kismet2.cpp
--- a/Engine/Source/Editor/UnrealEd/Private/Kismet2/Kismet2.cpp
+++ b/Engine/Source/Editor/UnrealEd/Private/Kismet2/Kismet2.cpp
@@ -937,6 +937,9 @@ void FKismetEditorUtilities::AddComponentsToBlueprint(UBlueprint* Blueprint, con
 			USCS_Node* SCSNode = SCS->CreateNode(ActorComponent->GetClass());
 			UEditorEngine::CopyPropertiesForUnrelatedObjects(ActorComponent,SCSNode->ComponentTemplate);
 
+			// Clear the instance component flag
+			SCSNode->ComponentTemplate->bInstanceComponent = false;
+
 			USceneComponent* SceneComponent = Cast<USceneComponent>(ActorComponent);
 
 			// The easy part is non-scene component or the Root simply add it
@@ -1024,9 +1027,9 @@ UBlueprint* FKismetEditorUtilities::CreateBlueprintFromActor(const FName Bluepri
 			Outer->MarkPackageDirty();
 
 			// If the source Actor has Instance Components we need to translate these in to SCS Nodes
-			if (Actor->InstanceComponents.Num() > 0)
+			if (Actor->GetInstanceComponents().Num() > 0)
 			{
-				AddComponentsToBlueprint(NewBlueprint, Actor->InstanceComponents);
+				AddComponentsToBlueprint(NewBlueprint, Actor->GetInstanceComponents());
 			}
 
 			if(NewBlueprint->GeneratedClass)
@@ -1048,7 +1051,7 @@ UBlueprint* FKismetEditorUtilities::CreateBlueprintFromActor(const FName Bluepri
 					}
 
 					// Clear the instance properties array as we created SCS nodes for them
-					CDOAsActor->InstanceComponents.Empty();
+					CDOAsActor->ClearInstanceComponents();
 				}
 			}
 
diff --git a/Engine/Source/Runtime/Engine/Classes/Components/ActorComponent.h b/Engine/Source/Runtime/Engine/Classes/Components/ActorComponent.h
--- a/Engine/Source/Runtime/Engine/Classes/Components/ActorComponent.h
+++ b/Engine/Source/Runtime/Engine/Classes/Components/ActorComponent.h
@@ -90,6 +90,10 @@ public:
 	UPROPERTY()
 	uint32 bCreatedByConstructionScript:1;
 
+	/** True if this component was created as an instance component */
+	UPROPERTY()
+	uint32 bInstanceComponent:1;
+
 	/** Whether to the component is activated at creation or must be explicitly activated. */
 	UPROPERTY(EditAnywhere, BlueprintReadOnly, Category=Activation)
 	uint32 bAutoActivate:1;
diff --git a/Engine/Source/Runtime/Engine/Classes/GameFramework/Actor.h b/Engine/Source/Runtime/Engine/Classes/GameFramework/Actor.h
--- a/Engine/Source/Runtime/Engine/Classes/GameFramework/Actor.h
+++ b/Engine/Source/Runtime/Engine/Classes/GameFramework/Actor.h
@@ -2227,10 +2227,25 @@ public:
 	UPROPERTY(TextExportTransient, NonTransactional)
 	TArray<UActorComponent*> BlueprintCreatedComponents;
 
+private:
 	/** Array of ActorComponents that have been added by the user on a per-instance basis. */
 	UPROPERTY()
 	TArray<UActorComponent*> InstanceComponents;
 
+public:
+
+	/** Adds a component to the instance components array */
+	void AddInstanceComponent(UActorComponent* Component);
+
+	/** Removes a component from the instance components array */
+	void RemoveInstanceComponent(UActorComponent* Component);
+
+	/** Clears the instance components array */
+	void ClearInstanceComponents();
+
+	/** Returns the instance components array */
+	const TArray<UActorComponent*>& GetInstanceComponents() const;
+
 public:
 	//=============================================================================
 	// Navigation related functions
diff --git a/Engine/Source/Runtime/Engine/Private/Actor.cpp b/Engine/Source/Runtime/Engine/Private/Actor.cpp
--- a/Engine/Source/Runtime/Engine/Private/Actor.cpp
+++ b/Engine/Source/Runtime/Engine/Private/Actor.cpp
@@ -2044,6 +2044,27 @@ void AActor::UpdateAllReplicatedComponents()
 	}
 }
 
+const TArray<UActorComponent*>& AActor::GetInstanceComponents() const
+{
+	return InstanceComponents;
+}
+
+void AActor::AddInstanceComponent(UActorComponent* Component)
+{
+	Component->bInstanceComponent = true;
+	InstanceComponents.Add(Component);
+}
+
+void AActor::RemoveInstanceComponent(UActorComponent* Component)
+{
+	InstanceComponents.Remove(Component);
+}
+
+void AActor::ClearInstanceComponents()
+{
+	InstanceComponents.Empty();
+}
+
 UActorComponent* AActor::FindComponentByClass(const TSubclassOf<UActorComponent> ComponentClass) const
 {
 	UActorComponent* FoundComponent = NULL;
diff --git a/Engine/Source/Runtime/Engine/Private/ActorComponent.cpp b/Engine/Source/Runtime/Engine/Private/ActorComponent.cpp
--- a/Engine/Source/Runtime/Engine/Private/ActorComponent.cpp
+++ b/Engine/Source/Runtime/Engine/Private/ActorComponent.cpp
@@ -704,7 +704,7 @@ void UActorComponent::DestroyComponent()
 		}
 		else
 		{
-			Owner->InstanceComponents.Remove(this);
+			Owner->RemoveInstanceComponent(this);
 		}
 		Owner->RemoveOwnedComponent(this);
 		if (Owner->GetRootComponent() == this)
diff --git a/Engine/Source/Runtime/Engine/Private/AudioComponent.cpp b/Engine/Source/Runtime/Engine/Private/AudioComponent.cpp
--- a/Engine/Source/Runtime/Engine/Private/AudioComponent.cpp
+++ b/Engine/Source/Runtime/Engine/Private/AudioComponent.cpp
@@ -60,6 +60,7 @@ void UAudioComponent::OnRegister()
 		SpriteComponent->SpriteInfo.Category = TEXT("Misc");
 		SpriteComponent->SpriteInfo.DisplayName = NSLOCTEXT( "SpriteCategory", "Misc", "Misc" );
 		SpriteComponent->bCreatedByConstructionScript = bCreatedByConstructionScript;
+		SpriteComponent->bInstanceComponent = bInstanceComponent;
 		SpriteComponent->bIsScreenSizeScaled = true;
 		SpriteComponent->bUseInEditorScaling = true;
 
diff --git a/Engine/Source/Runtime/Engine/Private/Camera/CameraComponent.cpp b/Engine/Source/Runtime/Engine/Private/Camera/CameraComponent.cpp
--- a/Engine/Source/Runtime/Engine/Private/Camera/CameraComponent.cpp
+++ b/Engine/Source/Runtime/Engine/Private/Camera/CameraComponent.cpp
@@ -52,6 +52,7 @@ void UCameraComponent::OnRegister()
 		ProxyMeshComponent->CastShadow = false;
 		ProxyMeshComponent->PostPhysicsComponentTick.bCanEverTick = false;
 		ProxyMeshComponent->bCreatedByConstructionScript = bCreatedByConstructionScript;
+		ProxyMeshComponent->bInstanceComponent = bInstanceComponent;
 		ProxyMeshComponent->RegisterComponentWithWorld(GetWorld());
 	}
 
@@ -62,6 +63,7 @@ void UCameraComponent::OnRegister()
 		DrawFrustum->AlwaysLoadOnClient = false;
 		DrawFrustum->AlwaysLoadOnServer = false;
 		DrawFrustum->bCreatedByConstructionScript = bCreatedByConstructionScript;
+		DrawFrustum->bInstanceComponent = bInstanceComponent;
 		DrawFrustum->RegisterComponentWithWorld(GetWorld());
 	}
 
diff --git a/Engine/Source/Runtime/Engine/Private/LightComponent.cpp b/Engine/Source/Runtime/Engine/Private/LightComponent.cpp
--- a/Engine/Source/Runtime/Engine/Private/LightComponent.cpp
+++ b/Engine/Source/Runtime/Engine/Private/LightComponent.cpp
@@ -183,6 +183,7 @@ void ULightComponentBase::OnRegister()
 		SpriteComponent->SpriteInfo.Category = TEXT("Lighting");
 		SpriteComponent->SpriteInfo.DisplayName = NSLOCTEXT("SpriteCategory", "Lighting", "Lighting");
 		SpriteComponent->bCreatedByConstructionScript = bCreatedByConstructionScript;
+		SpriteComponent->bInstanceComponent = bInstanceComponent;
 		SpriteComponent->bIsScreenSizeScaled = true;
 
 		SpriteComponent->RegisterComponent();
diff --git a/Engine/Source/Runtime/Engine/Private/PhysicsEngine/PhysicsConstraintComponent.cpp b/Engine/Source/Runtime/Engine/Private/PhysicsEngine/PhysicsConstraintComponent.cpp
--- a/Engine/Source/Runtime/Engine/Private/PhysicsEngine/PhysicsConstraintComponent.cpp
+++ b/Engine/Source/Runtime/Engine/Private/PhysicsEngine/PhysicsConstraintComponent.cpp
@@ -246,6 +246,7 @@ void UPhysicsConstraintComponent::OnRegister()
 		SpriteComponent->SpriteInfo.Category = TEXT("Physics");
 		SpriteComponent->SpriteInfo.DisplayName = NSLOCTEXT( "SpriteCategory", "Physics", "Physics" );
 		SpriteComponent->bCreatedByConstructionScript = bCreatedByConstructionScript;
+		SpriteComponent->bInstanceComponent = bInstanceComponent;
 		SpriteComponent->bIsScreenSizeScaled = true;
 
 		SpriteComponent->RegisterComponent();
