[INTEGRATE] Change 2416397 by [REDACTED_PERSON][REDACTED_HANDLE] on 2015/01/22 22:45:52

	- Added Button to convert an actor to a blueprint from the components tree
	- Added Combo to take actions on an existing blueprint
	-- Promote changes, reset changes, create new subclass

	- Removed blueprint category from actor details

diff --git a/Engine/Source/Editor/DetailCustomizations/Private/ActorDetails.cpp b/Engine/Source/Editor/DetailCustomizations/Private/ActorDetails.cpp
--- a/Engine/Source/Editor/DetailCustomizations/Private/ActorDetails.cpp
+++ b/Engine/Source/Editor/DetailCustomizations/Private/ActorDetails.cpp
@@ -122,7 +122,6 @@ void FActorDetails::CustomizeDetails( IDetailLayoutBuilder& DetailLayout )
 		// Add Blueprint category, if not being hidden
 		if (!HideCategories.Contains(TEXT("Blueprint")))
 		{
-			AddBlueprintCategory(DetailLayout, UniqueBlueprints);
 			AddBlutilityCategory(DetailLayout, UniqueBlueprints);
 		}
 
@@ -506,6 +505,7 @@ void FActorDetails::AddActorCategory( IDetailLayoutBuilder& DetailBuilder, const
 	IDetailCategoryBuilder& ActorCategory = DetailBuilder.EditCategory("Actor", FText::GetEmpty(), ECategoryPriority::Uncommon );
 
 
+#if 0
 	// Create the info buttons per level
 	for ( auto LevelIt( ActorsPerLevelCount.CreateConstIterator() ); LevelIt; ++LevelIt)
 	{
@@ -545,6 +545,7 @@ void FActorDetails::AddActorCategory( IDetailLayoutBuilder& DetailBuilder, const
 		];
 
 	}
+#endif
 
 	// WorldSettings should never convert to another class type
 	if( SelectedActorInfo.SelectionClass != AWorldSettings::StaticClass() && SelectedActorInfo.HasConvertableAsset() )
@@ -564,36 +565,6 @@ void FActorDetails::AddActorCategory( IDetailLayoutBuilder& DetailBuilder, const
 	}
 }
 
-void FActorDetails::AddBlueprintCategory( IDetailLayoutBuilder& DetailBuilder, const TMap<UBlueprint*, UObject*>& UniqueBlueprints )
-{
-	// Create the overarching structure and expose the vertical box to add entries per-blueprint
-	TSharedPtr<SVerticalBox> Box;
-
-	IDetailCategoryBuilder& BlueprintCategory = DetailBuilder.EditCategory("Blueprint", NSLOCTEXT("BlueprintDetails", "BlueprintTitle", "Blueprint"), ECategoryPriority::Uncommon );
-
-	// Only show the blueprints section if a single actor is selected
-	if ( SelectedActors.Num() > 0 )
-	{
-		// We don't need to replace the blueprints with a combined one unless the number of selected actors have 
-		AddUtilityBlueprintRows( BlueprintCategory );
-
-		if ( UniqueBlueprints.Num() > 0 )
-		{
-			// Add an edit row for each blueprint
-			for ( auto BlueprintIt( UniqueBlueprints.CreateConstIterator() ); BlueprintIt; ++BlueprintIt )
-			{
-				UBlueprint* Blueprint = BlueprintIt.Key();
-				UObject* Actor = BlueprintIt.Value();
-
-				AddSingleBlueprintRow( BlueprintCategory, Blueprint, Actor );
-			}
-		}
-		else
-		{
-			AddSingleBlueprintRow( BlueprintCategory, NULL, SelectedActors[0].Get() );
-		}
-	}
-}
 
 void FActorDetails::AddBlutilityCategory( IDetailLayoutBuilder& DetailBuilder, const TMap<UBlueprint*, UObject*>& UniqueBlueprints )
 {
@@ -683,811 +654,6 @@ const TArray< TWeakObjectPtr<AActor> >& FActorDetails::GetSelectedActors() const
 	return SelectedActors;
 }
 
-void FActorDetails::AddUtilityBlueprintRows( IDetailCategoryBuilder& BlueprintCategory )
-{
-	if ( SelectedActors.Num() == 1 )
-	{
-		// Find In Level Blueprint
-		TWeakObjectPtr<AActor> ActorPtr = SelectedActors[0];
-		FText ActorLabel = NSLOCTEXT( "UnrealEd", "None", "None" );
-
-		if ( ActorPtr.IsValid() )
-		{
-			ActorLabel = FText::FromString( ActorPtr.Get()->GetActorLabel() );
-		}
-
-		FFormatNamedArguments Args;
-		Args.Add( TEXT( "ActorLabel" ), ActorLabel );
-
-		const FText FindButtonLabel = LOCTEXT( "FindInLevelScript", "Find in Level Blueprint" );
-		const FText FindButtonToolTip = FText::Format( LOCTEXT( "FindInLevelScript_ToolTip", "Search for uses of {ActorLabel} in the Level Blueprint" ), Args );
-
-		BlueprintCategory.AddCustomRow( LOCTEXT( "FindInLevelScript_Filter", "Find in Level Script" ) )
-		.WholeRowContent()
-		[
-			SNew( SHorizontalBox )
-
-			+ SHorizontalBox::Slot()
-			.AutoWidth()
-			.Padding( 0, 0, 10, 0 )
-			[
-				SNew( SBox )
-				.WidthOverride( 200 )
-				[
-					SNew( SButton )
-					.ToolTipText( FindButtonToolTip )
-					.OnClicked( this, &FActorDetails::FindSelectedActorsInLevelScript )
-					.IsEnabled( this, &FActorDetails::AreAnySelectedActorsInLevelScript )
-					.HAlign( HAlign_Fill )
-					.ContentPadding( 2 )
-					[
-						SNew( STextBlock )
-						.Text( FindButtonLabel )
-						.Font( IDetailLayoutBuilder::GetDetailFont() )
-					]
-				]
-			]
-		];
-
-		// Add Level Events
-		const FText AddEventLabel = LOCTEXT( "ScriptingEvents", "Add Level Events" );
-		const FText AddEventToolTip = FText::Format( LOCTEXT( "ScripingEvents_Tooltip", "Adds or views events for {ActorLabel} in the Level Blueprint" ), Args );
-
-		BlueprintCategory.AddCustomRow( LOCTEXT( "ScriptingEvents_Filter", "Level Blueprint Events" ) )
-		.WholeRowContent()
-		[
-			SNew( SHorizontalBox )
-
-			+ SHorizontalBox::Slot()
-			.AutoWidth()
-			.Padding( 0, 0, 10, 0 )
-			[
-				SNew( SBox )
-				.WidthOverride( 200 )
-				[
-					SNew( SComboButton )
-					.IsEnabled( this, &FActorDetails::IsActorValidForLevelScript )
-					.OnGetMenuContent( this, &FActorDetails::MakeEventOptionsWidgetFromSelection )
-					.ContentPadding( 2 )
-					.HAlign( HAlign_Fill )
-					.ButtonContent()
-					[
-						SNew( STextBlock )
-						.Text( AddEventLabel )
-						.ToolTipText( AddEventToolTip )
-						.Font( IDetailLayoutBuilder::GetDetailFont() )
-					]
-				]
-			]
-		];
-	}
-
-	// Create the row for this blueprint
-	BlueprintCategory.AddCustomRow( NSLOCTEXT( "BlueprintDetails", "BlueprintFilter", "Blueprints" ) )
-	.WholeRowContent()
-	[
-		SNew( SHorizontalBox )
-
-		+ SHorizontalBox::Slot()
-		.AutoWidth()
-		.Padding( 0, 0, 10, 0 )
-		[
-			SNew( SBox )
-			.WidthOverride( 200 )
-			[
-				SNew( SButton )
-				.ToolTipText( LOCTEXT( "CreateHarvestBlueprint_ToolTip", "Copy the components from the selected actors and create a new Blueprint Class" ) )
-				.HAlign( HAlign_Fill )
-				//.ButtonColorAndOpacity( FLinearColor( 0.2f, 0.4f, 0.6f, 1.0f ) )
-				.OnClicked( FOnClicked::CreateRaw( this, &FActorDetails::OnPickBlueprintPathClicked, true ) )
-				.ContentPadding( 2 )
-				[
-					SNew( STextBlock )
-					.Text( LOCTEXT( "ReplaceWithAmalgamBlueprint", "Convert to Blueprint Class" ) )
-					.Font(IDetailLayoutBuilder::GetDetailFont())
-				]
-			]
-		]
-	];
-}
-
-// Creates a row for a single blueprint in the selection details view   
-BEGIN_SLATE_FUNCTION_BUILD_OPTIMIZATION
-void FActorDetails::AddSingleBlueprintRow( IDetailCategoryBuilder& BlueprintCategory, UBlueprint* InBlueprint, UObject* InExampleActor)
-{
-	TWeakObjectPtr<UBlueprint> WeakBP = InBlueprint;
-	TWeakObjectPtr<UObject> WeakFirstAsset = InExampleActor;
-
-	TAttribute<const FSlateBrush*> BlueprintStatusIcon;
-	BlueprintStatusIcon.Bind(FGetBrushImage::CreateSP(this, &FActorDetails::GetBlueprintStatusIcon, WeakBP));
-
-	TAttribute<FText> BlueprintStatusTooltip;
-	BlueprintStatusTooltip.Bind(FGetText::CreateSP(this, &FActorDetails::GetBlueprintStatusTooltip, WeakBP));
-
-	TAttribute<const FSlateBrush*> BreakpointStatusIcon;
-	BreakpointStatusIcon.Bind(FGetBrushImage::CreateSP(this, &FActorDetails::GetBlueprintBreakpointStatusIcon, WeakBP));
-
-	const FSlateFontInfo& InfoFont = FEditorStyle::GetFontStyle( TEXT("ExpandableArea.NormalFont") );
-
-	FText Name, Description;
-	if (InBlueprint)
-	{
-		Name = FText::FromString( InBlueprint->GetName() );
-		Description = FText::FromString( InBlueprint->GetDesc() );
-	}
-
-	// Set up PathPickerConfig.
-	PathForActorBlueprint = FString("/Game/");
-	FPathPickerConfig PathPickerConfig;
-	PathPickerConfig.DefaultPath = PathForActorBlueprint;
-	PathPickerConfig.OnPathSelected = FOnPathSelected::CreateRaw(this, &FActorDetails::OnSelectBlueprintPath);
-	FContentBrowserModule& ContentBrowserModule = FModuleManager::LoadModuleChecked<FContentBrowserModule>("ContentBrowser");
-
-	bool bCanCreateActorBlueprint = false;
-	FString SelectedActorClassName;
-	FText CreateBlueprintToolTip = LOCTEXT("CreateBlueprintSuccess_ToolTip", "Create a Blueprint based on this Actor");
-	if (SelectedActors.Num() == 1)
-	{
-		if (FKismetEditorUtilities::CanCreateBlueprintOfClass(SelectedActors[0]->GetClass()))
-		{
-			SelectedActorClassName = SelectedActors[0]->GetClass()->GetName();
-			bCanCreateActorBlueprint = true;
-		}
-		else
-		{
-			CreateBlueprintToolTip = LOCTEXT("CreateBlueprintFail_ToolTip", "Cannot create Blueprint using this Actor class");
-		}
-	}
-	else
-	{
-		CreateBlueprintToolTip = LOCTEXT("CreateBlueprintFailTooMany_ToolTip", "Cannot create Blueprint, please select only one Actor");
-	}
-	
-	FText CurrentBlueprintLong = FText::Format( FText::FromString( TEXT( "{0} {1}" ) ), Name, Description );
-
-	// Blueprint Details
-	BlueprintCategory.AddCustomRow( NSLOCTEXT("BlueprintDetails", "BlueprintFilter", "Blueprints") )
-	.Visibility( InBlueprint ? EVisibility::Visible : EVisibility::Collapsed )
-	.NameContent()
-	[
-		SNew( SHorizontalBox )
-
-		// Blueprint status icon
-		+ SHorizontalBox::Slot()
-		.AutoWidth()
-		.VAlign(VAlign_Center)
-		.Padding( 0.0f, 0.0f, 4.0f, 0.0f )
-		[
-			SNew(SButton)
-			.OnClicked( this, &FActorDetails::OnCompileBlueprintClicked, WeakBP )
-			.ToolTipText( BlueprintStatusTooltip )
-			.ButtonStyle( FEditorStyle::Get(), "NoBorder" )
-			.ContentPadding(0.0f)
-			[
-				SNew(SImage)
-				.Image( BlueprintStatusIcon )
-			]
-		]
-
-		// Blueprint name
-		+ SHorizontalBox::Slot()
-		.AutoWidth()
-		.VAlign( VAlign_Center )
-		.HAlign( HAlign_Fill )
-		.Padding( 0.0f, 0.0f, 4.0f, 0.0f )
-		[
-			SNew( STextBlock )
-			.Font( InfoFont )
-			.Text( Name )
-			.ToolTipText( CurrentBlueprintLong )
-		]
-
-		// Blueprint Description
-		+ SHorizontalBox::Slot()
-		.AutoWidth()
-		.VAlign( VAlign_Center )
-		.HAlign( HAlign_Fill )
-		.Padding( 0.0f, 0.0f, 4.0f, 0.0f )
-		[
-			SNew( STextBlock )
-			.Font( InfoFont )
-			.ColorAndOpacity(FSlateColor::UseSubduedForeground() )
-			.Text( Description )
-			.ToolTipText( CurrentBlueprintLong )
-		]
-	]
-	.ValueContent()
-	.MinDesiredWidth(200.0f)
-	.MaxDesiredWidth(200.0f)
-	[
-		SNew( SHorizontalBox )
-
-		// Breakpoint enable/disable toggle and status button
-		+ SHorizontalBox::Slot()
-		.AutoWidth()
-		.HAlign(HAlign_Left)
-		.VAlign(VAlign_Center)
-		.Padding( 0.0f, 0.0f, 4.0f, 0.0f )
-		[
-			SNew(SButton)
-			.OnClicked( this, &FActorDetails::OnToggleBreakpointsClicked, WeakBP )
-			.ToolTipText( LOCTEXT("BlueprintBreakpointToggleAll_Tooltip", "Toggle the enabled state of all breakpoints in this blueprint") )
-			.ButtonStyle( FEditorStyle::Get(), "NoBorder" )
-			.ContentPadding(0.0f)
-			[
-				SNew(SImage)
-				.Image( BreakpointStatusIcon )
-			]
-		]
-
-		// Button for editing the blueprint
-		+ SHorizontalBox::Slot()
-		.AutoWidth()
-		[
-			SNew( SButton )
-			.VAlign( VAlign_Center )
-			.HAlign( HAlign_Center )
-			.IsEnabled(InBlueprint != NULL)
-			.OnClicked( this, &FActorDetails::OnEditBlueprintClicked, WeakBP, WeakFirstAsset )
-			[
-				SNew( STextBlock )
-				.Font( InfoFont )
-				.Text( LOCTEXT("EditAsset", "Edit") )
-			]
-		]
-		
-		// Button for pushing to Blueprint defaults
-		+ SHorizontalBox::Slot()
-		.AutoWidth()
-		[
-			SNew( SButton )
-			.VAlign( VAlign_Center )
-			.HAlign( HAlign_Center )
-			.IsEnabled( this, &FActorDetails::PushToBlueprintDefaults_IsEnabled, WeakBP )
-			.OnClicked( this, &FActorDetails::PushToBlueprintDefaults_OnClicked, WeakBP )
-			.ToolTip(IDocumentation::Get()->CreateToolTip(
-				TAttribute<FText>::Create( TAttribute<FText>::FGetter::CreateRaw(this, &FActorDetails::PushToBlueprintDefaults_ToolTipText, WeakBP) ),
-				NULL,
-				TEXT("Shared/LevelEditor"),
-				TEXT("PushToBlueprintDefaults")))
-			[
-				SNew( STextBlock )
-				.Font( InfoFont )
-				.Text( LOCTEXT("PushToBlueprintDefaultsButtonLabel", "Apply") )
-			]
-		]
-
-		// Button for resetting to Blueprint defaults
-		+ SHorizontalBox::Slot()
-		.AutoWidth()
-		[
-			SNew( SButton )
-			.VAlign( VAlign_Center )
-			.HAlign( HAlign_Center )
-			.IsEnabled( this, &FActorDetails::ResetToBlueprintDefaults_IsEnabled, WeakBP )
-			.OnClicked( this, &FActorDetails::ResetToBlueprintDefaults_OnClicked, WeakBP )
-			.ToolTip(IDocumentation::Get()->CreateToolTip(
-				TAttribute<FText>::Create( TAttribute<FText>::FGetter::CreateRaw(this, &FActorDetails::ResetToBlueprintDefaults_ToolTipText, WeakBP) ),
-				NULL,
-				TEXT("Shared/LevelEditor"),
-				TEXT("ResetToBlueprintDefaults")))
-			[
-				SNew( STextBlock )
-				.Font( InfoFont )
-				.Text( LOCTEXT("ResetToBlueprintDefaultsButtonLabel", "Reset") )
-			]
-		]
-	];
-
-		BlueprintCategory.AddCustomRow( LOCTEXT("CreateBlueprintFilterString", "Create Blueprint"), true )
-		.WholeRowContent()
-		.MinDesiredWidth(200)
-		.MaxDesiredWidth(200)
-		[
-			SNew( SHorizontalBox )
-			+SHorizontalBox::Slot()
-			.AutoWidth()
-			[
-				SNew(SBorder)
-				.BorderImage(FEditorStyle::GetBrush("NoBorder"))
-				.Padding(0)
-				.ToolTipText(CreateBlueprintToolTip)
-				[
-					SNew( SButton )
-					.IsEnabled(bCanCreateActorBlueprint)
-					.ButtonColorAndOpacity(FLinearColor(0.2f, 0.4f, 0.6f, 1.0f))
-					.OnClicked(FOnClicked::CreateRaw(this, 	&FActorDetails::OnPickBlueprintPathClicked, false ))
-					.ContentPadding(2)
-					[
-						SNew( SHorizontalBox )
-						+SHorizontalBox::Slot()
-						.VAlign(VAlign_Center)
-						.AutoWidth()
-						[
-							SNew(STextBlock)
-							.Text( FText::Format( LOCTEXT("CreateBlueprintFromActor", "Create {0} Blueprint Class..."), FText::FromString( SelectedActorClassName ) ) )
-							.Font(IDetailLayoutBuilder::GetDetailFont())
-						]
-
-						+SHorizontalBox::Slot()
-						.VAlign(VAlign_Center)
-						.AutoWidth()
-						.Padding( 4.0f, 0.0f, 0.0f, 0.0f )
-						[
-							SNew( SImage )
-							.Image(FEditorStyle::GetBrush(TEXT("Kismet.CreateBlueprint")))
-						]
-					]
-				]
-			]
-		];
-	}
-END_SLATE_FUNCTION_BUILD_OPTIMIZATION
-
-FReply FActorDetails::OnPickBlueprintPathClicked(bool bHarvest )
-{
-	FCreateBlueprintFromActorDialog::OpenDialog(bHarvest);
-
-	return FReply::Handled();
-}
-
-void FActorDetails::OnSelectBlueprintPath( const FString& Path )
-{
-	PathForActorBlueprint = Path;
-}
-
-FText FActorDetails::GetBlueprintStatusTooltip( TWeakObjectPtr<class UBlueprint> Asset ) const
-{
-	if (UBlueprint* Blueprint = Asset.Get())
-	{
-		switch (Blueprint->Status)
-		{
-		default:
-		case BS_Unknown:
-		case BS_Dirty:
-			return LOCTEXT("BlueprintStatus_Tooltip_Unknown", "Blueprint may have been modified. Click to recompile.");
-		case BS_Error:
-			return LOCTEXT("BlueprintStatus_Tooltip_Error", "Blueprint has one or more errors!");
-		case BS_UpToDate:
-			return LOCTEXT("BlueprintStatus_Tooltip_UpToDate", "Blueprint is up to date.");
-		case BS_UpToDateWithWarnings:
-			return LOCTEXT("BlueprintStatus_Tooltip_Warning", "Blueprint has one or more warnings!");
-		}
-	}
-
-	return FText::GetEmpty();
-}
-
-const FSlateBrush* FActorDetails::GetBlueprintStatusIcon(TWeakObjectPtr<class UBlueprint> Asset) const
-{
-	if (UBlueprint* Blueprint = Asset.Get())
-	{
-		switch (Blueprint->Status)
-		{
-		default:
-		case BS_Unknown:
-		case BS_Dirty:
-			return FEditorStyle::GetBrush("Kismet.Status.Unknown.Small");
-		case BS_Error:
-			return FEditorStyle::GetBrush("Kismet.Status.Error.Small");
-		case BS_UpToDate:
-			return FEditorStyle::GetBrush("Kismet.Status.Good.Small");
-		case BS_UpToDateWithWarnings:
-			return FEditorStyle::GetBrush("Kismet.Status.Warning.Small");
-		}
-	}
-
-	return NULL;
-}
-
-const FSlateBrush* FActorDetails::GetBlueprintBreakpointStatusIcon(TWeakObjectPtr<class UBlueprint> Asset) const
-{
-	if (UBlueprint* Blueprint = Asset.Get())
-	{
-		const int32 TotalBreakpointCount = Blueprint->Breakpoints.Num();
-
-		// Count the number of enabled breakpoints
-		int32 EnabledCount = 0;
-		for (int32 BreakpointIndex = 0; BreakpointIndex < TotalBreakpointCount; ++BreakpointIndex)
-		{
-			UBreakpoint*& Breakpoint = Blueprint->Breakpoints[BreakpointIndex];
-			if (Breakpoint->IsEnabledByUser())
-			{
-				++EnabledCount;
-			}
-		}
-
-		if (TotalBreakpointCount > 0)
-		{
-			if (EnabledCount == TotalBreakpointCount)
-			{
-				// all are enabled
-				return FEditorStyle::GetBrush("Kismet.Breakpoint.EnabledAndValid"); //@TODO: Show enabled and invalid if not compiled?
-			}
-			else if (EnabledCount == 0)
-			{
-				// none are enabled
-				return FEditorStyle::GetBrush("Kismet.Breakpoint.Disabled");
-			}
-			else
-			{
-				// some are enabled but not all
-				return FEditorStyle::GetBrush("Kismet.Breakpoint.MixedStatus"); //@TODO: Show enabled and invalid if not compiled?
-			}
-		}
-	}
-
-	// show nothing, no breakpoints to care about
-	return FEditorStyle::GetBrush("Kismet.Breakpoint.NoneSpacer");
-}
-
-FReply FActorDetails::OnCompileBlueprintClicked(TWeakObjectPtr<class UBlueprint> Asset)
-{
-	if (UBlueprint* Blueprint = Asset.Get())
-	{
-		if (!Blueprint->IsUpToDate())
-		{
-			FKismetEditorUtilities::CompileBlueprint(Blueprint);
-		}
-	}
-	return FReply::Handled();
-}
-
-FReply FActorDetails::OnToggleBreakpointsClicked(TWeakObjectPtr<class UBlueprint> Asset)
-{
-	if (UBlueprint* Blueprint = Asset.Get())
-	{
-		// Recompile the blueprint before trying to set a breakpoint if it's dirty or unknown and we're not in PIE/SIE
-		if (((Blueprint->Status == BS_Unknown) || (Blueprint->Status == BS_Dirty)) && (GEditor->PlayWorld == NULL))
-		{
-			FKismetEditorUtilities::CompileBlueprint(Blueprint);
-		}
-
-		// If the blueprint is still not up to date, don't modify the breakpoints
-		if (!Blueprint->IsUpToDate())
-		{
-			return FReply::Handled();
-		}
-
-		const int32 TotalBreakpointCount = Blueprint->Breakpoints.Num();
-
-		// Count the number of enabled breakpoints
-		int32 EnabledCount = 0;
-		for (int32 BreakpointIndex = 0; BreakpointIndex < TotalBreakpointCount; ++BreakpointIndex)
-		{
-			UBreakpoint*& Breakpoint = Blueprint->Breakpoints[BreakpointIndex];
-			if (Breakpoint->IsEnabledByUser())
-			{
-				++EnabledCount;
-			}
-		}
-
-		if (TotalBreakpointCount > 0)
-		{
-			// Determine what the new state should be
-			bool bNewEnabledState = false;
-			if (EnabledCount == TotalBreakpointCount)
-			{
-				// all are enabled; make them all disabled
-				bNewEnabledState = false;
-			}
-			else if (EnabledCount == 0)
-			{
-				// none are enabled; make them all enabled
-				bNewEnabledState = true;
-			}
-			else
-			{
-				// some are enabled but not all; make them all enabled
-				bNewEnabledState = true;
-			}
-
-			// Now make all the breakpoints enabled/disabled
-			for (int32 BreakpointIndex = 0; BreakpointIndex < TotalBreakpointCount; ++BreakpointIndex)
-			{
-				UBreakpoint*& Breakpoint = Blueprint->Breakpoints[BreakpointIndex];
-				FKismetDebugUtilities::SetBreakpointEnabled(Breakpoint, bNewEnabledState);
-			}
-		}
-	}
-	return FReply::Handled();
-}
-
-FReply FActorDetails::OnEditBlueprintClicked( TWeakObjectPtr<UBlueprint> InBlueprint, TWeakObjectPtr<UObject> InAsset )
-{
-	if (UBlueprint* Blueprint = InBlueprint.Get())
-	{
-		// Set the object being debugged if given an actor reference (if we don't do this before we edit the object the editor wont know we are debebugging something)
-		if (UObject* Asset = InAsset.Get())
-		{
-			check(Asset->GetClass()->ClassGeneratedBy == Blueprint);
-			Blueprint->SetObjectBeingDebugged(Asset);
-		}
-		// Open the blueprint
-		GEditor->EditObject( Blueprint );
-	}
-	return FReply::Handled();
-}
-
-bool FActorDetails::PushToBlueprintDefaults_IsEnabled( TWeakObjectPtr<UBlueprint> InBlueprint ) const
-{
-	bool bIsEnabled = false;
-	if(SelectedActors.Num() == 1)
-	{
-		AActor* Actor = SelectedActors[0].Get();
-		UBlueprint* Blueprint = InBlueprint.Get();
-		if(Actor != NULL
-			&& Blueprint != NULL
-			&& Actor->GetClass()->ClassGeneratedBy == Blueprint)
-		{
-			if (Actor->GetInstanceComponents().Num() > 0)
-			{
-				bIsEnabled = true;
-			}
-			else
-			{
-				AActor* BlueprintCDO = Actor->GetClass()->GetDefaultObject<AActor>();
-				if(BlueprintCDO != NULL)
-				{
-					const auto CopyOptions = (EditorUtilities::ECopyOptions::Type)(EditorUtilities::ECopyOptions::PreviewOnly|EditorUtilities::ECopyOptions::OnlyCopyEditOrInterpProperties);
-					bIsEnabled = EditorUtilities::CopyActorProperties(Actor, BlueprintCDO, CopyOptions) > 0;
-				}
-			}
-		}
-	}
-
-	return bIsEnabled;
-}
-
-FText FActorDetails::PushToBlueprintDefaults_ToolTipText( TWeakObjectPtr<UBlueprint> InBlueprint ) const
-{
-	int32 NumChangedProperties = 0;
-
-	if(SelectedActors.Num() == 1)
-	{
-		AActor* Actor = SelectedActors[0].Get();
-		UBlueprint* Blueprint = InBlueprint.Get();
-		if(Actor != NULL
-			&& Blueprint != NULL
-			&& Actor->GetClass()->ClassGeneratedBy == Blueprint)
-		{
-			AActor* BlueprintCDO = Actor->GetClass()->GetDefaultObject<AActor>();
-			if(BlueprintCDO != NULL)
-			{
-				const auto CopyOptions = (EditorUtilities::ECopyOptions::Type)(EditorUtilities::ECopyOptions::PreviewOnly|EditorUtilities::ECopyOptions::OnlyCopyEditOrInterpProperties);
-				NumChangedProperties += EditorUtilities::CopyActorProperties(Actor, BlueprintCDO, CopyOptions);
-			}
-			NumChangedProperties += Actor->GetInstanceComponents().Num();
-		}
-	}
-
-	if (NumChangedProperties == 0)
-	{
-		return LOCTEXT("DisabledPushToBlueprintDefaults_ToolTip", "Replaces the Blueprint's defaults with any altered property values.");
-	}
-	else if (NumChangedProperties > 1)
-	{
-		return FText::Format( LOCTEXT("PushToBlueprintDefaults_ToolTip", "Click to apply {0} changed properties to the Blueprint."), FText::AsNumber( NumChangedProperties ) );
-	}
-	else
-	{
-		return LOCTEXT("PushOneToBlueprintDefaults_ToolTip", "Click to apply 1 changed property to the Blueprint.");
-	}
-}
-
-FReply FActorDetails::PushToBlueprintDefaults_OnClicked( TWeakObjectPtr<UBlueprint> InBlueprint )
-{
-	int32 NumChangedProperties = 0;
-
-	if(SelectedActors.Num() == 1)
-	{
-		const FScopedTransaction Transaction(LOCTEXT("PushToBlueprintDefaults_Transaction", "Apply Changes to Blueprint"));
-
-		// Perform the actual copy
-		AActor* Actor = SelectedActors[0].Get();
-		UBlueprint* Blueprint = InBlueprint.Get();
-		if(Actor != NULL
-			&& Blueprint != NULL
-			&& Actor->GetClass()->ClassGeneratedBy == Blueprint)
-		{
-			AActor* BlueprintCDO = Actor->GetClass()->GetDefaultObject<AActor>();
-			if(BlueprintCDO != NULL)
-			{
-				const auto CopyOptions = (EditorUtilities::ECopyOptions::Type)(EditorUtilities::ECopyOptions::OnlyCopyEditOrInterpProperties|EditorUtilities::ECopyOptions::PropagateChangesToArcheypeInstances);
-				NumChangedProperties = EditorUtilities::CopyActorProperties(Actor, BlueprintCDO, CopyOptions);
-				if (Actor->GetInstanceComponents().Num() > 0)
-				{
-					FKismetEditorUtilities::AddComponentsToBlueprint(Blueprint, Actor->GetInstanceComponents());
-					NumChangedProperties += Actor->GetInstanceComponents().Num();
-					Actor->ClearInstanceComponents();
-				}
-				if(NumChangedProperties > 0)
-				{
-					FBlueprintEditorUtils::MarkBlueprintAsModified(Blueprint);
-					FKismetEditorUtilities::CompileBlueprint(Blueprint);
-				}
-			}
-		}
-
-		// Set up a notification record to indicate success/failure
-		FNotificationInfo NotificationInfo( FText::GetEmpty() );
-		NotificationInfo.FadeInDuration = 1.0f;
-		NotificationInfo.FadeOutDuration = 2.0f;
-		NotificationInfo.bUseLargeFont = false;
-		SNotificationItem::ECompletionState CompletionState;
-		if( NumChangedProperties > 0 )
-		{
-			if( NumChangedProperties > 1 )
-			{
-				FFormatNamedArguments Args;
-				Args.Add( TEXT("BlueprintName"), FText::FromName(Blueprint->GetFName() ) );
-				Args.Add( TEXT("NumChangedProperties"), NumChangedProperties );
-				Args.Add( TEXT("ActorName"), FText::FromString( Actor->GetActorLabel() ) );
-				NotificationInfo.Text = FText::Format( LOCTEXT("PushToBlueprintDefaults_ApplySuccess", "Updated Blueprint {BlueprintName} ({NumChangedProperties} property changes applied from actor {ActorName})."), Args );
-			}
-			else
-			{
-				FFormatNamedArguments Args;
-				Args.Add( TEXT("BlueprintName"), FText::FromName(Blueprint->GetFName() ) );
-				Args.Add( TEXT("ActorName"), FText::FromString( Actor->GetActorLabel() ) );
-				NotificationInfo.Text = FText::Format( LOCTEXT("PushOneToBlueprintDefaults_ApplySuccess", "Updated Blueprint {BlueprintName} (1 property change applied from actor {ActorName})."), Args );
-			}
-			CompletionState = SNotificationItem::CS_Success;
-		}
-		else
-		{
-			NotificationInfo.Text = LOCTEXT("PushToBlueprintDefaults_ApplyFailed", "No properties were copied");
-			CompletionState = SNotificationItem::CS_Fail;
-		}
-
-		// Add the notification to the queue
-		const auto Notification = FSlateNotificationManager::Get().AddNotification(NotificationInfo);
-		Notification->SetCompletionState(CompletionState);
-	}
-
-	return FReply::Handled();
-}
-
-bool FActorDetails::ResetToBlueprintDefaults_IsEnabled( TWeakObjectPtr<UBlueprint> InBlueprint ) const
-{
-	bool bIsEnabled = false;
-	if(SelectedActors.Num() == 1)
-	{
-		AActor* Actor = SelectedActors[0].Get();
-		UBlueprint* Blueprint = InBlueprint.Get();
-		if(Actor != NULL
-			&& Blueprint != NULL
-			&& Actor->GetClass()->ClassGeneratedBy == Blueprint)
-		{
-			if (Actor->GetInstanceComponents().Num() > 0)
-			{
-				bIsEnabled = true;
-			}
-			else
-			{
-				AActor* BlueprintCDO = Actor->GetClass()->GetDefaultObject<AActor>();
-				if(BlueprintCDO != NULL)
-				{
-					const auto CopyOptions = (EditorUtilities::ECopyOptions::Type)(EditorUtilities::ECopyOptions::PreviewOnly|EditorUtilities::ECopyOptions::OnlyCopyEditOrInterpProperties);
-					bIsEnabled = EditorUtilities::CopyActorProperties(BlueprintCDO, Actor, CopyOptions) > 0;
-				}
-			}
-		}
-	}
-
-	return bIsEnabled;
-}
-
-FText FActorDetails::ResetToBlueprintDefaults_ToolTipText( TWeakObjectPtr<UBlueprint> InBlueprint ) const
-{
-	int32 NumChangedProperties = 0;
-
-	if(SelectedActors.Num() == 1)
-	{
-		AActor* Actor = SelectedActors[0].Get();
-		UBlueprint* Blueprint = InBlueprint.Get();
-		if(Actor != NULL
-			&& Blueprint != NULL
-			&& Actor->GetClass()->ClassGeneratedBy == Blueprint)
-		{
-			AActor* BlueprintCDO = Actor->GetClass()->GetDefaultObject<AActor>();
-			if(BlueprintCDO != NULL)
-			{
-				const auto CopyOptions = (EditorUtilities::ECopyOptions::Type)(EditorUtilities::ECopyOptions::PreviewOnly|EditorUtilities::ECopyOptions::OnlyCopyEditOrInterpProperties);
-				NumChangedProperties += EditorUtilities::CopyActorProperties(BlueprintCDO, Actor, CopyOptions);
-			}
-			NumChangedProperties += Actor->GetInstanceComponents().Num();
-		}
-	}
-
-	if (NumChangedProperties == 0)
-	{
-		return LOCTEXT("DisabledResetBlueprintDefaults_ToolTip", "Resets altered properties back to their Blueprint default values.");
-	}
-	else if (NumChangedProperties > 1)
-	{
-		return FText::Format( LOCTEXT("ResetToBlueprintDefaults_ToolTip", "Click to reset {0} changed properties to their Blueprint default values."), FText::AsNumber( NumChangedProperties ) );
-	}
-	else
-	{
-		return LOCTEXT("ResetOneToBlueprintDefaults_ToolTip", "Click to reset 1 changed property to its Blueprint default value.");
-	}
-}
-
-FReply FActorDetails::ResetToBlueprintDefaults_OnClicked( TWeakObjectPtr<UBlueprint> InBlueprint )
-{
-	int32 NumChangedProperties = 0;
-
-	if(SelectedActors.Num() == 1)
-	{
-		const FScopedTransaction Transaction( LOCTEXT("ResetToBlueprintDefaults_Transaction", "Reset to Blueprint Defaults") );
-
-		// Perform the actual copy
-		AActor* Actor = SelectedActors[0].Get();
-		UBlueprint* Blueprint = InBlueprint.Get();
-		if(Actor != NULL
-			&& Blueprint != NULL
-			&& Actor->GetClass()->ClassGeneratedBy == Blueprint)
-		{
-			AActor* BlueprintCDO = Actor->GetClass()->GetDefaultObject<AActor>();
-			if(BlueprintCDO != NULL)
-			{
-				const auto CopyOptions = (EditorUtilities::ECopyOptions::Type)(EditorUtilities::ECopyOptions::OnlyCopyEditOrInterpProperties|EditorUtilities::ECopyOptions::CallPostEditChangeProperty);
-				NumChangedProperties = EditorUtilities::CopyActorProperties(BlueprintCDO, Actor, CopyOptions);
-			}
-			if (Actor->GetInstanceComponents().Num() > 0)
-			{
-				TArray<UActorComponent*> InstanceComponents = Actor->GetInstanceComponents();
-				NumChangedProperties += InstanceComponents.Num();
-				for (UActorComponent* ActorComponent : InstanceComponents)
-				{
-					if (ActorComponent)
-					{
-						ActorComponent->DestroyComponent();
-					}
-				}
-			}
-		}
-
-		// Set up a notification record to indicate success/failure
-		FNotificationInfo NotificationInfo( FText::GetEmpty() );
-		NotificationInfo.FadeInDuration = 1.0f;
-		NotificationInfo.FadeOutDuration = 2.0f;
-		NotificationInfo.bUseLargeFont = false;
-		SNotificationItem::ECompletionState CompletionState;
-		if( NumChangedProperties > 0 )
-		{
-			if( NumChangedProperties > 1 )
-			{
-				FFormatNamedArguments Args;
-				Args.Add( TEXT("BlueprintName"), FText::FromName(Blueprint->GetFName() ) );
-				Args.Add( TEXT("NumChangedProperties"), NumChangedProperties );
-				Args.Add( TEXT("ActorName"), FText::FromString( Actor->GetActorLabel() ) );
-				NotificationInfo.Text = FText::Format( LOCTEXT("ResetToBlueprintDefaults_ApplySuccess", "Reset {ActorName} ({NumChangedProperties} property changes applied from Blueprint {BlueprintName})."), Args );
-			}
-			else
-			{
-				FFormatNamedArguments Args;
-				Args.Add( TEXT("BlueprintName"), FText::FromName(Blueprint->GetFName() ) );
-				Args.Add( TEXT("ActorName"), FText::FromString( Actor->GetActorLabel() ) );
-				NotificationInfo.Text = FText::Format( LOCTEXT("ResetOneToBlueprintDefaults_ApplySuccess", "Reset {ActorName} (1 property change applied from Blueprint {BlueprintName})."), Args );
-			}
-			CompletionState = SNotificationItem::CS_Success;
-		}
-		else
-		{
-			NotificationInfo.Text = LOCTEXT("ResetToBlueprintDefaults_Failed", "No properties were reset");
-			CompletionState = SNotificationItem::CS_Fail;
-		}
-
-		// Add the notification to the queue
-		const auto Notification = FSlateNotificationManager::Get().AddNotification(NotificationInfo);
-		Notification->SetCompletionState(CompletionState);
-	}
-
-	return FReply::Handled();
-}
-
 bool FActorDetails::DoesActorHaveBlutiltyFunctions() const
 {
 	bool bFunctionsFound = false;
diff --git a/Engine/Source/Editor/DetailCustomizations/Private/ActorDetails.h b/Engine/Source/Editor/DetailCustomizations/Private/ActorDetails.h
--- a/Engine/Source/Editor/DetailCustomizations/Private/ActorDetails.h
+++ b/Engine/Source/Editor/DetailCustomizations/Private/ActorDetails.h
@@ -80,35 +80,9 @@ private:
 
 	const TArray< TWeakObjectPtr<AActor> >& GetSelectedActors() const;
 
-	/** Creates an area meant to be displayed once for any number of selected blueprints that displays helper buttons. */
-	void AddUtilityBlueprintRows( IDetailCategoryBuilder& BlueprintCategory );
-
 	/** Creates a category to display blutility functionality in a separate space from the standard blueprint details */
 	void AddBlutilityCategory( IDetailLayoutBuilder& DetailBuilder, const TMap<UBlueprint*, UObject*>& UniqueBlueprints );
 
-	// Creates a row for a single blueprint in the selection details view   
-	void AddSingleBlueprintRow( IDetailCategoryBuilder& BlueprintCategory, UBlueprint* InBlueprint, UObject* InExampleActor);
-
-	FText GetBlueprintStatusTooltip( TWeakObjectPtr<class UBlueprint> Asset ) const;
-
-	const FSlateBrush* GetBlueprintStatusIcon(TWeakObjectPtr<class UBlueprint> Asset) const;
-
-	const FSlateBrush* GetBlueprintBreakpointStatusIcon(TWeakObjectPtr<class UBlueprint> Asset) const;
-
-	FReply OnCompileBlueprintClicked(TWeakObjectPtr<class UBlueprint> Asset);
-
-	FReply OnToggleBreakpointsClicked(TWeakObjectPtr<class UBlueprint> Asset);
-
-	FReply OnEditBlueprintClicked( TWeakObjectPtr<UBlueprint> InBlueprint, TWeakObjectPtr<UObject> InAsset );
-
-	bool PushToBlueprintDefaults_IsEnabled( TWeakObjectPtr<UBlueprint> InBlueprint ) const;
-	FReply PushToBlueprintDefaults_OnClicked( TWeakObjectPtr<UBlueprint> InBlueprint );
-	FText PushToBlueprintDefaults_ToolTipText( TWeakObjectPtr<UBlueprint> InBlueprint ) const;
-
-	bool ResetToBlueprintDefaults_IsEnabled( TWeakObjectPtr<UBlueprint> InBlueprint ) const;
-	FReply ResetToBlueprintDefaults_OnClicked( TWeakObjectPtr<UBlueprint> InBlueprint );
-	FText ResetToBlueprintDefaults_ToolTipText( TWeakObjectPtr<UBlueprint> InBlueprint ) const;
-
 	/** Called to determine if an actor has any valid call-in-editor utility functions */
 	bool DoesActorHaveBlutiltyFunctions() const;
 
diff --git a/Engine/Source/Editor/EditorStyle/Private/SlateEditorStyle.cpp b/Engine/Source/Editor/EditorStyle/Private/SlateEditorStyle.cpp
--- a/Engine/Source/Editor/EditorStyle/Private/SlateEditorStyle.cpp
+++ b/Engine/Source/Editor/EditorStyle/Private/SlateEditorStyle.cpp
@@ -3911,12 +3911,13 @@ void FSlateEditorStyle::FStyle::SetupGraphEditorStyles()
 		}
 
 		// SCSEditor
-		Set("SCSEditor.NewComponentIcon", new IMAGE_BRUSH("Icons/icon_Blueprint_SpawnActor_16x", Icon16x16));
 
 		Set("SCSEditor.ToggleComponentEditing" , new IMAGE_BRUSH( "Icons/icon_translate_40x", Icon40x40) );
 		Set("SCSEditor.ToggleComponentEditing.Small" , new IMAGE_BRUSH( "Icons/icon_translate_40x", Icon20x20) );
 		Set("SCSEditor.TileViewTooltip.NonContentBorder", new BOX_BRUSH("/Docking/TabContentArea", FMargin(4 / 16.0f)));
 
+		Set("SCSEditor.PromoteToBlueprintIcon", new IMAGE_BRUSH("Icons/AssetIcons/Blueprint_16x", Icon16x16) );
+
 		// Timeline Editor
 		{
 			Set( "TimelineEditor.AddFloatTrack", new IMAGE_BRUSH( "Icons/icon_TrackAddFloat_36x24px", Icon36x24, FLinearColor::Black ) );
@@ -5163,7 +5164,7 @@ void FSlateEditorStyle::FStyle::SetupContentBrowserStyle()
 		Set( "ContentBrowser.PrimitiveCylinder", new IMAGE_BRUSH( "ContentBrowser/ThumbnailCylinder", Icon32x32 ) );
 
 		Set( "ContentBrowser.TopBar.Font", FTextBlockStyle( NormalText )
-			.SetFont( TTF_CORE_FONT( "Fonts/Roboto-Bold", 12 ) )
+			.SetFont( TTF_CORE_FONT( "Fonts/Roboto-Bold", 11 ) )
 			.SetColorAndOpacity( FLinearColor( 1.0f, 1.0f, 1.0f ) )
 			.SetHighlightColor( FLinearColor( 1.0f, 1.0f, 1.0f ) )
 			.SetShadowOffset( FVector2D( 1, 1 ) )
@@ -5223,7 +5224,7 @@ void FSlateEditorStyle::FStyle::SetupContentBrowserStyle()
 		Set( "ContentBrowser.TileViewDeveloperFolderIcon.Base", new IMAGE_BRUSH( "Icons/Folders/FolderDev_Base_512x", FVector2D(512, 512) ) );
 		Set( "ContentBrowser.TileViewDeveloperFolderIcon.Mask", new IMAGE_BRUSH( "Icons/Folders/FolderDev_BaseHi_512x", FVector2D(512, 512) ) );
 		Set( "ContentBrowser.PathText", FTextBlockStyle(NormalText)
-			.SetFont( TTF_CORE_FONT( "Fonts/Roboto-Bold", 12 ) )
+			.SetFont( TTF_CORE_FONT( "Fonts/Roboto-Bold", 11 ) )
 			.SetColorAndOpacity( FLinearColor( 1.0f, 1.0f, 1.0f ) )
 			.SetHighlightColor( FLinearColor( 1.0f, 1.0f, 1.0f ) )
 			.SetShadowOffset( FVector2D( 1,1 ) )
diff --git a/Engine/Source/Editor/Kismet/Private/BlueprintEditorViewportContextMenuExtender.cpp b/Engine/Source/Editor/Kismet/Private/BlueprintEditorViewportContextMenuExtender.cpp
--- a/Engine/Source/Editor/Kismet/Private/BlueprintEditorViewportContextMenuExtender.cpp
+++ b/Engine/Source/Editor/Kismet/Private/BlueprintEditorViewportContextMenuExtender.cpp
@@ -188,7 +188,8 @@ void FillBlueprintOptions(FMenuBuilder& MenuBuilder, TArray<AActor*> SelectedAct
 
 		if(bCanHarvestComponentsForBlueprint)
 		{
-			FUIAction CreateBlueprintAction( FExecuteAction::CreateStatic( &FCreateBlueprintFromActorDialog::OpenDialog, true ) );
+			AActor* ActorOverride = nullptr;
+			FUIAction CreateBlueprintAction( FExecuteAction::CreateStatic( &FCreateBlueprintFromActorDialog::OpenDialog, true, ActorOverride ) );
 			MenuBuilder.AddMenuEntry(LOCTEXT("CreateBlueprint", "Create Blueprint..."), LOCTEXT("CreateBlueprint_Tooltip", "Harvest Components from Selected Actors and create Blueprint"), FSlateIcon(FEditorStyle::GetStyleSetName(), "Kismet.HarvestBlueprintFromActors"), CreateBlueprintAction);
 		}
 	}
diff --git a/Engine/Source/Editor/Kismet/Private/SSCSEditor.cpp b/Engine/Source/Editor/Kismet/Private/SSCSEditor.cpp
--- a/Engine/Source/Editor/Kismet/Private/SSCSEditor.cpp
+++ b/Engine/Source/Editor/Kismet/Private/SSCSEditor.cpp
@@ -36,8 +36,13 @@
 
 #include "Engine/InheritableComponentHandler.h"
 
+#include "CreateBlueprintFromActorDialog.h"
+
 #include "BPVariableDragDropAction.h"
 
+#include "SNotificationList.h"
+#include "NotificationManager.h"
+
 #define LOCTEXT_NAMESPACE "SSCSEditor"
 
 DEFINE_LOG_CATEGORY_STATIC(LogSCSEditor, Log, All);
@@ -1359,7 +1364,7 @@ EVisibility SSCS_RowWidget::GetAssetVisibility() const
 
 FSlateColor SSCS_RowWidget::GetColorTint() const
 {
-	if(SCSEditor.Pin()->EditorMode.Get() == SSCSEditor::EEditorMode::BlueprintSCS)
+	if(SCSEditor.Pin()->GetEditorMode() == SSCSEditor::EEditorMode::BlueprintSCS)
 	{
 		if(NodePtr->IsNative())
 		{
@@ -1421,7 +1426,7 @@ TSharedPtr<SWidget> SSCS_RowWidget::BuildSceneRootDropActionMenu(FSCSEditorTreeN
 
 		check(NodePtr.IsValid());
 		bool bDroppedInSameBlueprint = true;
-		if(SCSEditor.Pin()->EditorMode.Get() == SSCSEditor::EEditorMode::BlueprintSCS)
+		if(SCSEditor.Pin()->GetEditorMode() == SSCSEditor::EEditorMode::BlueprintSCS)
 		{
 			bDroppedInSameBlueprint = DroppedNodePtr->GetBlueprint() == GetBlueprint();
 		}
@@ -1614,7 +1619,7 @@ void SSCS_RowWidget::OnDragEnter( const FGeometry& MyGeometry, const FDragDropEv
 					DragRowOp->CurrentHoverText = LOCTEXT("DropActionToolTip_AttachToOrMakeNewRoot", "Drop here to see available actions.");
 					DragRowOp->PendingDropAction = FSCSRowDragDropOp::DropAction_AttachToOrMakeNewRoot;
 				}
-				else if(SCSEditor.Pin()->EditorMode.Get() == SSCSEditor::EEditorMode::BlueprintSCS && DraggedNodePtr->GetBlueprint() != GetBlueprint())
+				else if(SCSEditor.Pin()->GetEditorMode() == SSCSEditor::EEditorMode::BlueprintSCS && DraggedNodePtr->GetBlueprint() != GetBlueprint())
 				{
 					if(bCanMakeNewRoot)
 					{
@@ -1820,7 +1825,7 @@ void SSCS_RowWidget::OnAttachToDropAction(const TArray<FSCSEditorTreeNodePtrType
 	bool bRegenerateTreeNodes = false;
 	const FScopedTransaction TransactionContext(DroppedNodePtrs.Num() > 1 ? LOCTEXT("AttachComponents", "Attach Components") : LOCTEXT("AttachComponent", "Attach Component"));
 
-	if(SCSEditorPtr->EditorMode.Get() == SSCSEditor::EEditorMode::BlueprintSCS)
+	if(SCSEditorPtr->GetEditorMode() == SSCSEditor::EEditorMode::BlueprintSCS)
 	{
 		// Get the current Blueprint context
 		UBlueprint* Blueprint = GetBlueprint();
@@ -1963,7 +1968,7 @@ void SSCS_RowWidget::OnDetachFromDropAction(const TArray<FSCSEditorTreeNodePtrTy
 
 	const FScopedTransaction TransactionContext(DroppedNodePtrs.Num() > 1 ? LOCTEXT("DetachComponents", "Detach Components") : LOCTEXT("DetachComponent", "Detach Component"));
 
-	if(SCSEditorPtr->EditorMode.Get() == SSCSEditor::EEditorMode::BlueprintSCS)
+	if(SCSEditorPtr->GetEditorMode() == SSCSEditor::EEditorMode::BlueprintSCS)
 	{
 		// Get the current "preview" Actor instance
 		AActor* PreviewActor = SCSEditorPtr->PreviewActor.Get();
@@ -2069,7 +2074,7 @@ void SSCS_RowWidget::OnMakeNewRootDropAction(FSCSEditorTreeNodePtrType DroppedNo
 	// Create a transaction record
 	const FScopedTransaction TransactionContext(LOCTEXT("MakeNewSceneRoot", "Make New Scene Root"));
 
-	if(SCSEditorPtr->EditorMode.Get() == SSCSEditor::EEditorMode::BlueprintSCS)
+	if(SCSEditorPtr->GetEditorMode() == SSCSEditor::EEditorMode::BlueprintSCS)
 	{
 		// Get the current Blueprint context
 		UBlueprint* Blueprint = GetBlueprint();
@@ -2170,7 +2175,7 @@ void SSCS_RowWidget::PostDragDropAction(bool bRegenerateTreeNodes)
 
 		PinnedEditor->RefreshSelectionDetails();
 
-		if(PinnedEditor->EditorMode.Get() == SSCSEditor::EEditorMode::BlueprintSCS)
+		if(PinnedEditor->GetEditorMode() == SSCSEditor::EEditorMode::BlueprintSCS)
 		{
 			if(NodePtr.IsValid())
 			{
@@ -2258,7 +2263,7 @@ FString SSCS_RowWidget::GetDocumentationLink() const
 {
 	check(SCSEditor.IsValid());
 
-	if(SCSEditor.Pin()->EditorMode.Get() == SSCSEditor::EEditorMode::BlueprintSCS)
+	if(SCSEditor.Pin()->GetEditorMode() == SSCSEditor::EEditorMode::BlueprintSCS)
 	{
 		if( NodePtr == SCSEditor.Pin()->SceneRootNodePtr
 			|| NodePtr->IsNative() || NodePtr->IsInherited())
@@ -2278,7 +2283,7 @@ FString SSCS_RowWidget::GetDocumentationExcerptName() const
 {
 	check(SCSEditor.IsValid());
 
-	if(SCSEditor.Pin()->EditorMode.Get() == SSCSEditor::EEditorMode::BlueprintSCS)
+	if(SCSEditor.Pin()->GetEditorMode() == SSCSEditor::EEditorMode::BlueprintSCS)
 	{
 		if( NodePtr == SCSEditor.Pin()->SceneRootNodePtr)
 		{
@@ -2363,7 +2368,7 @@ void SSCS_RowWidget::OnNameTextCommit(const FText& InNewName, ETextCommit::Type
 	// No need to call UpdateTree() in SCS editor mode; it will already be called by MBASM internally
 	check(SCSEditor.IsValid());
 	TSharedPtr<SSCSEditor> PinnedEditor = SCSEditor.Pin();
-	if(PinnedEditor.IsValid() && PinnedEditor->EditorMode.Get() == SSCSEditor::EEditorMode::ActorInstance)
+	if(PinnedEditor.IsValid() && PinnedEditor->GetEditorMode() == SSCSEditor::EEditorMode::ActorInstance)
 	{
 		PinnedEditor->UpdateTree();
 	}
@@ -2372,7 +2377,6 @@ void SSCS_RowWidget::OnNameTextCommit(const FText& InNewName, ETextCommit::Type
 //////////////////////////////////////////////////////////////////////////
 // SSCSEditor
 
-BEGIN_SLATE_FUNCTION_BUILD_OPTIMIZATION
 void SSCSEditor::Construct( const FArguments& InArgs )
 {
 	EditorMode = InArgs._EditorMode;
@@ -2484,13 +2488,51 @@ void SSCSEditor::Construct( const FArguments& InArgs )
 
 	if (bSingleLayoutBPEditor)
 	{
+		FMenuBuilder EditBlueprintMenuBuilder( true, NULL );
+
+		EditBlueprintMenuBuilder.BeginSection( NAME_None, LOCTEXT("EditBlueprintMenu_ExistingBlueprintHeader", "Existing Blueprint" ) );
+
+		EditBlueprintMenuBuilder.AddMenuEntry
+		(
+			LOCTEXT("OpenBlueprintEditor", "Open Blueprint Editor"),
+			LOCTEXT("OpenBlueprintEditor_ToolTip", "Opens the blueprint editor for this asset"),
+			FSlateIcon(),
+			FUIAction(FExecuteAction::CreateSP(this, &SSCSEditor::OnOpenBlueprintEditor))
+		);
+
+
+		EditBlueprintMenuBuilder.AddMenuEntry
+		(
+			LOCTEXT("PushChangesToBlueprint", "Apply Instance Changes to Blueprint"),
+			TAttribute<FText>(this, &SSCSEditor::OnGetApplyChangesToBlueprintTooltip),
+			FSlateIcon(),
+			FUIAction(FExecuteAction::CreateSP(this, &SSCSEditor::OnApplyChangesToBlueprint))
+		);
+
+		EditBlueprintMenuBuilder.AddMenuEntry
+		(
+			LOCTEXT("ResetToDefault", "Reset Instance Changes to Blueprint Default"),
+			TAttribute<FText>(this, &SSCSEditor::OnGetResetToBlueprintDefaultsTooltip),
+			FSlateIcon(),
+			FUIAction(FExecuteAction::CreateSP(this, &SSCSEditor::OnResetToBlueprintDefaults))
+		);
+
+		EditBlueprintMenuBuilder.BeginSection( NAME_None, LOCTEXT("EditBlueprintMenu_NewHeader", "Create New" ) );
+		//EditBlueprintMenuBuilder.AddMenuSeparator();
+
+		EditBlueprintMenuBuilder.AddMenuEntry
+		(
+			LOCTEXT("PromoteToBlueprint", "Convert to Class Blueprint"),
+			LOCTEXT("PromoteToBluerprintTooltip","Converts the existing Blueprint into a new SubClass Blueprint" ),
+			FSlateIcon(),
+			FUIAction(FExecuteAction::CreateSP(this, &SSCSEditor::PromoteToBlueprint))
+		);
+
 		Contents = SNew(SVerticalBox)
-		
 		+ SVerticalBox::Slot()
 		.Padding(0.0f)
 		[
 			SNew(SVerticalBox)
-
 			+ SVerticalBox::Slot()
 			.AutoHeight()
 			.VAlign(VAlign_Top)
@@ -2501,12 +2543,83 @@ void SSCSEditor::Construct( const FArguments& InArgs )
 				.BorderImage(FEditorStyle::GetBrush("ToolPanel.GroupBorder"))
 				.AddMetaData<FTagMetaData>(FTagMetaData(TEXT("ComponentsPanel")))
 				[
-					SNew(SBox)
+					SNew(SHorizontalBox)
+					+ SHorizontalBox::Slot()
 					.HAlign(HAlign_Left)
+					.AutoWidth()
 					[
 						SNew(SComponentClassCombo)
 						.OnComponentClassSelected(this, &SSCSEditor::PerformComboAddClass)
-						.ToolTipText(LOCTEXT("AddComponent_Tooltip", "Add a component."))
+						.ToolTipText(LOCTEXT("AddComponent_Tooltip", "Adds a new component to this actor"))
+					]
+					+ SHorizontalBox::Slot()
+					.Padding( 4.f, 0.0f )
+					.HAlign(HAlign_Right)
+					[
+						SNew( SButton )
+						.Visibility( this, &SSCSEditor::GetPromoteToBlueprintButtonVisibility )
+						.OnClicked( this, &SSCSEditor::OnPromoteToBlueprintClicked )
+						.ButtonStyle(FEditorStyle::Get(), "ToggleButton")
+						.ContentPadding(0)
+						.ToolTip(IDocumentation::Get()->CreateToolTip(
+							LOCTEXT("PromoteToBluerprintTooltip","Converts this actor into a reusable Class Blueprint that can have script behavior" ),
+							NULL,
+							TEXT("Shared/LevelEditor"),
+							TEXT("ConvertToBlueprint")))
+						[
+							SNew(SHorizontalBox)
+							+SHorizontalBox::Slot()
+							.VAlign(VAlign_Center)
+							.AutoWidth()
+							.Padding(2.f, 1.f)
+							[
+								SNew(SImage)
+								.Image( FEditorStyle::Get().GetBrush( "ClassIcon.BlueprintCore" ) )
+							]
+							+ SHorizontalBox::Slot()
+							.VAlign(VAlign_Center)
+							.Padding(1.f)
+							[
+								SNew(STextBlock)
+								.TextStyle(FEditorStyle::Get(), "ContentBrowser.TopBar.Font")
+								.Text( LOCTEXT("PromoteToBlueprint", "Convert to Class Blueprint") )
+							]
+						]
+					]
+					+ SHorizontalBox::Slot()
+					.Padding(4.f, 0.0f)
+					.HAlign(HAlign_Right)
+					[
+						SNew(SComboButton)
+						.Visibility(this, &SSCSEditor::GetEditBlueprintButtonVisibility)
+						.ContentPadding(FMargin(0))
+						.ComboButtonStyle(FEditorStyle::Get(), "ContentBrowser.NewAsset.Style")
+						.ForegroundColor(FLinearColor::White)
+						.ButtonContent()
+						[
+							SNew( SHorizontalBox )
+							+ SHorizontalBox::Slot()
+							.AutoWidth()
+							.HAlign(HAlign_Center)
+							.VAlign(VAlign_Center)
+							.Padding( 2.0f, 0.0f )
+							[
+								SNew( SImage )
+								.Image( FEditorStyle::Get().GetBrush( "ClassIcon.BlueprintCore" ) )
+							]
+							.Padding( 2.0f, 0.0f )
+							+ SHorizontalBox::Slot()
+							[
+								SNew(STextBlock)
+								.TextStyle(FEditorStyle::Get(), "ContentBrowser.TopBar.Font")
+								.Text(LOCTEXT("EditBlueprint", "Edit Blueprint"))
+							]
+
+						]
+						.MenuContent()
+						[
+							EditBlueprintMenuBuilder.MakeWidget()
+						]
 					]
 				]
 			]
@@ -2632,6 +2745,7 @@ void SSCSEditor::Construct( const FArguments& InArgs )
 		OnActorSelected(ECheckBoxState::Checked);
 	}
 }
+
 TSharedRef<SToolTip> SSCSEditor::CreateToolTipWidget() const
 {
 	// Create a box to hold every line of info in the body of the tooltip
@@ -2731,7 +2845,7 @@ void SSCSEditor::Tick( const FGeometry& AllottedGeometry, const double InCurrent
 {
 	SCompoundWidget::Tick(AllottedGeometry, InCurrentTime, InDeltaTime);
 
-	if(EditorMode.Get() == EEditorMode::ActorInstance)
+	if(EditorMode == EEditorMode::ActorInstance)
 	{
 		TFunction<bool (const TArray<FSCSEditorTreeNodePtrType>&, int32&)> AreAnyNodesInvalidLambda = [&](const TArray<FSCSEditorTreeNodePtrType>& InNodes, int32& OutNumValidNodes) -> bool
 		{
@@ -2840,7 +2954,7 @@ TSharedPtr< SWidget > SSCSEditor::CreateContextMenu()
 				MenuBuilder.AddMenuEntry( FGenericCommands::Get().Delete );
 				MenuBuilder.AddMenuEntry( FGenericCommands::Get().Rename );
 
-				if(EditorMode.Get() == EEditorMode::BlueprintSCS)
+				if(EditorMode == EEditorMode::BlueprintSCS)
 				{
 					// Collect the classes of all selected objects
 					TArray<UClass*> SelectionClasses;
@@ -3397,7 +3511,7 @@ void SSCSEditor::UpdateTree(bool bRegenerateTreeNodes)
 		SceneRootNodePtr.Reset();
 
 		// Build the tree data source according to what mode we're in
-		if(EditorMode.Get() == EEditorMode::BlueprintSCS)
+		if(EditorMode == EEditorMode::BlueprintSCS)
 		{
 			// Get the class default object
 			AActor* CDO = NULL;
@@ -3588,7 +3702,7 @@ UActorComponent* SSCSEditor::AddNewComponent( UClass* NewComponentClass, UObject
 {
 	const FScopedTransaction Transaction( LOCTEXT("AddComponent", "Add Component") );
 
-	if(EditorMode.Get() == EEditorMode::BlueprintSCS)
+	if(EditorMode == EEditorMode::BlueprintSCS)
 	{
 		UBlueprint* Blueprint = GetBlueprint();
 		check(Blueprint != nullptr && Blueprint->SimpleConstructionScript != nullptr);
@@ -3887,7 +4001,7 @@ void SSCSEditor::PasteNodes()
 	// Get the object that's being edited
 	UBlueprint* Blueprint = nullptr;
 	AActor* ActorInstance = nullptr;
-	if(EditorMode.Get() == EEditorMode::BlueprintSCS)
+	if(EditorMode == EEditorMode::BlueprintSCS)
 	{
 		Blueprint = GetBlueprint();
 		check(Blueprint != nullptr && Blueprint->SimpleConstructionScript != nullptr);
@@ -3996,7 +4110,7 @@ void SSCSEditor::OnDeleteNodes()
 {
 	const FScopedTransaction Transaction( LOCTEXT("RemoveComponent", "Remove Component") );
 
-	if(EditorMode.Get() == SSCSEditor::EEditorMode::BlueprintSCS)
+	if(EditorMode == SSCSEditor::EEditorMode::BlueprintSCS)
 	{
 		// Remove node from SCS
 		UBlueprint* Blueprint = GetBlueprint();
@@ -4100,7 +4214,7 @@ void SSCSEditor::RemoveComponentNode(FSCSEditorTreeNodePtrType InNodePtr)
 		SCSTreeWidget->ClearSelection();
 	}
 
-	if(EditorMode.Get() == SSCSEditor::EEditorMode::BlueprintSCS)
+	if(EditorMode == SSCSEditor::EEditorMode::BlueprintSCS)
 	{
 		USCS_Node* SCS_Node = InNodePtr->GetSCSNode();
 		if(SCS_Node != NULL)
@@ -4521,6 +4635,258 @@ void SSCSEditor::GetCollapsedNodes(const FSCSEditorTreeNodePtrType& InNodePtr, T
 	}
 }
 
+EVisibility SSCSEditor::GetPromoteToBlueprintButtonVisibility() const
+{
+	EVisibility Visibility = EVisibility::Collapsed;
+	if( EditorMode == SSCSEditor::EEditorMode::ActorInstance )
+	{
+		AActor* Actor = ActorContext.Get();
+		if( Actor && Actor->GetClass()->ClassGeneratedBy == nullptr )
+		{
+			Visibility = EVisibility::Visible;
+		}
+	}
+
+	return Visibility;
+}
+
+EVisibility SSCSEditor::GetEditBlueprintButtonVisibility() const
+{
+	EVisibility Visibility = EVisibility::Collapsed;
+	if(EditorMode == SSCSEditor::EEditorMode::ActorInstance)
+	{
+		AActor* Actor = ActorContext.Get();
+		if(Actor && Actor->GetClass()->ClassGeneratedBy != nullptr )
+		{
+			Visibility = EVisibility::Visible;
+		}
+	}
+
+	return Visibility;
+}
+
+FText SSCSEditor::OnGetApplyChangesToBlueprintTooltip() const
+{
+	int32 NumChangedProperties = 0;
+
+	AActor* Actor = ActorContext.Get();
+	UBlueprint* Blueprint = ActorInstance ? Cast<UBlueprint>(Actor->GetClass()->ClassGeneratedBy) : nullptr;
+
+	if(Actor != NULL && Blueprint != NULL && Actor->GetClass()->ClassGeneratedBy == Blueprint)
+	{
+		AActor* BlueprintCDO = Actor->GetClass()->GetDefaultObject<AActor>();
+		if(BlueprintCDO != NULL)
+		{
+			const auto CopyOptions = (EditorUtilities::ECopyOptions::Type)(EditorUtilities::ECopyOptions::PreviewOnly|EditorUtilities::ECopyOptions::OnlyCopyEditOrInterpProperties);
+			NumChangedProperties += EditorUtilities::CopyActorProperties(Actor, BlueprintCDO, CopyOptions);
+		}
+		NumChangedProperties += Actor->GetInstanceComponents().Num();
+	}
+
+
+	if(NumChangedProperties == 0)
+	{
+		return LOCTEXT("DisabledPushToBlueprintDefaults_ToolTip", "Replaces the Blueprint's defaults with any altered property values.");
+	}
+	else if(NumChangedProperties > 1)
+	{
+		return FText::Format(LOCTEXT("PushToBlueprintDefaults_ToolTip", "Click to apply {0} changed properties to the Blueprint."), FText::AsNumber(NumChangedProperties));
+	}
+	else
+	{
+		return LOCTEXT("PushOneToBlueprintDefaults_ToolTip", "Click to apply 1 changed property to the Blueprint.");
+	}
+}
+
+FText SSCSEditor::OnGetResetToBlueprintDefaultsTooltip() const
+{
+	int32 NumChangedProperties = 0;
+
+	AActor* Actor = ActorContext.Get();
+	UBlueprint* Blueprint = ActorInstance ? Cast<UBlueprint>(Actor->GetClass()->ClassGeneratedBy) : nullptr;
+	if(Actor != NULL && Blueprint != NULL && Actor->GetClass()->ClassGeneratedBy == Blueprint)
+	{
+		AActor* BlueprintCDO = Actor->GetClass()->GetDefaultObject<AActor>();
+		if(BlueprintCDO != NULL)
+		{
+			const auto CopyOptions = (EditorUtilities::ECopyOptions::Type)(EditorUtilities::ECopyOptions::PreviewOnly|EditorUtilities::ECopyOptions::OnlyCopyEditOrInterpProperties);
+			NumChangedProperties += EditorUtilities::CopyActorProperties(BlueprintCDO, Actor, CopyOptions);
+		}
+		NumChangedProperties += Actor->GetInstanceComponents().Num();
+	}
+
+	if(NumChangedProperties == 0)
+	{
+		return LOCTEXT("DisabledResetBlueprintDefaults_ToolTip", "Resets altered properties back to their Blueprint default values.");
+	}
+	else if(NumChangedProperties > 1)
+	{
+		return FText::Format(LOCTEXT("ResetToBlueprintDefaults_ToolTip", "Click to reset {0} changed properties to their Blueprint default values."), FText::AsNumber(NumChangedProperties));
+	}
+	else
+	{
+		return LOCTEXT("ResetOneToBlueprintDefaults_ToolTip", "Click to reset 1 changed property to its Blueprint default value.");
+	}
+}
+
+void SSCSEditor::OnOpenBlueprintEditor() const
+{
+	AActor* ActorInstance = ActorContext.Get();
+	UBlueprint* Blueprint = ActorInstance ? Cast<UBlueprint>( ActorInstance->GetClass()->ClassGeneratedBy ) : nullptr;
+
+	FAssetEditorManager::Get().OpenEditorForAsset(Blueprint);
+}
+
+void SSCSEditor::OnApplyChangesToBlueprint() const
+{
+	int32 NumChangedProperties = 0;
+
+	AActor* Actor = ActorContext.Get();
+	UBlueprint* Blueprint = Actor ? Cast<UBlueprint>(Actor->GetClass()->ClassGeneratedBy) : nullptr;
+
+	if (Actor != NULL && Blueprint != NULL && Actor->GetClass()->ClassGeneratedBy == Blueprint)
+	{
+		const FScopedTransaction Transaction(LOCTEXT("PushToBlueprintDefaults_Transaction", "Apply Changes to Blueprint"));
+
+		// Perform the actual copy
+		{
+			AActor* BlueprintCDO = Actor->GetClass()->GetDefaultObject<AActor>();
+			if (BlueprintCDO != NULL)
+			{
+				const auto CopyOptions = (EditorUtilities::ECopyOptions::Type)(EditorUtilities::ECopyOptions::OnlyCopyEditOrInterpProperties | EditorUtilities::ECopyOptions::PropagateChangesToArcheypeInstances);
+				NumChangedProperties = EditorUtilities::CopyActorProperties(Actor, BlueprintCDO, CopyOptions);
+				if (Actor->GetInstanceComponents().Num() > 0)
+				{
+					FKismetEditorUtilities::AddComponentsToBlueprint(Blueprint, Actor->GetInstanceComponents());
+					NumChangedProperties += Actor->GetInstanceComponents().Num();
+					Actor->ClearInstanceComponents();
+				}
+				if (NumChangedProperties > 0)
+				{
+					FBlueprintEditorUtils::MarkBlueprintAsModified(Blueprint);
+					FKismetEditorUtilities::CompileBlueprint(Blueprint);
+				}
+			}
+		}
+
+		// Set up a notification record to indicate success/failure
+		FNotificationInfo NotificationInfo(FText::GetEmpty());
+		NotificationInfo.FadeInDuration = 1.0f;
+		NotificationInfo.FadeOutDuration = 2.0f;
+		NotificationInfo.bUseLargeFont = false;
+		SNotificationItem::ECompletionState CompletionState;
+		if (NumChangedProperties > 0)
+		{
+			if (NumChangedProperties > 1)
+			{
+				FFormatNamedArguments Args;
+				Args.Add(TEXT("BlueprintName"), FText::FromName(Blueprint->GetFName()));
+				Args.Add(TEXT("NumChangedProperties"), NumChangedProperties);
+				Args.Add(TEXT("ActorName"), FText::FromString(Actor->GetActorLabel()));
+				NotificationInfo.Text = FText::Format(LOCTEXT("PushToBlueprintDefaults_ApplySuccess", "Updated Blueprint {BlueprintName} ({NumChangedProperties} property changes applied from actor {ActorName})."), Args);
+			}
+			else
+			{
+				FFormatNamedArguments Args;
+				Args.Add(TEXT("BlueprintName"), FText::FromName(Blueprint->GetFName()));
+				Args.Add(TEXT("ActorName"), FText::FromString(Actor->GetActorLabel()));
+				NotificationInfo.Text = FText::Format(LOCTEXT("PushOneToBlueprintDefaults_ApplySuccess", "Updated Blueprint {BlueprintName} (1 property change applied from actor {ActorName})."), Args);
+			}
+			CompletionState = SNotificationItem::CS_Success;
+		}
+		else
+		{
+			NotificationInfo.Text = LOCTEXT("PushToBlueprintDefaults_ApplyFailed", "No properties were copied");
+			CompletionState = SNotificationItem::CS_Fail;
+		}
+
+		// Add the notification to the queue
+		const auto Notification = FSlateNotificationManager::Get().AddNotification(NotificationInfo);
+		Notification->SetCompletionState(CompletionState);
+	}
+}
+
+void SSCSEditor::OnResetToBlueprintDefaults() const
+{
+	int32 NumChangedProperties = 0;
+
+	AActor* Actor = ActorContext.Get();
+	UBlueprint* Blueprint = Actor ? Cast<UBlueprint>(Actor->GetClass()->ClassGeneratedBy) : nullptr;
+
+	if (Actor != NULL && Blueprint != NULL && Actor->GetClass()->ClassGeneratedBy == Blueprint)
+	{
+		const FScopedTransaction Transaction(LOCTEXT("ResetToBlueprintDefaults_Transaction", "Reset to Blueprint Defaults"));
+
+		{
+			AActor* BlueprintCDO = Actor->GetClass()->GetDefaultObject<AActor>();
+			if (BlueprintCDO != NULL)
+			{
+				const auto CopyOptions = (EditorUtilities::ECopyOptions::Type)(EditorUtilities::ECopyOptions::OnlyCopyEditOrInterpProperties | EditorUtilities::ECopyOptions::CallPostEditChangeProperty);
+				NumChangedProperties = EditorUtilities::CopyActorProperties(BlueprintCDO, Actor, CopyOptions);
+			}
+			if (Actor->GetInstanceComponents().Num() > 0)
+			{
+				TArray<UActorComponent*> InstanceComponents = Actor->GetInstanceComponents();
+				NumChangedProperties += InstanceComponents.Num();
+				for (UActorComponent* ActorComponent : InstanceComponents)
+				{
+					if (ActorComponent)
+					{
+						ActorComponent->DestroyComponent();
+					}
+				}
+			}
+		}
+
+		// Set up a notification record to indicate success/failure
+		FNotificationInfo NotificationInfo(FText::GetEmpty());
+		NotificationInfo.FadeInDuration = 1.0f;
+		NotificationInfo.FadeOutDuration = 2.0f;
+		NotificationInfo.bUseLargeFont = false;
+		SNotificationItem::ECompletionState CompletionState;
+		if (NumChangedProperties > 0)
+		{
+			if (NumChangedProperties > 1)
+			{
+				FFormatNamedArguments Args;
+				Args.Add(TEXT("BlueprintName"), FText::FromName(Blueprint->GetFName()));
+				Args.Add(TEXT("NumChangedProperties"), NumChangedProperties);
+				Args.Add(TEXT("ActorName"), FText::FromString(Actor->GetActorLabel()));
+				NotificationInfo.Text = FText::Format(LOCTEXT("ResetToBlueprintDefaults_ApplySuccess", "Reset {ActorName} ({NumChangedProperties} property changes applied from Blueprint {BlueprintName})."), Args);
+			}
+			else
+			{
+				FFormatNamedArguments Args;
+				Args.Add(TEXT("BlueprintName"), FText::FromName(Blueprint->GetFName()));
+				Args.Add(TEXT("ActorName"), FText::FromString(Actor->GetActorLabel()));
+				NotificationInfo.Text = FText::Format(LOCTEXT("ResetOneToBlueprintDefaults_ApplySuccess", "Reset {ActorName} (1 property change applied from Blueprint {BlueprintName})."), Args);
+			}
+			CompletionState = SNotificationItem::CS_Success;
+		}
+		else
+		{
+			NotificationInfo.Text = LOCTEXT("ResetToBlueprintDefaults_Failed", "No properties were reset");
+			CompletionState = SNotificationItem::CS_Fail;
+		}
+
+		// Add the notification to the queue
+		const auto Notification = FSlateNotificationManager::Get().AddNotification(NotificationInfo);
+		Notification->SetCompletionState(CompletionState);
+	}
+}
+
+void SSCSEditor::PromoteToBlueprint() const
+{
+	bool bHarvest = false;
+	FCreateBlueprintFromActorDialog::OpenDialog(bHarvest, ActorContext.Get());
+}
+
+FReply SSCSEditor::OnPromoteToBlueprintClicked()
+{
+	PromoteToBlueprint();
+	return FReply::Handled();
+}
+
 const TArray<FSCSEditorTreeNodePtrType>& SSCSEditor::GetRootComponentNodes()
 {
 	return RootNodes;
diff --git a/Engine/Source/Editor/Kismet/Public/SSCSEditor.h b/Engine/Source/Editor/Kismet/Public/SSCSEditor.h
--- a/Engine/Source/Editor/Kismet/Public/SSCSEditor.h
+++ b/Engine/Source/Editor/Kismet/Public/SSCSEditor.h
@@ -430,7 +430,7 @@ public:
 		,_OnHighlightPropertyInDetailsView()
 		{}
 
-		SLATE_ATTRIBUTE(EEditorMode, EditorMode)
+		SLATE_ARGUMENT(EEditorMode, EditorMode)
 		SLATE_ATTRIBUTE(class AActor*, ActorContext)
 		SLATE_ATTRIBUTE(class AActor*, PreviewActor)
 		SLATE_ATTRIBUTE(bool, AllowEditing)
@@ -594,6 +594,8 @@ public:
 	/** Returns the set of root nodes */
 	const TArray<FSCSEditorTreeNodePtrType>& GetRootComponentNodes();
 
+	/** @return The current editor mode (editing live actors or editing blueprints) */
+	EEditorMode GetEditorMode() const { return EditorMode; }
 protected:
 	FString GetSelectedClassText() const;
 
@@ -683,6 +685,32 @@ protected:
 	/** Returns the set of expandable nodes that are currently collapsed in the UI */
 	void GetCollapsedNodes(const FSCSEditorTreeNodePtrType& InNodePtr, TSet<FSCSEditorTreeNodePtrType>& OutCollapsedNodes) const;
 
+	/** @return The visibility of the promote to blueprint button (only visible with an actor instance that is not created from a blueprint)*/
+	EVisibility GetPromoteToBlueprintButtonVisibility() const;
+
+	/** @return The visibility of the Edit Blueprint button (only visible with an actor instance that is created from a blueprint)*/
+	EVisibility GetEditBlueprintButtonVisibility() const;
+
+	/** @return the tooltip describing how many properties will be applied to the blueprint */
+	FText OnGetApplyChangesToBlueprintTooltip() const;
+
+	/** @return the tooltip describing how many properties will be reset to the blueprint default*/
+	FText OnGetResetToBlueprintDefaultsTooltip() const;
+
+	/** Opens the blueprint editor for the blueprint being viewed by the scseditor */
+	void OnOpenBlueprintEditor() const;
+
+	/** Propagates instance changes to the blueprint */
+	void OnApplyChangesToBlueprint() const;
+
+	/** Resets instance changes to the blueprint default */
+	void OnResetToBlueprintDefaults() const;
+
+	/** Converts the current actor instance to a blueprint */
+	void PromoteToBlueprint() const;
+
+	/** Called when the promote to blueprint button is clicked */
+	FReply OnPromoteToBlueprintClicked();
 public:
 	/** Tree widget */
 	TSharedPtr<SSCSTreeType> SCSTreeWidget;
@@ -699,9 +727,6 @@ public:
 	/** Whether or not the deferred rename request was flagged as transactional */
 	bool bIsDeferredRenameRequestTransactional;
 
-	/** Attribute to indicate which editor mode we're in. */
-	TAttribute<EEditorMode> EditorMode;
-
 	/** Attribute that provides access to the Actor context for which we are viewing/editing the SCS. */
 	TAttribute<class AActor*> ActorContext;
 
@@ -723,6 +748,8 @@ public:
 	/** Delegate to invoke when the given property should be highlighted in the details view (e.g. diff). */
 	FOnHighlightPropertyInDetailsView OnHighlightPropertyInDetailsView;
 private:
+	/** Indicates which editor mode we're in. */
+	EEditorMode EditorMode;
 
 	/** Root set of tree */
 	TArray<FSCSEditorTreeNodePtrType> RootNodes;
diff --git a/Engine/Source/Editor/KismetWidgets/Private/CreateBlueprintFromActorDialog.cpp b/Engine/Source/Editor/KismetWidgets/Private/CreateBlueprintFromActorDialog.cpp
--- a/Engine/Source/Editor/KismetWidgets/Private/CreateBlueprintFromActorDialog.cpp
+++ b/Engine/Source/Editor/KismetWidgets/Private/CreateBlueprintFromActorDialog.cpp
@@ -11,8 +11,12 @@
 
 #define LOCTEXT_NAMESPACE "CreateBlueprintFromActorDialog"
 
-void FCreateBlueprintFromActorDialog::OpenDialog(bool bInHarvest)
+TWeakObjectPtr<AActor> FCreateBlueprintFromActorDialog::ActorOverride;
+
+void FCreateBlueprintFromActorDialog::OpenDialog(bool bInHarvest, AActor* InActorOverride )
 {
+	ActorOverride = InActorOverride;
+
 	TSharedPtr<SWindow> PickBlueprintPathWidget;
 	SAssignNew(PickBlueprintPathWidget, SWindow)
 		.Title(LOCTEXT("SelectPath", "Select Path"))
@@ -21,13 +25,13 @@ void FCreateBlueprintFromActorDialog::OpenDialog(bool bInHarvest)
 
 	TSharedPtr<SCreateAssetFromActor> CreateBlueprintFromActorDialog;
 	PickBlueprintPathWidget->SetContent
-		(
+	(
 		SAssignNew(CreateBlueprintFromActorDialog, SCreateAssetFromActor, PickBlueprintPathWidget)
 		.AssetFilenameSuffix(TEXT("Blueprint"))
 		.HeadingText(LOCTEXT("CreateBlueprintFromActor_Heading", "Blueprint Name"))
 		.CreateButtonText(LOCTEXT("CreateBlueprintFromActor_ButtonLabel", "Create Blueprint"))
 		.OnCreateAssetAction(FOnPathChosen::CreateStatic(FCreateBlueprintFromActorDialog::OnCreateBlueprint, bInHarvest))
-		);
+	);
 
 	TSharedPtr<SWindow> RootWindow = FGlobalTabmanager::Get()->GetRootWindow();
 	if (RootWindow.IsValid())
@@ -62,10 +66,18 @@ void FCreateBlueprintFromActorDialog::OnCreateBlueprint(const FString& InAssetPa
 	}
 	else
 	{
-		TArray< UObject* > SelectedActors;
-		GEditor->GetSelectedActors()->GetSelectedObjects(AActor::StaticClass(), SelectedActors);
-		check(SelectedActors.Num());
-		Blueprint = FKismetEditorUtilities::CreateBlueprintFromActor(InAssetPath, (AActor*)SelectedActors[0], true);
+		AActor* ActorToUse = ActorOverride.Get();
+
+		if( !ActorToUse )
+		{
+			TArray< UObject* > SelectedActors;
+			GEditor->GetSelectedActors()->GetSelectedObjects(AActor::StaticClass(), SelectedActors);
+			check(SelectedActors.Num());
+			ActorToUse =  Cast<AActor>(SelectedActors[0]);
+		}
+
+		const bool bReplaceActor = true;
+		Blueprint = FKismetEditorUtilities::CreateBlueprintFromActor(InAssetPath, ActorToUse, bReplaceActor);
 	}
 
 	if(Blueprint)
diff --git a/Engine/Source/Editor/KismetWidgets/Public/CreateBlueprintFromActorDialog.h b/Engine/Source/Editor/KismetWidgets/Public/CreateBlueprintFromActorDialog.h
--- a/Engine/Source/Editor/KismetWidgets/Public/CreateBlueprintFromActorDialog.h
+++ b/Engine/Source/Editor/KismetWidgets/Public/CreateBlueprintFromActorDialog.h
@@ -12,8 +12,9 @@ public:
 	 * Static function to access constructing this window.
 	 *
 	 * @param bInHarvest		true if the components of the selected actors should be harvested for the blueprint.
+	 * @param ActorOverride		If set convert the specified actor, if null use the currently selected actor
 	 */
-	static KISMETWIDGETS_API void OpenDialog(bool bInHarvest);
+	static KISMETWIDGETS_API void OpenDialog(bool bInHarvest, AActor* InActorOverride = nullptr);
 private:
 
 	/** 
@@ -23,4 +24,6 @@ private:
 	 * @param bInHarvest		true if the components of the selected actors should be harvested for the blueprint.
 	 */
 	static void OnCreateBlueprint(const FString& InAssetPath, bool bInHarvest);
+private:
+	static TWeakObjectPtr<AActor> ActorOverride;
 };
diff --git a/Engine/Source/Editor/UnrealEd/Private/SComponentClassCombo.cpp b/Engine/Source/Editor/UnrealEd/Private/SComponentClassCombo.cpp
--- a/Engine/Source/Editor/UnrealEd/Private/SComponentClassCombo.cpp
+++ b/Engine/Source/Editor/UnrealEd/Private/SComponentClassCombo.cpp
@@ -45,7 +45,7 @@ void SComponentClassCombo::Construct(const FArguments& InArgs)
 		.Padding(2.f,1.f)
 		[
 			SNew(SImage)
-			.Image(FEditorStyle::GetBrush(TEXT("SCSEditor.NewComponentIcon")))
+			.Image(FEditorStyle::GetBrush(TEXT("Plus")))
 		]
 		+ SHorizontalBox::Slot()
 		.VAlign(VAlign_Center)
