Fixed a crash in GUID struct customization

diff --git a/Engine/Source/Editor/DetailCustomizations/Private/GuidStructCustomization.cpp b/Engine/Source/Editor/DetailCustomizations/Private/GuidStructCustomization.cpp
--- a/Engine/Source/Editor/DetailCustomizations/Private/GuidStructCustomization.cpp
+++ b/Engine/Source/Editor/DetailCustomizations/Private/GuidStructCustomization.cpp
@@ -134,6 +134,11 @@ FText FGuidStructCustomization::HandleTextBoxText( ) const
 		return LOCTEXT("MultipleValues", "Multiple Values");
 	}
 
+	if (RawData[0] == nullptr)
+	{
+		return FText::GetEmpty();
+	}
+
 	return FText::FromString(((FGuid*)RawData[0])->ToString(EGuidFormats::DigitsWithHyphensInBraces));
 }
 
