Re-enabled actor level info in the details pane.  Foliage does not show in the scene outliner so there's no other way to figure out what level it's in, and the foliage editing tool requires making the correct level current before you can edit the right foliage.

diff --git a/Engine/Source/Editor/DetailCustomizations/Private/ActorDetails.cpp b/Engine/Source/Editor/DetailCustomizations/Private/ActorDetails.cpp
--- a/Engine/Source/Editor/DetailCustomizations/Private/ActorDetails.cpp
+++ b/Engine/Source/Editor/DetailCustomizations/Private/ActorDetails.cpp
@@ -505,7 +505,7 @@ void FActorDetails::AddActorCategory( IDetailLayoutBuilder& DetailBuilder, const
 	IDetailCategoryBuilder& ActorCategory = DetailBuilder.EditCategory("Actor", FText::GetEmpty(), ECategoryPriority::Uncommon );
 
 
-#if 0
+#if 1
 	// Create the info buttons per level
 	for ( auto LevelIt( ActorsPerLevelCount.CreateConstIterator() ); LevelIt; ++LevelIt)
 	{
