[INTEGRATE] Change 2420282 by [REDACTED_PERSON][REDACTED_HANDLE] on 2015/01/27 09:38:23

	Fix very slow compile times in Slate widgets (~75% faster compile times in Development Win64)
	- Important:  BEGIN_SLATE_FUNCTION_BUILD_OPTIMIZATION only works at global scope on MSVC, not within classes (it silently does nothing)

diff --git a/Engine/Source/Editor/DetailCustomizations/Private/CollisionProfileDetails.cpp b/Engine/Source/Editor/DetailCustomizations/Private/CollisionProfileDetails.cpp
--- a/Engine/Source/Editor/DetailCustomizations/Private/CollisionProfileDetails.cpp
+++ b/Engine/Source/Editor/DetailCustomizations/Private/CollisionProfileDetails.cpp
@@ -29,6 +29,8 @@ DECLARE_DELEGATE_RetVal_TwoParams(bool, FOnValidateProfile, const FCollisionResp
 // SChannelEditDialog 
 //=====================================================================================
 
+BEGIN_SLATE_FUNCTION_BUILD_OPTIMIZATION
+
 class SChannelEditDialog : public SCompoundWidget
 {
 public:
@@ -2199,3 +2201,5 @@ void FCollisionProfileDetails::FCollsiionProfileData::Save(UCollisionProfile * P
 
 #undef LOCTEXT_NAMESPACE
 #undef RowWidth_Customization
+
+END_SLATE_FUNCTION_BUILD_OPTIMIZATION
\ No newline at end of file
diff --git a/Engine/Source/Editor/NiagaraEditor/Private/SNiagaraEffectEditorWidget.cpp b/Engine/Source/Editor/NiagaraEditor/Private/SNiagaraEffectEditorWidget.cpp
--- a/Engine/Source/Editor/NiagaraEditor/Private/SNiagaraEffectEditorWidget.cpp
+++ b/Engine/Source/Editor/NiagaraEditor/Private/SNiagaraEffectEditorWidget.cpp
@@ -10,6 +10,8 @@
 
 #define LOCTEXT_NAMESPACE "NiagaraEffectEditor"
 
+BEGIN_SLATE_FUNCTION_BUILD_OPTIMIZATION
+
 void SNiagaraEffectEditorWidget::Construct(const FArguments& InArgs)
 {
 	/*
@@ -589,4 +591,7 @@ void SEmitterWidget::OnSpawnScriptSelectedFromPicker(UObject *Asset)
 				.ItemHeight(20).ListItemsSource(&CurSpawnScript->Source->ExposedVectorConstants).OnGenerateRow(this, &SEmitterWidget::OnGenerateConstantListRow)
 			]
 		];
-}
\ No newline at end of file
+}
+
+
+END_SLATE_FUNCTION_BUILD_OPTIMIZATION
\ No newline at end of file
