[INTEGRATE] Change 2422951 by [REDACTED_PERSON][REDACTED_HANDLE] on 2015/01/28 18:01:07

	Material details have now moved to the primitive component customization so they can be seen on individual components

diff --git a/Engine/Source/Editor/DetailCustomizations/Private/ActorDetails.cpp b/Engine/Source/Editor/DetailCustomizations/Private/ActorDetails.cpp
--- a/Engine/Source/Editor/DetailCustomizations/Private/ActorDetails.cpp
+++ b/Engine/Source/Editor/DetailCustomizations/Private/ActorDetails.cpp
@@ -14,7 +14,6 @@
 #include "Runtime/AssetRegistry/Public/AssetRegistryModule.h"
 #include "BlueprintGraphDefinitions.h"
 #include "Engine/Breakpoint.h"
-#include "ActorMaterialCategory.h"
 #include "IContentBrowserSingleton.h"
 #include "ContentBrowserModule.h"
 #include "ScopedTransaction.h"
@@ -108,11 +107,6 @@ void FActorDetails::CustomizeDetails( IDetailLayoutBuilder& DetailLayout )
 		{
 			AddTransformCategory(DetailLayout);
 		}
-
-		if (!HideCategories.Contains(TEXT("Materials")))
-		{
-			AddMaterialCategory(DetailLayout);
-		}
 		
 		if (!HideCategories.Contains(TEXT("Actor")))
 		{
@@ -485,13 +479,6 @@ void FActorDetails::AddExperimentalWarningCategory( IDetailLayoutBuilder& Detail
 	}
 }
 
-void FActorDetails::AddMaterialCategory( IDetailLayoutBuilder& DetailBuilder )
-{
-	MaterialCategory = MakeShareable( new FActorMaterialCategory( SelectedActors ) );
-	MaterialCategory->Create( DetailBuilder );
-}
-
-
 void FActorDetails::AddActorCategory( IDetailLayoutBuilder& DetailBuilder, const TMap<ULevel*, int32>& ActorsPerLevelCount )
 {		
 	FLevelEditorModule& LevelEditor = FModuleManager::GetModuleChecked<FLevelEditorModule>( TEXT("LevelEditor") );
diff --git a/Engine/Source/Editor/DetailCustomizations/Private/ActorDetails.h b/Engine/Source/Editor/DetailCustomizations/Private/ActorDetails.h
--- a/Engine/Source/Editor/DetailCustomizations/Private/ActorDetails.h
+++ b/Engine/Source/Editor/DetailCustomizations/Private/ActorDetails.h
@@ -116,9 +116,6 @@ private:
 
 	TSharedPtr< IDetailCustomization > LayersLayoutDetails;
 
-	/** The material category that displays used materials */
-	TSharedPtr< class FActorMaterialCategory > MaterialCategory;
-
 	/** The active blutility function, picked from the combo list and cached here */
 	TWeakObjectPtr<UFunction> ActiveBlutilityFunction;
 
diff --git a/Engine/Source/Editor/DetailCustomizations/Private/ActorMaterialCategory.cpp b/Engine/Source/Editor/DetailCustomizations/Private/ActorMaterialCategory.cpp
deleted file mode 100644
diff --git a/Engine/Source/Editor/DetailCustomizations/Private/ActorMaterialCategory.h b/Engine/Source/Editor/DetailCustomizations/Private/ActorMaterialCategory.h
deleted file mode 100644
diff --git a/Engine/Source/Editor/DetailCustomizations/Private/ComponentMaterialCategory.cpp b/Engine/Source/Editor/DetailCustomizations/Private/ComponentMaterialCategory.cpp
new file mode 100644
--- /dev/null
+++ b/Engine/Source/Editor/DetailCustomizations/Private/ComponentMaterialCategory.cpp
@@ -0,0 +1,320 @@
+// Copyright 1998-2015 Epic Games, Inc. All Rights Reserved.
+
+#include "DetailCustomizationsPrivatePCH.h"
+#include "ComponentMaterialCategory.h"
+#include "AssetThumbnail.h"
+#include "ActorEditorUtils.h"
+
+#include "LandscapeProxy.h"
+#include "LandscapeComponent.h"
+#include "Components/DecalComponent.h"
+#include "Components/TextRenderComponent.h"
+#include "AI/Navigation/NavigationSystem.h"
+
+#define LOCTEXT_NAMESPACE "SMaterialList"
+
+/**
+ * Specialized iterator for stepping through used materials on actors
+ * Iterates through all materials on the provided list of actors by examining each actors component for materials
+ */
+class FMaterialIterator
+{
+public:
+	FMaterialIterator( TArray< TWeakObjectPtr<USceneComponent> >& InSelectedComponents )
+		: SelectedComponents( InSelectedComponents )
+		, CurMaterial( NULL )
+		, CurComponent( NULL )
+		, CurComponentIndex( 0 )
+		, CurMaterialIndex( -1 )
+		, bReachedEnd( false )
+		, bNeedComponents( true )
+	{
+		// Step to the first material
+		++(*this);
+	}
+
+	/**
+	 * Advances to the next material
+	 */
+	void operator++()
+	{
+		// Advance to the next material
+		++CurMaterialIndex;
+
+		// Examine each actor until we are out of actors
+		while( SelectedComponents.IsValidIndex(CurComponentIndex) )
+		{
+			USceneComponent* TestComponent = SelectedComponents[CurComponentIndex].Get();
+
+			if( TestComponent != NULL )
+			{
+				CurComponent = TestComponent;
+				int32 NumMaterials = 0;
+
+				// Primitive components and some actor components have materials
+				UPrimitiveComponent* PrimitiveComp = Cast<UPrimitiveComponent>( CurComponent );
+				UDecalComponent* DecalComponent = (PrimitiveComp ? NULL : Cast<UDecalComponent>( CurComponent ));
+
+				if( CurComponent->CreationMethod != EComponentCreationMethod::ConstructionScript)
+				{
+					if( PrimitiveComp )
+					{
+						NumMaterials = PrimitiveComp->GetNumMaterials();
+					}
+					else if( DecalComponent )
+					{
+						// DecalComponent isn't a primitive component so we must get the materials directly from it
+						NumMaterials = DecalComponent->GetNumMaterials();
+					}
+				}
+
+				// Check materials
+				while( CurMaterialIndex < NumMaterials )
+				{
+					UMaterialInterface* Material = NULL;
+
+					if( PrimitiveComp )
+					{
+						Material = PrimitiveComp->GetMaterial(CurMaterialIndex);
+					}
+					else if( DecalComponent )
+					{
+						Material = DecalComponent->GetMaterial(CurMaterialIndex);
+					}
+
+					CurMaterial = Material;
+							
+					// We step only once in the iterator if we have a material.  Note: A null material is considered valid
+					return;
+				}
+				// Out of materials on this component, reset for the next component
+				CurMaterialIndex = 0;
+			}
+			// Advance to the next compoment
+			++CurComponentIndex;
+		}
+
+
+		// Out of actors to check, reset to an invalid state
+		CurComponentIndex = INDEX_NONE;
+		bReachedEnd = true;
+		CurComponent = NULL;
+		CurMaterial = NULL;
+		CurMaterialIndex = INDEX_NONE;
+	}
+
+	/**
+	 * @return Whether or not the iterator is valid
+	 */
+	operator bool()
+	{
+		return !bReachedEnd;	
+	}
+
+	/**
+	 * @return The current material the iterator is stopped on
+	 */
+	UMaterialInterface* GetMaterial() const { return CurMaterial; }
+
+	/**
+	 * @return The index of the material in the current component
+	 */
+	int32 GetMaterialIndex() const { return CurMaterialIndex; }
+
+	/**
+	 * Swaps the current material the iterator is stopped on with the new material
+	 *
+	 * @param NewMaterial	The material to swap in
+	 */
+	void SwapMaterial( UMaterialInterface* NewMaterial )
+	{
+		UPrimitiveComponent* PrimitiveComp = Cast<UPrimitiveComponent>( CurComponent );
+		UDecalComponent* DecalComponent = Cast<UDecalComponent>( CurComponent );
+
+		if( PrimitiveComp )
+		{
+			PrimitiveComp->SetMaterial( CurMaterialIndex, NewMaterial );
+		}
+		else if( DecalComponent )
+		{
+			DecalComponent->SetMaterial( CurMaterialIndex, NewMaterial );
+		}
+	}
+
+	/**
+	 * @return The current component using the current material
+	 */
+	UActorComponent* GetComponent() const { return CurComponent; }
+
+private:
+	/** Reference to the selected actors */
+	TArray< TWeakObjectPtr<USceneComponent> >& SelectedComponents;
+	/** Reference to the current actor's components */
+	TInlineComponentArray<USceneComponent*> CurActorComponents;
+	/** The current material the iterator is stopped on */
+	UMaterialInterface* CurMaterial;
+	/** The current component using the current material */
+	UActorComponent* CurComponent;
+	/** The index of the actor we are stopped on */
+	int32 CurActorIndex;
+	/** The index of the component we are stopped on */
+	int32 CurComponentIndex;
+	/** The index of the material we are stopped on */
+	int32 CurMaterialIndex;
+	/** Whether or not we've reached the end of the actors */
+	uint32 bReachedEnd:1;
+	/** Whether we need to get components for the next actor */
+	uint32 bNeedComponents;
+};
+
+FComponentMaterialCategory::FComponentMaterialCategory( TArray< TWeakObjectPtr<USceneComponent> >& InSelectedComponents )
+	: SelectedComponents( InSelectedComponents )
+{
+}
+
+void FComponentMaterialCategory::Create( IDetailLayoutBuilder& DetailBuilder )
+{
+	FMaterialListDelegates MaterialListDelegates;
+	MaterialListDelegates.OnGetMaterials.BindSP( this, &FComponentMaterialCategory::OnGetMaterialsForView );
+	MaterialListDelegates.OnMaterialChanged.BindSP( this, &FComponentMaterialCategory::OnMaterialChanged );
+
+	TSharedRef<FMaterialList> MaterialList = MakeShareable( new FMaterialList( DetailBuilder, MaterialListDelegates ) );
+
+	bool bAnyMaterialsToDisplay = false;
+
+	for( FMaterialIterator It( SelectedComponents ); It; ++It )
+	{	
+		UActorComponent* CurrentComponent = It.GetComponent();
+
+		if( !bAnyMaterialsToDisplay && CurrentComponent->CreationMethod != EComponentCreationMethod::ConstructionScript )
+		{
+			bAnyMaterialsToDisplay = true;
+			break;
+		}
+	}
+
+
+	// only show the category if there are materials to display
+	if( bAnyMaterialsToDisplay )
+	{
+		// Make a category for the materials.
+		IDetailCategoryBuilder& MaterialCategory = DetailBuilder.EditCategory("Materials", FText::GetEmpty(), ECategoryPriority::TypeSpecific );
+
+		MaterialCategory.AddCustomBuilder( MaterialList );
+	}
+}
+
+void FComponentMaterialCategory::OnGetMaterialsForView( IMaterialListBuilder& MaterialList )
+{
+	const bool bAllowNullEntries = true;
+
+	// Iterate over every material on the actors
+	for( FMaterialIterator It( SelectedComponents ); It; ++It )
+	{	
+		int32 MaterialIndex = It.GetMaterialIndex();
+
+		UActorComponent* CurrentComponent = It.GetComponent();
+
+		if( CurrentComponent && CurrentComponent->CreationMethod != EComponentCreationMethod::ConstructionScript )
+		{
+			UMaterialInterface* Material = It.GetMaterial();
+
+			AActor* Actor = CurrentComponent->GetOwner();
+
+			// Component materials can be replaced if the component supports material overrides
+			const bool bCanBeReplaced =
+				( CurrentComponent->IsA( UMeshComponent::StaticClass() ) ||
+				CurrentComponent->IsA( UTextRenderComponent::StaticClass() ) ||
+				CurrentComponent->IsA( ULandscapeComponent::StaticClass() ) );
+
+			// Add the material if we allow null materials to be added or we have a valid material
+			if( bAllowNullEntries || Material )
+			{
+				MaterialList.AddMaterial( MaterialIndex, Material, bCanBeReplaced );
+			}
+		}
+	}
+}
+
+void FComponentMaterialCategory::OnMaterialChanged( UMaterialInterface* NewMaterial, UMaterialInterface* PrevMaterial, int32 SlotIndex, bool bReplaceAll )
+{
+	// Whether or not we should begin a transaction on swap
+	// Note we only begin a transaction on the first swap
+	bool bShouldMakeTransaction = true;
+
+	// Whether or not we made a transaction and need to end it
+	bool bMadeTransaction = false;
+
+	// Scan the selected actors mesh components for the old material and swap it with the new material 
+	for( FMaterialIterator It( SelectedComponents ); It; ++It )
+	{
+		int32 MaterialIndex = It.GetMaterialIndex();
+
+		UActorComponent* CurrentComponent = It.GetComponent();
+
+		if( CurrentComponent )
+		{
+			AActor* Actor = CurrentComponent->GetOwner();
+
+			// Component materials can be replaced if they are not created from a blueprint (not exposed to the user) and have material overrides on the component
+			bool bCanBeReplaced = Actor && Actor->GetClass()->ClassGeneratedBy == NULL && 
+				( CurrentComponent->IsA( UMeshComponent::StaticClass() ) ||
+				CurrentComponent->IsA( UDecalComponent::StaticClass() ) ||
+				CurrentComponent->IsA( UTextRenderComponent::StaticClass() ) ||
+				CurrentComponent->IsA( ULandscapeComponent::StaticClass() ) );
+
+			UMaterialInterface* Material = It.GetMaterial();
+			// Check if the material is the same as the previous material or we are replaceing all in the same slot.  If so we will swap it with the new material
+			if( bCanBeReplaced && ( Material == PrevMaterial || bReplaceAll ) && It.GetMaterialIndex() == SlotIndex )
+			{
+				// Begin a transaction for undo/redo the first time we encounter a material to replace.  
+				// There is only one transaction for all replacement
+				if( bShouldMakeTransaction && !bMadeTransaction )
+				{
+					GEditor->BeginTransaction( NSLOCTEXT("UnrealEd", "ReplaceComponentUsedMaterial", "Replace component used material") );
+
+					bMadeTransaction = true;
+				}
+
+				UProperty* MaterialProperty = NULL;
+				UObject* EditChangeObject = CurrentComponent;
+				if( CurrentComponent->IsA( UMeshComponent::StaticClass() ) )
+				{
+					MaterialProperty = FindField<UProperty>( UMeshComponent::StaticClass(), "Materials" );
+				}
+				else if( CurrentComponent->IsA( UDecalComponent::StaticClass() ) )
+				{
+					MaterialProperty = FindField<UProperty>( UDecalComponent::StaticClass(), "DecalMaterial" );
+				}
+				else if( CurrentComponent->IsA( UTextRenderComponent::StaticClass() ) )
+				{
+					MaterialProperty = FindField<UProperty>( UTextRenderComponent::StaticClass(), "TextMaterial" );
+				}
+				else if (CurrentComponent->IsA<ULandscapeComponent>() )
+				{
+					MaterialProperty = FindField<UProperty>( ALandscapeProxy::StaticClass(), "LandscapeMaterial" );
+					EditChangeObject = CastChecked<ULandscapeComponent>(CurrentComponent)->GetLandscapeProxy();
+				}
+
+				FNavigationLockContext NavUpdateLock(Actor->GetWorld(), ENavigationLockReason::MaterialUpdate);
+
+				EditChangeObject->PreEditChange( MaterialProperty );
+
+				It.SwapMaterial( NewMaterial );
+
+				FPropertyChangedEvent PropertyChangedEvent( MaterialProperty );
+				EditChangeObject->PostEditChangeProperty( PropertyChangedEvent );
+			}
+		}
+	}
+
+	if( bMadeTransaction )
+	{
+		// End the transation if we created one
+		GEditor->EndTransaction();
+		// Redraw viewports to reflect the material changes 
+		GUnrealEd->RedrawLevelEditingViewports();
+	}
+}
+
+#undef LOCTEXT_NAMESPACE
diff --git a/Engine/Source/Editor/DetailCustomizations/Private/ComponentMaterialCategory.h b/Engine/Source/Editor/DetailCustomizations/Private/ComponentMaterialCategory.h
new file mode 100644
--- /dev/null
+++ b/Engine/Source/Editor/DetailCustomizations/Private/ComponentMaterialCategory.h
@@ -0,0 +1,46 @@
+// Copyright 1998-2015 Epic Games, Inc. All Rights Reserved.
+
+#pragma once
+
+/**
+ * Encapsulates functionality for the ActorDetails material category
+ */
+class FComponentMaterialCategory : public TSharedFromThis<FComponentMaterialCategory>
+{
+public:
+	/**
+	 * Constructs a category in the details panel for displaying used materials
+	 */
+	FComponentMaterialCategory( TArray< TWeakObjectPtr<USceneComponent> >& InSelectedComponents );
+
+	/**
+	 * Creates the category
+	 */
+	void Create( IDetailLayoutBuilder& DetailBuilder );
+private:
+	/**
+	 * Called by the material list widget when we need to get new materials for the list
+	 *
+	 * @param OutMaterials	Handle to a material list builder that materials should be added to
+	 */
+	void OnGetMaterialsForView( class IMaterialListBuilder& OutMaterials );
+
+	/**
+	 * Called when a user drags a new material over a list item to replace it
+	 *
+	 * @param NewMaterial	The material that should replace the existing material
+	 * @param PrevMaterial	The material that should be replaced
+	 * @param SlotIndex		The index of the slot on the component where materials should be replaces
+	 * @param bReplaceAll	If true all materials in the slot should be replaced not just ones using PrevMaterial
+	 */
+	void OnMaterialChanged( UMaterialInterface* NewMaterial, UMaterialInterface* PrevMaterial, int32 SlotIndex, bool bReplaceAll );
+
+	/** 
+	 * @return true if a component is editable (and visible in the view)
+	 */
+	bool IsComponentEditable( UActorComponent* InComponent ) const;
+
+private:
+	/** Reference to the list of selected actors */
+	TArray< TWeakObjectPtr<USceneComponent> > SelectedComponents;
+};
\ No newline at end of file
diff --git a/Engine/Source/Editor/DetailCustomizations/Private/PrimitiveComponentDetails.cpp b/Engine/Source/Editor/DetailCustomizations/Private/PrimitiveComponentDetails.cpp
--- a/Engine/Source/Editor/DetailCustomizations/Private/PrimitiveComponentDetails.cpp
+++ b/Engine/Source/Editor/DetailCustomizations/Private/PrimitiveComponentDetails.cpp
@@ -11,7 +11,7 @@
 #include "IDocumentation.h"
 #include "EditorCategoryUtils.h"
 #include "PhysicsEngine/PhysicsSettings.h"
-
+#include "ComponentMaterialCategory.h"
 
 #define LOCTEXT_NAMESPACE "PrimitiveComponentDetails"
 
@@ -160,6 +160,18 @@ EVisibility FPrimitiveComponentDetails::IsMassVisible(bool bOverrideMass) const
 
 void FPrimitiveComponentDetails::CustomizeDetails( IDetailLayoutBuilder& DetailBuilder )
 {
+	// Get the objects being customized so we can enable/disable editing of 'Simulate Physics'
+	DetailBuilder.GetObjectsBeingCustomized(ObjectsCustomized);
+
+	// See if we are hiding Physics category
+	TArray<FString> HideCategories;
+	FEditorCategoryUtils::GetClassHideCategories(DetailBuilder.GetDetailsView().GetBaseClass(), HideCategories);
+
+	if(!HideCategories.Contains("Materials"))
+	{
+		AddMaterialCategory(DetailBuilder);
+	}
+
 	TSharedRef<IPropertyHandle> MobilityHandle = DetailBuilder.GetProperty(GET_MEMBER_NAME_CHECKED(UPrimitiveComponent, Mobility), USceneComponent::StaticClass());
 	MobilityHandle->SetToolTipText(LOCTEXT("PrimitiveMobilityTooltip", "Mobility for primitive components controls how they can be modified in game and therefore how they interact with lighting and physics.\n● A movable primitive component can be changed in game, but requires dynamic lighting and shadowing from lights which have a large performance cost.\n● A static primitive component can't be changed in game, but can have its lighting baked, which allows rendering to be very efficient."));
 
@@ -169,15 +181,12 @@ void FPrimitiveComponentDetails::CustomizeDetails( IDetailLayoutBuilder& DetailB
 		uint32 NumChildren = 0;
 		BodyInstanceHandler->GetNumChildren(NumChildren);
 
-		// See if we are hiding Physics category
-		TArray<FString> HideCategories;
-		FEditorCategoryUtils::GetClassHideCategories(DetailBuilder.GetDetailsView().GetBaseClass(), HideCategories);
+
 		if (!HideCategories.Contains(TEXT("Physics")))
 		{
 			IDetailCategoryBuilder& PhysicsCategory = DetailBuilder.EditCategory("Physics");
 
-			// Get the objects being customized so we can enable/disable editing of 'Simulate Physics'
-			DetailBuilder.GetObjectsBeingCustomized(ObjectsCustomized);
+
 
 			bool bDisplayMass = true;
 			bool bDisplayMassOverride = true;
@@ -310,6 +319,22 @@ void FPrimitiveComponentDetails::CustomizeDetails( IDetailLayoutBuilder& DetailB
 	AddAdvancedSubCategory( DetailBuilder, "Rendering", "LOD");
 }
 
+void FPrimitiveComponentDetails::AddMaterialCategory( IDetailLayoutBuilder& DetailBuilder )
+{
+	TArray<TWeakObjectPtr<USceneComponent> > Components;
+
+	for( TWeakObjectPtr<UObject> Object : ObjectsCustomized )
+	{
+		USceneComponent*  Component = Cast<USceneComponent>(Object.Get());
+		if( Component )
+		{
+			Components.Add( Component );
+		}
+	}
+
+	MaterialCategory = MakeShareable(new FComponentMaterialCategory(Components));
+	MaterialCategory->Create(DetailBuilder);
+}
 void FPrimitiveComponentDetails::AddAdvancedSubCategory( IDetailLayoutBuilder& DetailBuilder, FName MainCategoryName, FName SubCategoryName)
 {
 	TArray<TSharedRef<IPropertyHandle> > SubCategoryProperties;
diff --git a/Engine/Source/Editor/DetailCustomizations/Private/PrimitiveComponentDetails.h b/Engine/Source/Editor/DetailCustomizations/Private/PrimitiveComponentDetails.h
--- a/Engine/Source/Editor/DetailCustomizations/Private/PrimitiveComponentDetails.h
+++ b/Engine/Source/Editor/DetailCustomizations/Private/PrimitiveComponentDetails.h
@@ -12,6 +12,8 @@ public:
 	virtual void CustomizeDetails( IDetailLayoutBuilder& DetailBuilder ) override;
 
 private:
+	void AddMaterialCategory( IDetailLayoutBuilder& DetailBuilder);
+
 	ECheckBoxState IsMobilityActive(TSharedRef<IPropertyHandle> MobilityHandle, EComponentMobility::Type InMobility) const;
 
 	void OnMobilityChanged(ECheckBoxState InCheckedState, TSharedRef<IPropertyHandle> MobilityHandle, EComponentMobility::Type InMobility);
@@ -35,6 +37,7 @@ private:
 	/** Objects being customized so we can update the 'Simulate Physics' state if physics geometry is added/removed */
 	TArray< TWeakObjectPtr<UObject> > ObjectsCustomized;
 	TSharedPtr<IPropertyHandle> LockedAxisProperty;
+	TSharedPtr<class FComponentMaterialCategory> MaterialCategory;
 	EVisibility IsCustomLockedAxisSelected() const;
 	EVisibility IsLockAxisEnabled() const;
 
