[INTEGRATE] Change 2423650 by [REDACTED_PERSON][REDACTED_HANDLE] on 2015/01/29 08:24:09

	Another fix for non-unity build.

diff --git a/Engine/Source/Editor/DetailCustomizations/Private/DetailCustomizations.cpp b/Engine/Source/Editor/DetailCustomizations/Private/DetailCustomizations.cpp
--- a/Engine/Source/Editor/DetailCustomizations/Private/DetailCustomizations.cpp
+++ b/Engine/Source/Editor/DetailCustomizations/Private/DetailCustomizations.cpp
@@ -95,6 +95,7 @@
 #include "RigDetails.h"
 #include "SceneCaptureDetails.h"
 #include "CurveColorCustomization.h"
+#include "ActorComponentDetails.h"
 
 IMPLEMENT_MODULE( FDetailCustomizationsModule, DetailCustomizations );
 
