[INTEGRATE] Change 2428190 by [REDACTED_PERSON][REDACTED_HANDLE] on 2015/02/02 10:55:55

	Transform (location, rotation, scale) change from an original component template is propagated into overridden component templates. - WIP

diff --git a/Engine/Source/Editor/DetailCustomizations/Private/ComponentTransformDetails.cpp b/Engine/Source/Editor/DetailCustomizations/Private/ComponentTransformDetails.cpp
--- a/Engine/Source/Editor/DetailCustomizations/Private/ComponentTransformDetails.cpp
+++ b/Engine/Source/Editor/DetailCustomizations/Private/ComponentTransformDetails.cpp
@@ -52,19 +52,35 @@ static void PropagateTransformPropertyChange(UObject* InObject, UProperty* InPro
 	check(InObject != nullptr);
 	check(InProperty != nullptr);
 
-	TArray<UObject*> ArchetypeInstances;
 	TSet<USceneComponent*> UpdatedComponents;
-	FComponentEditorUtils::GetArchetypeInstances(InObject, ArchetypeInstances);
-	for(int32 InstanceIndex = 0; InstanceIndex < ArchetypeInstances.Num(); ++InstanceIndex)
+	if (InObject->HasAnyFlags(RF_ClassDefaultObject | RF_DefaultSubObject))
 	{
-		USceneComponent* InstancedSceneComponent = FComponentEditorUtils::GetSceneComponent(ArchetypeInstances[InstanceIndex], InObject);
-		if(InstancedSceneComponent != nullptr && !UpdatedComponents.Contains(InstancedSceneComponent))
+		TArray<UObject*> ArchetypeInstances;
+		FComponentEditorUtils::GetArchetypeInstances(InObject, ArchetypeInstances);
+		for (int32 InstanceIndex = 0; InstanceIndex < ArchetypeInstances.Num(); ++InstanceIndex)
 		{
-			FComponentEditorUtils::PropagateTransformPropertyChange(InstancedSceneComponent, InProperty, OldValue, NewValue, UpdatedComponents);
+			USceneComponent* InstancedSceneComponent = FComponentEditorUtils::GetSceneComponent(ArchetypeInstances[InstanceIndex], InObject);
+			if (InstancedSceneComponent != nullptr && !UpdatedComponents.Contains(InstancedSceneComponent))
+			{
+				FComponentEditorUtils::PropagateTransformPropertyChange(InstancedSceneComponent, InProperty, OldValue, NewValue, UpdatedComponents);
+			}
 		}
 	}
-}
 
+	if (InObject->HasAnyFlags(RF_ArchetypeObject))
+	{
+		TArray<UObject*> ArchetypeInstances;
+		InObject->GetArchetypeInstances(ArchetypeInstances);
+		for (int32 InstanceIndex = 0; InstanceIndex < ArchetypeInstances.Num(); ++InstanceIndex)
+		{
+			USceneComponent* InstancedSceneComponent = FComponentEditorUtils::GetSceneComponent(ArchetypeInstances[InstanceIndex], InObject);
+			if (InstancedSceneComponent != nullptr && InstancedSceneComponent->HasAnyFlags(RF_InheritableComponentTemplate) && !UpdatedComponents.Contains(InstancedSceneComponent))
+			{
+				FComponentEditorUtils::PropagateTransformPropertyChange(InstancedSceneComponent, InProperty, OldValue, NewValue, UpdatedComponents);
+			}
+		}
+	}
+}
 
 FComponentTransformDetails::FComponentTransformDetails( const TArray< TWeakObjectPtr<UObject> >& InSelectedObjects, const FSelectedActorInfo& InSelectedActorInfo, IDetailLayoutBuilder& DetailBuilder )
 	: TNumericUnitTypeInterface(EUnit::Centimeters)
@@ -695,10 +711,9 @@ void FComponentTransformDetails::OnToggleAbsoluteLocation( bool bEnable )
 				RootComponent->bAbsoluteLocation = !RootComponent->bAbsoluteLocation;
 
 				FPropertyChangedEvent PropertyChangedEvent( AbsoluteLocationProperty );
-				Object->PostEditChangeProperty( PropertyChangedEvent );
+				Object->PostEditChangeProperty(PropertyChangedEvent);
 
 				// If this is a default object or subobject, propagate the change out to any current instances of this object
-				if(Object->HasAnyFlags(RF_ClassDefaultObject|RF_DefaultSubObject))
 				{
 					uint32 NewValue = RootComponent->bAbsoluteLocation;
 					uint32 OldValue = !NewValue;
@@ -761,10 +776,9 @@ void FComponentTransformDetails::OnToggleAbsoluteRotation( bool bEnable )
 				RootComponent->bAbsoluteRotation = !RootComponent->bAbsoluteRotation;
 
 				FPropertyChangedEvent PropertyChangedEvent( AbsoluteRotationProperty );
-				Object->PostEditChangeProperty( PropertyChangedEvent );
+				Object->PostEditChangeProperty(PropertyChangedEvent);
 
 				// If this is a default object or subobject, propagate the change out to any current instances of this object
-				if(Object->HasAnyFlags(RF_ClassDefaultObject|RF_DefaultSubObject))
 				{
 					uint32 NewValue = RootComponent->bAbsoluteRotation;
 					uint32 OldValue = !NewValue;
@@ -826,10 +840,9 @@ void FComponentTransformDetails::OnToggleAbsoluteScale( bool bEnable )
 				RootComponent->bAbsoluteScale = !RootComponent->bAbsoluteScale;
 
 				FPropertyChangedEvent PropertyChangedEvent( AbsoluteScaleProperty );
-				Object->PostEditChangeProperty( PropertyChangedEvent );
+				Object->PostEditChangeProperty(PropertyChangedEvent);
 
 				// If this is a default object or subobject, propagate the change out to any current instances of this object
-				if(Object->HasAnyFlags(RF_ClassDefaultObject|RF_DefaultSubObject))
 				{
 					uint32 NewValue = RootComponent->bAbsoluteScale;
 					uint32 OldValue = !NewValue;
@@ -1081,10 +1094,7 @@ void FComponentTransformDetails::OnSetLocation( float NewValue, ETextCommit::Typ
 					Object->PostEditChangeProperty( PropertyChangedEvent );
 
 					// If this is a default object or subobject, propagate the change out to any current instances of this object
-					if(Object->HasAnyFlags(RF_ClassDefaultObject|RF_DefaultSubObject))
-					{
-						PropagateTransformPropertyChange(Object, RelativeLocationProperty, OldRelativeLocation, RelativeLocation);
-					}
+					PropagateTransformPropertyChange(Object, RelativeLocationProperty, OldRelativeLocation, RelativeLocation);
 
 					if( NotifyHook )
 					{
@@ -1198,10 +1208,7 @@ void FComponentTransformDetails::OnSetRotation( float NewValue, bool bCommitted,
 						}
 
 						// If this is a default object or subobject, propagate the change out to any current instances of this object
-						if(bIsEditingTemplateObject)
-						{
-							PropagateTransformPropertyChange(Object, RelativeRotationProperty, OldRelativeRotation, RelativeRotation);
-						}
+						PropagateTransformPropertyChange(Object, RelativeRotationProperty, OldRelativeRotation, RelativeRotation);
 
 						FPropertyChangedEvent PropertyChangedEvent( RelativeRotationProperty, !bCommitted && bEditingRotationInUI ? EPropertyChangeType::Interactive : EPropertyChangeType::ValueSet );
 
@@ -1452,10 +1459,7 @@ void FComponentTransformDetails::ScaleObject( float NewValue, int32 Axis, bool b
 					}
 
 					// If this is a default object or subobject, propagate the change out to any current instances of this object
-					if(Object->HasAnyFlags(RF_ClassDefaultObject|RF_DefaultSubObject))
-					{
-						PropagateTransformPropertyChange(Object, RelativeScale3DProperty, OldRelativeScale, RelativeScale);
-					}
+					PropagateTransformPropertyChange(Object, RelativeScale3DProperty, OldRelativeScale, RelativeScale);
 
 					if( NotifyHook )
 					{
diff --git a/Engine/Source/Editor/Kismet/Private/SCSEditorViewportClient.cpp b/Engine/Source/Editor/Kismet/Private/SCSEditorViewportClient.cpp
--- a/Engine/Source/Editor/Kismet/Private/SCSEditorViewportClient.cpp
+++ b/Engine/Source/Editor/Kismet/Private/SCSEditorViewportClient.cpp
@@ -443,12 +443,17 @@ bool FSCSEditorViewportClient::InputWidgetDelta( FViewport* Viewport, EAxisList:
 
 							FComponentEditorUtils::FTransformData NewDefaultTransform(*SelectedTemplate);
 
-							if (SelectedNodePtr->IsNative())
+							const bool bCanPropagateForNativeComponent = SelectedNodePtr->IsNative()
+								&& SelectedTemplate->HasAnyFlags(RF_DefaultSubObject);
+							if (bCanPropagateForNativeComponent)
 							{
-								if (ensure(SelectedTemplate->HasAnyFlags(RF_DefaultSubObject)))
-								{
-									FComponentEditorUtils::PropagateTransformPropertyChange(SelectedTemplate, OldDefaultTransform, NewDefaultTransform, UpdatedComponents);
-								}
+								FComponentEditorUtils::PropagateTransformPropertyChange(SelectedTemplate, OldDefaultTransform, NewDefaultTransform, UpdatedComponents);
+							}
+							const bool bCanPropagateForOverridenComponent = !SelectedNodePtr->IsNative()
+								&& SelectedTemplate->HasAnyFlags(RF_ArchetypeObject);
+							if (bCanPropagateForOverridenComponent)
+							{
+								FComponentEditorUtils::PropagateTransformPropertyChangeAmongOverridenTemplates(SelectedTemplate, OldDefaultTransform, NewDefaultTransform, UpdatedComponents);
 							}
 
 							if(PreviewBlueprint != NULL)
diff --git a/Engine/Source/Editor/UnrealEd/Private/Kismet2/ComponentEditorUtils.cpp b/Engine/Source/Editor/UnrealEd/Private/Kismet2/ComponentEditorUtils.cpp
--- a/Engine/Source/Editor/UnrealEd/Private/Kismet2/ComponentEditorUtils.cpp
+++ b/Engine/Source/Editor/UnrealEd/Private/Kismet2/ComponentEditorUtils.cpp
@@ -193,11 +193,11 @@ void FComponentEditorUtils::BindComponentSelectionOverride(USceneComponent* Scen
 		{
 			PrimComponent->Modify(bShouldMarkPackageDirty);
 			if (bBind)
-			{
-				PrimComponent->SelectionOverrideDelegate = UPrimitiveComponent::FSelectionOverride::CreateUObject(GUnrealEd, &UUnrealEdEngine::IsComponentSelected);
-			}
-			else
-			{
+		{
+			PrimComponent->SelectionOverrideDelegate = UPrimitiveComponent::FSelectionOverride::CreateUObject(GUnrealEd, &UUnrealEdEngine::IsComponentSelected);
+		}
+		else
+		{
 				PrimComponent->SelectionOverrideDelegate.Unbind();
 			}
 		}
@@ -211,18 +211,18 @@ void FComponentEditorUtils::BindComponentSelectionOverride(USceneComponent* Scen
 				{
 					PrimComponent->Modify(bShouldMarkPackageDirty);
 					if (bBind)
-					{
-						PrimComponent->SelectionOverrideDelegate = UPrimitiveComponent::FSelectionOverride::CreateUObject(GUnrealEd, &UUnrealEdEngine::IsComponentSelected);
-					}
+				{
+					PrimComponent->SelectionOverrideDelegate = UPrimitiveComponent::FSelectionOverride::CreateUObject(GUnrealEd, &UUnrealEdEngine::IsComponentSelected);
+				}
 					else
 					{
 						PrimComponent->SelectionOverrideDelegate.Unbind();
 					}
 				}
 			}
+			}
 		}
 	}
-}
 
 bool FComponentEditorUtils::AttemptApplyMaterialToComponent(USceneComponent* SceneComponent, UMaterialInterface* MaterialToApply, int32 OptionalMaterialSlot)
 {
@@ -274,6 +274,43 @@ bool FComponentEditorUtils::AttemptApplyMaterialToComponent(USceneComponent* Sce
 	return bResult;
 }
 
+void FComponentEditorUtils::PropagateTransformPropertyChangeAmongOverridenTemplates(
+	class USceneComponent* InSceneComponentTemplate,
+	const FTransformData& OldDefaultTransform,
+	const FTransformData& NewDefaultTransform,
+	TSet<class USceneComponent*>& UpdatedComponents)
+{
+	check(InSceneComponentTemplate != nullptr);
+	TArray<UObject*> ArchetypeInstances;
+	InSceneComponentTemplate->GetArchetypeInstances(ArchetypeInstances);
+	for (auto ArchetypeInstance : ArchetypeInstances)
+	{
+		USceneComponent* InstancedSceneComponent = Cast<USceneComponent>(ArchetypeInstance);
+		if (InstancedSceneComponent && InstancedSceneComponent->HasAnyFlags(RF_InheritableComponentTemplate) && !UpdatedComponents.Contains(InstancedSceneComponent))
+		{
+
+			static const UProperty* RelativeLocationProperty = FindFieldChecked<UProperty>(USceneComponent::StaticClass(), GET_MEMBER_NAME_CHECKED(USceneComponent, RelativeLocation));
+			if (RelativeLocationProperty != nullptr)
+			{
+				PropagateTransformPropertyChange(InstancedSceneComponent, RelativeLocationProperty, OldDefaultTransform.Trans, NewDefaultTransform.Trans, UpdatedComponents);
+			}
+
+			static const UProperty* RelativeRotationProperty = FindFieldChecked<UProperty>(USceneComponent::StaticClass(), GET_MEMBER_NAME_CHECKED(USceneComponent, RelativeLocation));
+			if (RelativeRotationProperty != nullptr)
+			{
+				PropagateTransformPropertyChange(InstancedSceneComponent, RelativeRotationProperty, OldDefaultTransform.Rot, NewDefaultTransform.Rot, UpdatedComponents);
+			}
+
+			static const UProperty* RelativeScale3DProperty = FindFieldChecked<UProperty>(USceneComponent::StaticClass(), GET_MEMBER_NAME_CHECKED(USceneComponent, RelativeScale3D));
+			if (RelativeScale3DProperty != nullptr)
+			{
+				PropagateTransformPropertyChange(InstancedSceneComponent, RelativeScale3DProperty, OldDefaultTransform.Scale, NewDefaultTransform.Scale, UpdatedComponents);
+			}
+
+		}
+	}
+}
+
 void FComponentEditorUtils::PropagateTransformPropertyChange(
 	class USceneComponent* InSceneComponentTemplate,
 	const FTransformData& OldDefaultTransform,
diff --git a/Engine/Source/Editor/UnrealEd/Public/Kismet2/ComponentEditorUtils.h b/Engine/Source/Editor/UnrealEd/Public/Kismet2/ComponentEditorUtils.h
--- a/Engine/Source/Editor/UnrealEd/Public/Kismet2/ComponentEditorUtils.h
+++ b/Engine/Source/Editor/UnrealEd/Public/Kismet2/ComponentEditorUtils.h
@@ -31,11 +31,11 @@ public:
 
 	/**
 	 * Ensures that the selection override delegate is properly bound for the supplied component
-	 * This includes any attached editor-only primitive components (such as billboard visualizers)
+	* This includes any attached editor-only primitive components (such as billboard visualizers)
 	 * 
 	 * @param SceneComponent The component to set the selection override for
 	 * @param bBind Whether the override should be bound
-	 */
+	*/
 	static void BindComponentSelectionOverride(USceneComponent* SceneComponent, bool bBind);
 
 	/**
@@ -68,6 +68,9 @@ public:
 	// Given a template, propagates a default transform change to all instances of the template
 	static void PropagateTransformPropertyChange(class USceneComponent* InSceneComponentTemplate, const FTransformData& OldDefaultTransform, const FTransformData& NewDefaultTransform, TSet<class USceneComponent*>& UpdatedComponents);
 
+	static void PropagateTransformPropertyChangeAmongOverridenTemplates(class USceneComponent* InSceneComponentTemplate, const FTransformData& OldDefaultTransform, const FTransformData& NewDefaultTransform, TSet<class USceneComponent*>& UpdatedComponents);
+
+
 	// Given an instance of a template and a property, propagates a default value change to the instance (only if applicable)
 	template<typename T>
 	static void PropagateTransformPropertyChange(class USceneComponent* InSceneComponent, const class UProperty* InProperty, const T& OldDefaultValue, const T& NewDefaultValue, TSet<class USceneComponent*>& UpdatedComponents)
