Templatized FInterval
Added Interval details customization, prerequisite for new foliage UI

diff --git a/Engine/Source/Editor/DetailCustomizations/Private/DetailCustomizations.cpp b/Engine/Source/Editor/DetailCustomizations/Private/DetailCustomizations.cpp
--- a/Engine/Source/Editor/DetailCustomizations/Private/DetailCustomizations.cpp
+++ b/Engine/Source/Editor/DetailCustomizations/Private/DetailCustomizations.cpp
@@ -32,6 +32,7 @@
 #include "MathStructCustomizations.h"
 #include "MathStructProxyCustomizations.h"
 #include "RangeStructCustomization.h"
+#include "IntervalStructCustomization.h"
 #include "StringAssetReferenceCustomization.h"
 #include "StringClassReferenceCustomization.h"
 #include "AttenuationSettingsCustomizations.h"
@@ -170,6 +171,8 @@ void FDetailCustomizationsModule::RegisterPropertyTypeCustomizations()
 	RegisterCustomPropertyTypeLayout("Key", FOnGetPropertyTypeCustomizationInstance::CreateStatic(&FKeyStructCustomization::MakeInstance));
 	RegisterCustomPropertyTypeLayout("FloatRange", FOnGetPropertyTypeCustomizationInstance::CreateStatic(&FRangeStructCustomization<float>::MakeInstance));
 	RegisterCustomPropertyTypeLayout("Int32Range", FOnGetPropertyTypeCustomizationInstance::CreateStatic(&FRangeStructCustomization<int32>::MakeInstance));
+	RegisterCustomPropertyTypeLayout("FloatInterval", FOnGetPropertyTypeCustomizationInstance::CreateStatic(&FIntervalStructCustomization<float>::MakeInstance));
+	RegisterCustomPropertyTypeLayout("Int32Interval", FOnGetPropertyTypeCustomizationInstance::CreateStatic(&FIntervalStructCustomization<int32>::MakeInstance));
 	RegisterCustomPropertyTypeLayout("DateTime", FOnGetPropertyTypeCustomizationInstance::CreateStatic(&FDateTimeStructCustomization::MakeInstance));
 	RegisterCustomPropertyTypeLayout("Timespan", FOnGetPropertyTypeCustomizationInstance::CreateStatic(&FTimespanStructCustomization::MakeInstance));
 	RegisterCustomPropertyTypeLayout("BlackboardEntry", FOnGetPropertyTypeCustomizationInstance::CreateStatic(&FBlackboardEntryDetails::MakeInstance));
diff --git a/Engine/Source/Editor/DetailCustomizations/Private/IntervalStructCustomization.cpp b/Engine/Source/Editor/DetailCustomizations/Private/IntervalStructCustomization.cpp
new file mode 100644
--- /dev/null
+++ b/Engine/Source/Editor/DetailCustomizations/Private/IntervalStructCustomization.cpp
@@ -0,0 +1,234 @@
+// Copyright 1998-2015 Epic Games, Inc. All Rights Reserved.
+
+#include "DetailCustomizationsPrivatePCH.h"
+#include "IntervalStructCustomization.h"
+#include "SNumericEntryBox.h"
+
+
+#define LOCTEXT_NAMESPACE "IntervalStructCustomization"
+
+
+/* Helper traits for getting a metadata property based on the template parameter type
+ *****************************************************************************/
+
+namespace IntervalMetadata
+{
+	template <typename NumericType>
+	struct FGetHelper
+	{ };
+
+	template <>
+	struct FGetHelper<float>
+	{
+		static float GetMetaData(const UProperty* Property, const TCHAR* Key)
+		{
+			return Property->GetFLOATMetaData(Key);
+		}
+	};
+
+	template <>
+	struct FGetHelper<int32>
+	{
+		static int32 GetMetaData(const UProperty* Property, const TCHAR* Key)
+		{
+			return Property->GetINTMetaData(Key);
+		}
+	};
+}
+
+
+/* FIntervalStructCustomization static interface
+ *****************************************************************************/
+
+template <typename NumericType>
+TSharedRef<IPropertyTypeCustomization> FIntervalStructCustomization<NumericType>::MakeInstance()
+{
+	return MakeShareable(new FIntervalStructCustomization<NumericType>);
+}
+
+
+/* IPropertyTypeCustomization interface
+ *****************************************************************************/
+
+template <typename NumericType>
+void FIntervalStructCustomization<NumericType>::CustomizeHeader(TSharedRef<IPropertyHandle> StructPropertyHandle, FDetailWidgetRow& HeaderRow, IPropertyTypeCustomizationUtils& StructCustomizationUtils)
+{
+	// Get handles to the properties we're interested in
+	MinValueHandle = StructPropertyHandle->GetChildHandle(TEXT("Min"));
+	MaxValueHandle = StructPropertyHandle->GetChildHandle(TEXT("Max"));
+	check(MinValueHandle.IsValid());
+	check(MaxValueHandle.IsValid());
+
+	// Get min/max metadata values if defined
+	auto Property = StructPropertyHandle->GetProperty();
+	if (Property != nullptr)
+	{
+		if (Property->HasMetaData(TEXT("UIMin")))
+		{
+			MinAllowedValue = TOptional<NumericType>(IntervalMetadata::FGetHelper<NumericType>::GetMetaData(Property, TEXT("UIMin")));
+		}
+
+		if (Property->HasMetaData(TEXT("UIMax")))
+		{
+			MaxAllowedValue = TOptional<NumericType>(IntervalMetadata::FGetHelper<NumericType>::GetMetaData(Property, TEXT("UIMax")));
+		}
+	}
+
+	// Make weak pointers to be passed as payloads to the widgets
+	TWeakPtr<IPropertyHandle> MinValueWeakPtr = MinValueHandle;
+	TWeakPtr<IPropertyHandle> MaxValueWeakPtr = MaxValueHandle;
+
+	// Build the widgets
+	HeaderRow.NameContent()
+	[
+		StructPropertyHandle->CreatePropertyNameWidget()
+	]
+	.ValueContent()
+	.MinDesiredWidth(200.0f)
+	//.MaxDesiredWidth(200.0f)
+	[
+		SNew(SBox)
+		.Padding(FMargin(0.0f, 3.0f, 0.0f, 2.0f))
+		[
+			SNew(SHorizontalBox)
+
+			+SHorizontalBox::Slot()
+			.Padding(FMargin(0.0f, 0.0f, 5.0f, 0.0f))
+			.VAlign(VAlign_Center)
+			[
+				SNew(SNumericEntryBox<NumericType>)
+				.Value(this, &FIntervalStructCustomization<NumericType>::OnGetValue, MinValueWeakPtr)
+				.MinValue(MinAllowedValue)
+				.MinSliderValue(MinAllowedValue)
+				.MaxValue(this, &FIntervalStructCustomization<NumericType>::OnGetValue, MaxValueWeakPtr)
+				.MaxSliderValue(this, &FIntervalStructCustomization<NumericType>::OnGetValue, MaxValueWeakPtr)
+				.OnValueCommitted(this, &FIntervalStructCustomization<NumericType>::OnValueCommitted, MinValueWeakPtr)
+				.OnValueChanged(this, &FIntervalStructCustomization<NumericType>::OnValueChanged, MinValueWeakPtr)
+				.OnBeginSliderMovement(this, &FIntervalStructCustomization<NumericType>::OnBeginSliderMovement)
+				.OnEndSliderMovement(this, &FIntervalStructCustomization<NumericType>::OnEndSliderMovement)
+				.Font(IDetailLayoutBuilder::GetDetailFont())
+				.AllowSpin(true)
+				.LabelVAlign(VAlign_Center)
+				.Label()
+				[
+					SNew(STextBlock)
+					.Font(IDetailLayoutBuilder::GetDetailFont())
+					.Text(LOCTEXT("MinLabel", "Min"))
+				]
+			]
+
+			+SHorizontalBox::Slot()
+			.Padding(FMargin(0.0f, 0.0f, 5.0f, 0.0f))
+			.VAlign(VAlign_Center)
+			[
+				SNew(SNumericEntryBox<NumericType>)
+				.Value(this, &FIntervalStructCustomization<NumericType>::OnGetValue, MaxValueWeakPtr)
+				.MinValue(this, &FIntervalStructCustomization<NumericType>::OnGetValue, MinValueWeakPtr)
+				.MinSliderValue(this, &FIntervalStructCustomization<NumericType>::OnGetValue, MinValueWeakPtr)
+				.MaxValue(MaxAllowedValue)
+				.MaxSliderValue(MaxAllowedValue)
+				.OnValueCommitted(this, &FIntervalStructCustomization<NumericType>::OnValueCommitted, MaxValueWeakPtr)
+				.OnValueChanged(this, &FIntervalStructCustomization<NumericType>::OnValueChanged, MaxValueWeakPtr)
+				.OnBeginSliderMovement(this, &FIntervalStructCustomization<NumericType>::OnBeginSliderMovement)
+				.OnEndSliderMovement(this, &FIntervalStructCustomization<NumericType>::OnEndSliderMovement)
+				.Font(IDetailLayoutBuilder::GetDetailFont())
+				.AllowSpin(true)
+				.LabelVAlign(VAlign_Center)
+				.Label()
+				[
+					SNew(STextBlock)
+					.Font(IDetailLayoutBuilder::GetDetailFont())
+					.Text(LOCTEXT("MaxLabel", "Max"))
+				]
+			]
+		]
+	];
+}
+
+
+template <typename NumericType>
+void FIntervalStructCustomization<NumericType>::CustomizeChildren(TSharedRef<IPropertyHandle> StructPropertyHandle, IDetailChildrenBuilder& StructBuilder, IPropertyTypeCustomizationUtils& StructCustomizationUtils)
+{
+	// Don't display children, as editing them directly can break the constraints
+}
+
+
+/* FIntervalStructCustomization callbacks
+ *****************************************************************************/
+
+template <typename NumericType>
+TOptional<NumericType> FIntervalStructCustomization<NumericType>::OnGetValue(TWeakPtr<IPropertyHandle> ValueWeakPtr) const
+{
+	auto ValueSharedPtr = ValueWeakPtr.Pin();
+
+	NumericType Value;
+	if (ensure(ValueSharedPtr->GetValue(Value) == FPropertyAccess::Success))
+	{
+		return TOptional<NumericType>(Value);
+	}
+
+	// Value couldn't be accessed. Return an unset value
+	return TOptional<NumericType>();
+}
+
+template <typename NumericType>
+void FIntervalStructCustomization<NumericType>::OnValueCommitted(NumericType NewValue, ETextCommit::Type CommitType, TWeakPtr<IPropertyHandle> HandleWeakPtr)
+{
+	auto HandleSharedPtr = HandleWeakPtr.Pin();
+	if (HandleSharedPtr.IsValid() && (!bIsUsingSlider || (bIsUsingSlider && ShouldAllowSpin())))
+	{
+		ensure(HandleSharedPtr->SetValue(NewValue) == FPropertyAccess::Success);
+	}
+}
+
+template <typename NumericType>
+void FIntervalStructCustomization<NumericType>::OnValueChanged(NumericType NewValue, TWeakPtr<IPropertyHandle> HandleWeakPtr)
+{
+	if (bIsUsingSlider && ShouldAllowSpin())
+	{
+		auto HandleSharedPtr = HandleWeakPtr.Pin();
+		if (HandleSharedPtr.IsValid())
+		{
+			ensure(HandleSharedPtr->SetValue(NewValue, EPropertyValueSetFlags::InteractiveChange) == FPropertyAccess::Success);
+		}
+	}
+}
+
+
+template <typename NumericType>
+void FIntervalStructCustomization<NumericType>::OnBeginSliderMovement()
+{
+	bIsUsingSlider = true;
+
+	if (ShouldAllowSpin())
+	{
+		GEditor->BeginTransaction(LOCTEXT("SetIntervalProperty", "Set Interval Property"));
+	}
+}
+
+
+template <typename NumericType>
+void FIntervalStructCustomization<NumericType>::OnEndSliderMovement(NumericType /*NewValue*/)
+{
+	bIsUsingSlider = false;
+
+	if (ShouldAllowSpin())
+	{
+		GEditor->EndTransaction();
+	}
+}
+
+template <typename NumericType>
+bool FIntervalStructCustomization<NumericType>::ShouldAllowSpin() const
+{
+	return true;
+}
+
+/* Only explicitly instantiate the types which are supported
+ *****************************************************************************/
+
+template class FIntervalStructCustomization<float>;
+template class FIntervalStructCustomization<int32>;
+
+
+#undef LOCTEXT_NAMESPACE
diff --git a/Engine/Source/Editor/DetailCustomizations/Private/IntervalStructCustomization.h b/Engine/Source/Editor/DetailCustomizations/Private/IntervalStructCustomization.h
new file mode 100644
--- /dev/null
+++ b/Engine/Source/Editor/DetailCustomizations/Private/IntervalStructCustomization.h
@@ -0,0 +1,88 @@
+// Copyright 1998-2015 Epic Games, Inc. All Rights Reserved.
+
+#pragma once
+
+/**
+* Implements a details panel customization for numeric TInterval structures.
+*/
+template <typename NumericType>
+class FIntervalStructCustomization : public IPropertyTypeCustomization
+{
+public:
+
+	FIntervalStructCustomization()
+		: bIsUsingSlider(false) {}
+
+public:
+
+	// IPropertyTypeCustomization interface
+
+	virtual void CustomizeHeader(TSharedRef<IPropertyHandle> StructPropertyHandle, FDetailWidgetRow& HeaderRow, IPropertyTypeCustomizationUtils& StructCustomizationUtils) override;
+	virtual void CustomizeChildren(TSharedRef<IPropertyHandle> StructPropertyHandle, IDetailChildrenBuilder& StructBuilder, IPropertyTypeCustomizationUtils& StructCustomizationUtils) override;
+
+public:
+
+	/**
+	* Creates a new instance.
+	*
+	* @return A new struct customization for Guids.
+	*/
+	static TSharedRef<IPropertyTypeCustomization> MakeInstance();
+
+protected:
+
+	/**
+	 * Gets the value for the provided property handle
+	 *
+	 * @param	ValueWeakPtr	Handle to the property to get the value from
+	 * @return	The value or unset if it could not be accessed
+	 */
+	TOptional<NumericType> OnGetValue(TWeakPtr<IPropertyHandle> ValueWeakPtr) const;
+
+	/**
+	 * Called when the value is committed from the property editor
+	 *
+	 * @param	NewValue		The new value of the property
+	 * @param	CommitType		How the value was committed (unused)
+	 * @param	HandleWeakPtr	Handle to the property that the new value is for
+	 */
+	void OnValueCommitted(NumericType NewValue, ETextCommit::Type CommitType, TWeakPtr<IPropertyHandle> HandleWeakPtr);
+
+	/**
+	 * Called when the value is changed in the property editor
+	 *
+	 * @param	NewValue		The new value of the property
+	 * @param	HandleWeakPtr	Handle to the property that the new value is for
+	 */
+	void OnValueChanged(NumericType NewValue, TWeakPtr<IPropertyHandle> HandleWeakPtr);
+
+	/**
+	 * Called when a value starts to be changed by a slider
+	 */
+	void OnBeginSliderMovement();
+
+	/**
+	 * Called when a value stops being changed by a slider
+	 *
+	 * @param	NewValue		The new value of the property
+	 */
+	void OnEndSliderMovement(NumericType NewValue);
+
+	/**
+	 * Determines if the spin box is enabled on a numeric value widget
+	 *
+	 * @return Whether the spin box should be enabled.
+	 */
+	bool ShouldAllowSpin() const;
+
+	// Cached shared pointers to properties that we are managing
+	TSharedPtr<IPropertyHandle> MinValueHandle;
+	TSharedPtr<IPropertyHandle> MaxValueHandle;
+
+	// Min and max allowed values from the metadata
+	TOptional<NumericType> MinAllowedValue;
+	TOptional<NumericType> MaxAllowedValue;
+	
+	// Flags whether the slider is being moved at the moment on any of our widgets
+	bool bIsUsingSlider;
+};
diff --git a/Engine/Source/Runtime/Core/Public/Math/Interval.h b/Engine/Source/Runtime/Core/Public/Math/Interval.h
--- a/Engine/Source/Runtime/Core/Public/Math/Interval.h
+++ b/Engine/Source/Runtime/Core/Public/Math/Interval.h
@@ -4,15 +4,15 @@
 
 
 /**
- * Structure for intervals of floating-point numbers.
+ * Template for interval of numbers.
  */
-struct FInterval
+template<typename ElementType> struct TInterval
 {
 	/** Holds the lower bound of the interval. */
-	float Min;
+	ElementType Min;
 	
 	/** Holds the upper bound of the interval. */
-	float Max;
+	ElementType Max;
 	
 	/** Holds a flag indicating whether the interval is empty. */
 	bool bIsEmpty;
@@ -24,9 +24,9 @@ public:
 	 *
 	 * The interval is initialized to [0, 0].
 	 */
-	FInterval()
-		: Min(0.0f)
-		, Max(0.0f)
+	TInterval()
+		: Min(0)
+		, Max(0)
 		, bIsEmpty(true)
 	{ }
 
@@ -36,7 +36,7 @@ public:
 	 * @param InMin The lower bound of the constructed interval.
 	 * @param InMax The upper bound of the constructed interval.
 	 */
-	FInterval( float InMin, float InMax )
+	TInterval( ElementType InMin, ElementType InMax )
 		: Min(InMin)
 		, Max(InMax)
 		, bIsEmpty(InMin >= InMax)
@@ -49,7 +49,7 @@ public:
 	 *
 	 * @param X The offset.
 	 */
-	void operator+= ( float X )
+	void operator+= ( ElementType X )
 	{
 		if (!bIsEmpty)
 		{
@@ -63,7 +63,7 @@ public:
 	 *
 	 * @param X The offset.
 	 */
-	void operator-= ( float X )
+	void operator-= ( ElementType X )
 	{
 		if (!bIsEmpty)
 		{
@@ -79,7 +79,7 @@ public:
 	 *
 	 * @param ExpandAmount The amount to expand by.
 	 */
-	void Expand( float ExpandAmount )
+	void Expand( ElementType ExpandAmount )
 	{
 		if (!bIsEmpty)
 		{
@@ -93,7 +93,27 @@ public:
 	 *
 	 * @param X The element to include.
 	 */
-	FORCEINLINE void Include( float X );
+	void Include( ElementType X )
+	{
+		if (bIsEmpty)
+		{
+			Min = X;
+			Max = X;
+			bIsEmpty = false;
+		}
+		else
+		{
+			if (X < Min)
+			{
+				Min = X;
+			}
+
+			if (X > Max)
+			{
+				Max = X;
+			}
+		}
+	}
 
 public:
 
@@ -104,39 +124,20 @@ public:
 	 * @param B The second interval.
 	 * @return The intersection.
 	 */
-	friend FInterval Intersect( const FInterval& A, const FInterval& B )
+	friend TInterval Intersect( const TInterval& A, const TInterval& B )
 	{
 		if (A.bIsEmpty || B.bIsEmpty)
 		{
-			return FInterval();
+			return TInterval();
 		}
 
-		return FInterval(FMath::Max(A.Min, B.Min), FMath::Min(A.Max, B.Max));
+		return TInterval(FMath::Max(A.Min, B.Min), FMath::Min(A.Max, B.Max));
 	}
 };
 
-
-/* FInterval inline functions
+/* Default intervals for built-in types
  *****************************************************************************/
 
-FORCEINLINE void FInterval::Include( float X )
-{
-	if (bIsEmpty)
-	{
-		Min = X;
-		Max = X;
-		bIsEmpty = false;
-	}
-	else
-	{
-		if (X < Min)
-		{
-			Min = X;
-		}
-
-		if (X > Max)
-		{
-			Max = X;
-		}
-	}
-}
+typedef TInterval<float> FInterval;
+typedef TInterval<float> FFloatInterval;
+typedef TInterval<int32> FInt32Interval;
diff --git a/Engine/Source/Runtime/CoreUObject/Classes/Object.h b/Engine/Source/Runtime/CoreUObject/Classes/Object.h
--- a/Engine/Source/Runtime/CoreUObject/Classes/Object.h
+++ b/Engine/Source/Runtime/CoreUObject/Classes/Object.h
@@ -840,6 +840,30 @@ struct FInt32Range
 	FInt32RangeBound UpperBound;
 };
 
+// A float interval
+
+USTRUCT(noexport)
+struct FFloatInterval
+{
+	UPROPERTY(EditAnywhere, Category=Interval)
+	float Min;
+
+	UPROPERTY(EditAnywhere, Category=Interval)
+	float Max;
+};
+
+// An int32 interval
+
+USTRUCT(noexport)
+struct FInt32Interval
+{
+	UPROPERTY(EditAnywhere, Category=Interval)
+	int32 Min;
+
+	UPROPERTY(EditAnywhere, Category=Interval)
+	int32 Max;
+};
+
 //=============================================================================
 /**
  * Object: The base class all objects.
