[INTEGRATE] Change 2433275 by [REDACTED_PERSON][REDACTED_HANDLE] on 2015/02/05 05:08:13

	UE-8715 You can edit the position and rotation of the root component in a derived blueprint

diff --git a/Engine/Source/Editor/DetailCustomizations/Private/SceneComponentDetails.cpp b/Engine/Source/Editor/DetailCustomizations/Private/SceneComponentDetails.cpp
--- a/Engine/Source/Editor/DetailCustomizations/Private/SceneComponentDetails.cpp
+++ b/Engine/Source/Editor/DetailCustomizations/Private/SceneComponentDetails.cpp
@@ -12,6 +12,7 @@
 #include "Engine/BlueprintGeneratedClass.h"
 #include "Engine/SCS_Node.h"
 #include "Engine/SimpleConstructionScript.h"
+#include "Engine/InheritableComponentHandler.h"
 #include "Components/LightComponentBase.h"
 
 #define LOCTEXT_NAMESPACE "SceneComponentDetails"
@@ -661,6 +662,20 @@ void FSceneComponentDetails::MakeTransformDetails( IDetailLayoutBuilder& DetailB
 					}
 				}
 			}
+
+			if (bShouldShowTransform && SceneComponent && SceneComponent->HasAnyFlags(RF_InheritableComponentTemplate))
+			{
+				auto OwnerClass = Cast<UClass>(SceneComponent->GetOuter());
+				auto Bluepirnt = UBlueprint::GetBlueprintFromClass(OwnerClass);
+				auto InheritableComponentHandler = Bluepirnt ? Bluepirnt->GetInheritableComponentHandler(false) : nullptr;
+				auto ComponentKey = InheritableComponentHandler ? InheritableComponentHandler->FindKey(SceneComponent) : FComponentKey();
+				auto SCSNode = ComponentKey.FindSCSNode();
+				const bool bProperRootNodeFound = ComponentKey.IsValid() && SCSNode && SCSNode->IsRootNode() && (SCSNode->ParentComponentOrVariableName == NAME_None);
+				if (bProperRootNodeFound)
+				{
+					bShouldShowTransform = false;
+				}
+			}
 		}
 
 		TSharedRef<FComponentTransformDetails> TransformDetails = MakeShareable( new FComponentTransformDetails( bIsEditingBlueprintDefaults ? SceneComponentObjects : DetailBuilder.GetDetailsView().GetSelectedObjects(), SelectedActorInfo, DetailBuilder ) );
diff --git a/Engine/Source/Runtime/Engine/Classes/Engine/InheritableComponentHandler.h b/Engine/Source/Runtime/Engine/Classes/Engine/InheritableComponentHandler.h
--- a/Engine/Source/Runtime/Engine/Classes/Engine/InheritableComponentHandler.h
+++ b/Engine/Source/Runtime/Engine/Classes/Engine/InheritableComponentHandler.h
@@ -88,6 +88,7 @@ public:
 
 	bool RenameTemplate(FComponentKey OldKey, FName NewName);
 
+	FComponentKey FindKey(UActorComponent* ComponentTemplate) const;
 #endif
 
 public:
diff --git a/Engine/Source/Runtime/Engine/Private/InheritableComponentHandler.cpp b/Engine/Source/Runtime/Engine/Private/InheritableComponentHandler.cpp
--- a/Engine/Source/Runtime/Engine/Private/InheritableComponentHandler.cpp
+++ b/Engine/Source/Runtime/Engine/Private/InheritableComponentHandler.cpp
@@ -246,6 +246,18 @@ bool UInheritableComponentHandler::RenameTemplate(FComponentKey OldKey, FName Ne
 	return false;
 }
 
+FComponentKey UInheritableComponentHandler::FindKey(UActorComponent* ComponentTemplate) const
+{
+	for (auto& Record : Records)
+	{
+		if (Record.ComponentTemplate == ComponentTemplate)
+		{
+			return Record.ComponentKey;
+		}
+	}
+	return FComponentKey();
+}
+
 #endif
 
 UActorComponent* UInheritableComponentHandler::GetOverridenComponentTemplate(FComponentKey Key) const
@@ -286,7 +298,7 @@ bool FComponentKey::Match(const FComponentKey OtherKey) const
 
 USCS_Node* FComponentKey::FindSCSNode() const
 {
-	auto ParentSCS = OwnerClass ? OwnerClass->SimpleConstructionScript : nullptr;
+	auto ParentSCS = (OwnerClass && IsValid()) ? OwnerClass->SimpleConstructionScript : nullptr;
 	return ParentSCS ? ParentSCS->FindSCSNodeByGuid(VariableGuid) : nullptr;
 }
 
