[INTEGRATE] Change 2434513 by [REDACTED_PERSON][REDACTED_HANDLE] on 2015/02/05 19:09:36

	Material customization in components on a blueprint now accepts new materials properly

diff --git a/Engine/Source/Editor/DetailCustomizations/Private/ComponentMaterialCategory.cpp b/Engine/Source/Editor/DetailCustomizations/Private/ComponentMaterialCategory.cpp
--- a/Engine/Source/Editor/DetailCustomizations/Private/ComponentMaterialCategory.cpp
+++ b/Engine/Source/Editor/DetailCustomizations/Private/ComponentMaterialCategory.cpp
@@ -27,7 +27,6 @@ public:
 		, CurComponentIndex( 0 )
 		, CurMaterialIndex( -1 )
 		, bReachedEnd( false )
-		, bNeedComponents( true )
 	{
 		// Step to the first material
 		++(*this);
@@ -41,7 +40,7 @@ public:
 		// Advance to the next material
 		++CurMaterialIndex;
 
-		// Examine each actor until we are out of actors
+		// Examine each component until we are out of them
 		while( SelectedComponents.IsValidIndex(CurComponentIndex) )
 		{
 			USceneComponent* TestComponent = SelectedComponents[CurComponentIndex].Get();
@@ -95,7 +94,7 @@ public:
 		}
 
 
-		// Out of actors to check, reset to an invalid state
+		// Out of components to check, reset to an invalid state
 		CurComponentIndex = INDEX_NONE;
 		bReachedEnd = true;
 		CurComponent = NULL;
@@ -147,24 +146,18 @@ public:
 	UActorComponent* GetComponent() const { return CurComponent; }
 
 private:
-	/** Reference to the selected actors */
+	/** Reference to the selected components */
 	TArray< TWeakObjectPtr<USceneComponent> >& SelectedComponents;
-	/** Reference to the current actor's components */
-	TInlineComponentArray<USceneComponent*> CurActorComponents;
 	/** The current material the iterator is stopped on */
 	UMaterialInterface* CurMaterial;
 	/** The current component using the current material */
 	UActorComponent* CurComponent;
-	/** The index of the actor we are stopped on */
-	int32 CurActorIndex;
 	/** The index of the component we are stopped on */
 	int32 CurComponentIndex;
 	/** The index of the material we are stopped on */
 	int32 CurMaterialIndex;
-	/** Whether or not we've reached the end of the actors */
+	/** Whether or not we've reached the end of the components */
 	uint32 bReachedEnd:1;
-	/** Whether we need to get components for the next actor */
-	uint32 bNeedComponents;
 };
 
 FComponentMaterialCategory::FComponentMaterialCategory( TArray< TWeakObjectPtr<USceneComponent> >& InSelectedComponents )
@@ -257,7 +250,7 @@ void FComponentMaterialCategory::OnMaterialChanged( UMaterialInterface* NewMater
 			AActor* Actor = CurrentComponent->GetOwner();
 
 			// Component materials can be replaced if they are not created from a blueprint (not exposed to the user) and have material overrides on the component
-			bool bCanBeReplaced = Actor && Actor->GetClass()->ClassGeneratedBy == NULL && 
+			bool bCanBeReplaced = (!Actor || Actor->GetClass()->ClassGeneratedBy == NULL) && 
 				( CurrentComponent->IsA( UMeshComponent::StaticClass() ) ||
 				CurrentComponent->IsA( UDecalComponent::StaticClass() ) ||
 				CurrentComponent->IsA( UTextRenderComponent::StaticClass() ) ||
@@ -296,7 +289,10 @@ void FComponentMaterialCategory::OnMaterialChanged( UMaterialInterface* NewMater
 					EditChangeObject = CastChecked<ULandscapeComponent>(CurrentComponent)->GetLandscapeProxy();
 				}
 
-				FNavigationLockContext NavUpdateLock(Actor->GetWorld(), ENavigationLockReason::MaterialUpdate);
+				if( Actor )
+				{
+					FNavigationLockContext NavUpdateLock(Actor->GetWorld(), ENavigationLockReason::MaterialUpdate);
+				}
 
 				EditChangeObject->PreEditChange( MaterialProperty );
 
