[INTEGRATE] Change 2436618 by [REDACTED_PERSON][REDACTED_HANDLE] on 2015/02/07 00:24:51

	Fix various issues with default transform property value propagation and "live" property defaults editing in the Blueprint editor.

	change summary:
	- revised/consolidated transform propagation code to make direct use of GetArchetypeInstances() for RF_ArchetypeObject template objects, now that GetArchetypes() is able to find SCS component templates stored in the BPGC.
	- removed unnecessary/unused default transform value propagation methods and code paths.
	- ensure that the SCS viewport calls Invalidate() on refresh so that changes are reflected in the preview viewport. this may have been a regression from preview actor instance update changes made to support single-screen mode.

diff --git a/Engine/Source/Editor/DetailCustomizations/Private/ComponentTransformDetails.cpp b/Engine/Source/Editor/DetailCustomizations/Private/ComponentTransformDetails.cpp
--- a/Engine/Source/Editor/DetailCustomizations/Private/ComponentTransformDetails.cpp
+++ b/Engine/Source/Editor/DetailCustomizations/Private/ComponentTransformDetails.cpp
@@ -46,40 +46,15 @@ private:
 	UWorld* PrevWorld;
 };
 
-template<typename T>
-static void PropagateTransformPropertyChange(UObject* InObject, UProperty* InProperty, const T& OldValue, const T& NewValue)
+static USceneComponent* GetSceneComponentFromDetailsObject(UObject* InObject)
 {
-	check(InObject != nullptr);
-	check(InProperty != nullptr);
-
-	TSet<USceneComponent*> UpdatedComponents;
-	if (InObject->HasAnyFlags(RF_ClassDefaultObject | RF_DefaultSubObject))
-	{
-		TArray<UObject*> ArchetypeInstances;
-		FComponentEditorUtils::GetArchetypeInstances(InObject, ArchetypeInstances);
-		for (int32 InstanceIndex = 0; InstanceIndex < ArchetypeInstances.Num(); ++InstanceIndex)
-		{
-			USceneComponent* InstancedSceneComponent = FComponentEditorUtils::GetSceneComponent(ArchetypeInstances[InstanceIndex], InObject);
-			if (InstancedSceneComponent != nullptr && !UpdatedComponents.Contains(InstancedSceneComponent))
-			{
-				FComponentEditorUtils::PropagateTransformPropertyChange(InstancedSceneComponent, InProperty, OldValue, NewValue, UpdatedComponents);
-			}
-		}
-	}
-
-	if (InObject->HasAnyFlags(RF_ArchetypeObject))
+	AActor* Actor = Cast<AActor>(InObject);
+	if(Actor)
 	{
-		TArray<UObject*> ArchetypeInstances;
-		InObject->GetArchetypeInstances(ArchetypeInstances);
-		for (int32 InstanceIndex = 0; InstanceIndex < ArchetypeInstances.Num(); ++InstanceIndex)
-		{
-			USceneComponent* InstancedSceneComponent = FComponentEditorUtils::GetSceneComponent(ArchetypeInstances[InstanceIndex], InObject);
-			if (InstancedSceneComponent != nullptr && InstancedSceneComponent->HasAnyFlags(RF_InheritableComponentTemplate) && !UpdatedComponents.Contains(InstancedSceneComponent))
-			{
-				FComponentEditorUtils::PropagateTransformPropertyChange(InstancedSceneComponent, InProperty, OldValue, NewValue, UpdatedComponents);
-			}
-		}
+		return Actor->GetRootComponent();
 	}
+	
+	return Cast<USceneComponent>(InObject);
 }
 
 FComponentTransformDetails::FComponentTransformDetails( const TArray< TWeakObjectPtr<UObject> >& InSelectedObjects, const FSelectedActorInfo& InSelectedActorInfo, IDetailLayoutBuilder& DetailBuilder )
@@ -683,8 +658,8 @@ void FComponentTransformDetails::OnToggleAbsoluteLocation( bool bEnable )
 		if( ObjectPtr.IsValid() )
 		{
 			UObject* Object = ObjectPtr.Get();
-			USceneComponent* RootComponent = FComponentEditorUtils::GetSceneComponent( Object );
-			if( RootComponent && RootComponent->bAbsoluteLocation != bEnable )
+			USceneComponent* SceneComponent = GetSceneComponentFromDetailsObject( Object );
+			if( SceneComponent && SceneComponent->bAbsoluteLocation != bEnable )
 			{
 				if( !bBeganTransaction )
 				{
@@ -708,16 +683,18 @@ void FComponentTransformDetails::OnToggleAbsoluteLocation( bool bEnable )
 					NotifyHook->NotifyPreChange( AbsoluteLocationProperty );
 				}
 
-				RootComponent->bAbsoluteLocation = !RootComponent->bAbsoluteLocation;
+				SceneComponent->bAbsoluteLocation = !SceneComponent->bAbsoluteLocation;
 
 				FPropertyChangedEvent PropertyChangedEvent( AbsoluteLocationProperty );
 				Object->PostEditChangeProperty(PropertyChangedEvent);
 
-				// If this is a default object or subobject, propagate the change out to any current instances of this object
+				// If it's a template, propagate the change out to any current instances of the object
+				if(SceneComponent->IsTemplate())
 				{
-					uint32 NewValue = RootComponent->bAbsoluteLocation;
+					uint32 NewValue = SceneComponent->bAbsoluteLocation;
 					uint32 OldValue = !NewValue;
-					PropagateTransformPropertyChange(Object, AbsoluteLocationProperty, OldValue, NewValue);
+					TSet<USceneComponent*> UpdatedInstances;
+					FComponentEditorUtils::PropagateDefaultValueChange(SceneComponent, AbsoluteLocationProperty, OldValue, NewValue, UpdatedInstances);
 				}
 
 				if( NotifyHook )
@@ -748,8 +725,8 @@ void FComponentTransformDetails::OnToggleAbsoluteRotation( bool bEnable )
 		if( ObjectPtr.IsValid() )
 		{
 			UObject* Object = ObjectPtr.Get();
-			USceneComponent* RootComponent = FComponentEditorUtils::GetSceneComponent( Object );
-			if( RootComponent && RootComponent->bAbsoluteRotation != bEnable )
+			USceneComponent* SceneComponent = GetSceneComponentFromDetailsObject( Object );
+			if( SceneComponent && SceneComponent->bAbsoluteRotation != bEnable )
 			{
 				if( !bBeganTransaction )
 				{
@@ -773,16 +750,18 @@ void FComponentTransformDetails::OnToggleAbsoluteRotation( bool bEnable )
 					NotifyHook->NotifyPreChange( AbsoluteRotationProperty );
 				}
 
-				RootComponent->bAbsoluteRotation = !RootComponent->bAbsoluteRotation;
+				SceneComponent->bAbsoluteRotation = !SceneComponent->bAbsoluteRotation;
 
 				FPropertyChangedEvent PropertyChangedEvent( AbsoluteRotationProperty );
 				Object->PostEditChangeProperty(PropertyChangedEvent);
 
-				// If this is a default object or subobject, propagate the change out to any current instances of this object
+				// If it's a template, propagate the change out to any current instances of the object
+				if(SceneComponent->IsTemplate())
 				{
-					uint32 NewValue = RootComponent->bAbsoluteRotation;
+					uint32 NewValue = SceneComponent->bAbsoluteRotation;
 					uint32 OldValue = !NewValue;
-					PropagateTransformPropertyChange(Object, AbsoluteRotationProperty, OldValue, NewValue);
+					TSet<USceneComponent*> UpdatedInstances;
+					FComponentEditorUtils::PropagateDefaultValueChange(SceneComponent, AbsoluteRotationProperty, OldValue, NewValue, UpdatedInstances);
 				}
 
 				if( NotifyHook )
@@ -812,8 +791,8 @@ void FComponentTransformDetails::OnToggleAbsoluteScale( bool bEnable )
 		if( ObjectPtr.IsValid() )
 		{
 			UObject* Object = ObjectPtr.Get();
-			USceneComponent* RootComponent = FComponentEditorUtils::GetSceneComponent( Object );
-			if( RootComponent && RootComponent->bAbsoluteScale != bEnable )
+			USceneComponent* SceneComponent = GetSceneComponentFromDetailsObject( Object );
+			if( SceneComponent && SceneComponent->bAbsoluteScale != bEnable )
 			{
 				if( !bBeganTransaction )
 				{
@@ -837,16 +816,18 @@ void FComponentTransformDetails::OnToggleAbsoluteScale( bool bEnable )
 					NotifyHook->NotifyPreChange( AbsoluteScaleProperty );
 				}
 
-				RootComponent->bAbsoluteScale = !RootComponent->bAbsoluteScale;
+				SceneComponent->bAbsoluteScale = !SceneComponent->bAbsoluteScale;
 
 				FPropertyChangedEvent PropertyChangedEvent( AbsoluteScaleProperty );
 				Object->PostEditChangeProperty(PropertyChangedEvent);
 
-				// If this is a default object or subobject, propagate the change out to any current instances of this object
+				// If it's a template, propagate the change out to any current instances of the object
+				if(SceneComponent->IsTemplate())
 				{
-					uint32 NewValue = RootComponent->bAbsoluteScale;
+					uint32 NewValue = SceneComponent->bAbsoluteScale;
 					uint32 OldValue = !NewValue;
-					PropagateTransformPropertyChange(Object, AbsoluteScaleProperty, OldValue, NewValue);
+					TSet<USceneComponent*> UpdatedInstances;
+					FComponentEditorUtils::PropagateDefaultValueChange(SceneComponent, AbsoluteScaleProperty, OldValue, NewValue, UpdatedInstances);
 				}
 
 				if( NotifyHook )
@@ -869,7 +850,7 @@ struct FGetRootComponentArchetype
 {
 	static USceneComponent* Get(UObject* Object)
 	{
-		auto RootComponent = Object ? FComponentEditorUtils::GetSceneComponent(Object) : nullptr;
+		auto RootComponent = Object ? GetSceneComponentFromDetailsObject(Object) : nullptr;
 		return RootComponent ? Cast<USceneComponent>(RootComponent->GetArchetype()) : nullptr;
 	}
 };
@@ -988,16 +969,16 @@ void FComponentTransformDetails::CacheTransform()
 		if( ObjectPtr.IsValid() )
 		{
 			UObject* Object = ObjectPtr.Get();
-			USceneComponent* RootComponent = FComponentEditorUtils::GetSceneComponent( Object );
+			USceneComponent* SceneComponent = GetSceneComponentFromDetailsObject( Object );
 
 			FVector Loc;
 			FRotator Rot;
 			FVector Scale;
-			if( RootComponent )
+			if( SceneComponent )
 			{
-				Loc = RootComponent->RelativeLocation;
-				Rot = (bEditingRotationInUI && !Object->HasAnyFlags(RF_ClassDefaultObject | RF_DefaultSubObject)) ? ObjectToRelativeRotationMap.FindOrAdd(Object) : RootComponent->RelativeRotation;
-				Scale = RootComponent->RelativeScale3D;
+				Loc = SceneComponent->RelativeLocation;
+				Rot = (bEditingRotationInUI && !Object->IsTemplate()) ? ObjectToRelativeRotationMap.FindOrAdd(Object) : SceneComponent->RelativeRotation;
+				Scale = SceneComponent->RelativeScale3D;
 
 				if( ObjectIndex == 0 )
 				{
@@ -1010,9 +991,9 @@ void FComponentTransformDetails::CacheTransform()
 					CachedRotation.Set( Rot );
 					CachedScale.Set( Scale );
 
-					bAbsoluteLocation = RootComponent->bAbsoluteLocation;
-					bAbsoluteScale = RootComponent->bAbsoluteScale;
-					bAbsoluteRotation = RootComponent->bAbsoluteRotation;
+					bAbsoluteLocation = SceneComponent->bAbsoluteLocation;
+					bAbsoluteScale = SceneComponent->bAbsoluteScale;
+					bAbsoluteRotation = SceneComponent->bAbsoluteRotation;
 				}
 				else if( CurLoc != Loc || CurRot != Rot || CurScale != Scale )
 				{
@@ -1052,10 +1033,10 @@ void FComponentTransformDetails::OnSetLocation( float NewValue, ETextCommit::Typ
 		if( ObjectPtr.IsValid() )
 		{
 			UObject* Object = ObjectPtr.Get();
-			USceneComponent* RootComponent = FComponentEditorUtils::GetSceneComponent( Object );
-			if( RootComponent )
+			USceneComponent* SceneComponent = GetSceneComponentFromDetailsObject( Object );
+			if( SceneComponent )
 			{
-				FVector OldRelativeLocation = RootComponent->RelativeLocation;
+				FVector OldRelativeLocation = SceneComponent->RelativeLocation;
 				FVector RelativeLocation = OldRelativeLocation;
 				float OldValue = RelativeLocation[Axis];
 				if( OldValue != NewValue )
@@ -1101,18 +1082,22 @@ void FComponentTransformDetails::OnSetLocation( float NewValue, ETextCommit::Typ
 					{
 						// HACK: Set directly if no actors are selected since this causes Rot->Quat->Rot conversion
 						// (recalculates relative rotation even though this is location)
-						RootComponent->RelativeLocation = RelativeLocation;
+						SceneComponent->RelativeLocation = RelativeLocation;
 					}
 					else
 					{
-						RootComponent->SetRelativeLocation( RelativeLocation );
+						SceneComponent->SetRelativeLocation( RelativeLocation );
 					}
 
 					FPropertyChangedEvent PropertyChangedEvent( RelativeLocationProperty );
 					Object->PostEditChangeProperty( PropertyChangedEvent );
 
-					// If this is a default object or subobject, propagate the change out to any current instances of this object
-					PropagateTransformPropertyChange(Object, RelativeLocationProperty, OldRelativeLocation, RelativeLocation);
+					// If it's a template, propagate the change out to any current instances of the object
+					if(SceneComponent->IsTemplate())
+					{
+						TSet<USceneComponent*> UpdatedInstances;
+						FComponentEditorUtils::PropagateDefaultValueChange(SceneComponent, RelativeLocationProperty, OldRelativeLocation, RelativeLocation, UpdatedInstances);
+					}
 
 					if( NotifyHook )
 					{
@@ -1150,12 +1135,12 @@ void FComponentTransformDetails::OnSetRotation( float NewValue, bool bCommitted,
 			if( ObjectPtr.IsValid() )
 			{
 				UObject* Object = ObjectPtr.Get();
-				USceneComponent* RootComponent = FComponentEditorUtils::GetSceneComponent( Object );
-				if( RootComponent )
+				USceneComponent* SceneComponent = GetSceneComponentFromDetailsObject( Object );
+				if( SceneComponent )
 				{
-					const bool bIsEditingTemplateObject = Object->HasAnyFlags(RF_ClassDefaultObject|RF_DefaultSubObject);
+					const bool bIsEditingTemplateObject = Object->IsTemplate();
 
-					FRotator& RelativeRotation = (bEditingRotationInUI && !bIsEditingTemplateObject) ? ObjectToRelativeRotationMap.FindOrAdd(Object) : RootComponent->RelativeRotation;
+					FRotator& RelativeRotation = (bEditingRotationInUI && !bIsEditingTemplateObject) ? ObjectToRelativeRotationMap.FindOrAdd(Object) : SceneComponent->RelativeRotation;
 					FRotator OldRelativeRotation = RelativeRotation;
 
 					float& ValueToChange = Axis == 0 ? RelativeRotation.Roll : Axis == 1 ? RelativeRotation.Pitch : RelativeRotation.Yaw;
@@ -1175,7 +1160,7 @@ void FComponentTransformDetails::OnSetRotation( float NewValue, bool bCommitted,
 								GEditor->BeginTransaction( LOCTEXT( "OnSetRotation_ComponentDirect", "Modify Component(s)") );
 							}
 
-							if (!Object->HasAnyFlags(RF_ClassDefaultObject|RF_DefaultSubObject))
+							if (!bIsEditingTemplateObject)
 							{
 								// Broadcast the first time an actor is about to move
 								GEditor->BroadcastBeginObjectMovement( *Object );
@@ -1211,11 +1196,11 @@ void FComponentTransformDetails::OnSetRotation( float NewValue, bool bCommitted,
 							{
 								// HACK: Set directly if no actors are selected since this causes Rot->Quat->Rot conversion issues
 								// (recalculates relative rotation from quat which can give an equivalent but different value than the user typed)
-								RootComponent->RelativeRotation = RelativeRotation;
+								SceneComponent->RelativeRotation = RelativeRotation;
 							}
 							else
 							{
-								RootComponent->SetRelativeRotation( RelativeRotation );
+								SceneComponent->SetRelativeRotation( RelativeRotation );
 							}
 						}
 
@@ -1225,8 +1210,12 @@ void FComponentTransformDetails::OnSetRotation( float NewValue, bool bCommitted,
 							ObjectAsActor->ReregisterAllComponents();
 						}
 
-						// If this is a default object or subobject, propagate the change out to any current instances of this object
-						PropagateTransformPropertyChange(Object, RelativeRotationProperty, OldRelativeRotation, RelativeRotation);
+						// If it's a template, propagate the change out to any current instances of the object
+						if(SceneComponent->IsTemplate())
+						{
+							TSet<USceneComponent*> UpdatedInstances;
+							FComponentEditorUtils::PropagateDefaultValueChange(SceneComponent, RelativeRotationProperty, OldRelativeRotation, RelativeRotation, UpdatedInstances);
+						}
 
 						FPropertyChangedEvent PropertyChangedEvent( RelativeRotationProperty, !bCommitted && bEditingRotationInUI ? EPropertyChangeType::Interactive : EPropertyChangeType::ValueSet );
 
@@ -1299,8 +1288,8 @@ void FComponentTransformDetails::OnBeginRotatonSlider()
 				bBeganTransaction = true;
 			}
 
-			USceneComponent* RootComponent = FComponentEditorUtils::GetSceneComponent( Object );
-			if( RootComponent )
+			USceneComponent* SceneComponent = GetSceneComponentFromDetailsObject( Object );
+			if( SceneComponent )
 			{
 				FScopedSwitchWorldForObject WorldSwitcher( Object );
 				
@@ -1314,7 +1303,7 @@ void FComponentTransformDetails::OnBeginRotatonSlider()
 				Object->PreEditChange( RelativeRotationProperty );
 
 				// Add/update cached rotation value prior to slider interaction
-				ObjectToRelativeRotationMap.FindOrAdd(Object) = RootComponent->RelativeRotation;
+				ObjectToRelativeRotationMap.FindOrAdd(Object) = SceneComponent->RelativeRotation;
 			}
 		}
 	}
@@ -1336,8 +1325,8 @@ void FComponentTransformDetails::OnEndRotationSlider(float NewValue)
 		if( ObjectPtr.IsValid() )
 		{
 			UObject* Object = ObjectPtr.Get();
-			USceneComponent* RootComponent = FComponentEditorUtils::GetSceneComponent( Object );
-			if( RootComponent )
+			USceneComponent* SceneComponent = GetSceneComponentFromDetailsObject( Object );
+			if( SceneComponent )
 			{
 				FScopedSwitchWorldForObject WorldSwitcher( Object );
 
@@ -1370,10 +1359,10 @@ void FComponentTransformDetails::ScaleObject( float NewValue, int32 Axis, bool b
 		if( ObjectPtr.IsValid() )
 		{
 			UObject* Object = ObjectPtr.Get();
-			USceneComponent* RootComponent = FComponentEditorUtils::GetSceneComponent( Object );
-			if( RootComponent )
+			USceneComponent* SceneComponent = GetSceneComponentFromDetailsObject( Object );
+			if( SceneComponent )
 			{
-				FVector OldRelativeScale = RootComponent->RelativeScale3D;
+				FVector OldRelativeScale = SceneComponent->RelativeScale3D;
 				FVector RelativeScale = OldRelativeScale;
 				if( bMirror )
 				{
@@ -1433,7 +1422,7 @@ void FComponentTransformDetails::ScaleObject( float NewValue, int32 Axis, bool b
 						}
 					}
 
-					RootComponent->SetRelativeScale3D( RelativeScale );
+					SceneComponent->SetRelativeScale3D( RelativeScale );
 
 					// Build property chain so the actor knows whether we changed the X, Y or Z
 					FEditPropertyChain PropertyChain;
@@ -1476,8 +1465,12 @@ void FComponentTransformDetails::ScaleObject( float NewValue, int32 Axis, bool b
 						CacheTransform();
 					}
 
-					// If this is a default object or subobject, propagate the change out to any current instances of this object
-					PropagateTransformPropertyChange(Object, RelativeScale3DProperty, OldRelativeScale, RelativeScale);
+					// If it's a template, propagate the change out to any current instances of the object
+					if(SceneComponent->IsTemplate())
+					{
+						TSet<USceneComponent*> UpdatedInstances;
+						FComponentEditorUtils::PropagateDefaultValueChange(SceneComponent, RelativeScale3DProperty, OldRelativeScale, RelativeScale, UpdatedInstances);
+					}
 
 					if( NotifyHook )
 					{
diff --git a/Engine/Source/Editor/Kismet/Private/SCSEditorViewportClient.cpp b/Engine/Source/Editor/Kismet/Private/SCSEditorViewportClient.cpp
--- a/Engine/Source/Editor/Kismet/Private/SCSEditorViewportClient.cpp
+++ b/Engine/Source/Editor/Kismet/Private/SCSEditorViewportClient.cpp
@@ -108,7 +108,6 @@ namespace
 FSCSEditorViewportClient::FSCSEditorViewportClient(TWeakPtr<FBlueprintEditor>& InBlueprintEditorPtr, FPreviewScene* InPreviewScene, const TSharedRef<SSCSEditorViewport>& InSCSEditorViewport)
 	: FEditorViewportClient(nullptr, InPreviewScene, StaticCastSharedRef<SEditorViewport>(InSCSEditorViewport))
 	, BlueprintEditorPtr(InBlueprintEditorPtr)
-	, PreviewBlueprint(NULL)
 	, PreviewActorBounds(ForceInitToZero)
 	, bIsManipulating(false)
 	, ScopedTransaction(NULL)
@@ -409,8 +408,6 @@ bool FSCSEditorViewportClient::InputWidgetDelta( FViewport* Viewport, EAxisList:
 
 							// Adjust the deltas as necessary
 							FComponentEditorUtils::AdjustComponentDelta(SceneComp, Drag, Rot);
-							
-							FComponentEditorUtils::FTransformData OldDefaultTransform(*SelectedTemplate);
 
 							TSharedPtr<ISCSEditorCustomization> Customization = BlueprintEditor->CustomizeSCSEditor(SceneComp);
 							if(Customization.IsValid() && Customization->HandleViewportDrag(SceneComp, SelectedTemplate, Drag, Rot, ModifiedScale, GetWidgetLocation()))
@@ -441,21 +438,7 @@ bool FSCSEditorViewportClient::InputWidgetDelta( FViewport* Viewport, EAxisList:
 								UpdatedComponents.Add(SelectedTemplate);
 							}
 
-							FComponentEditorUtils::FTransformData NewDefaultTransform(*SelectedTemplate);
-
-							const bool bCanPropagateForNativeComponent = SelectedNodePtr->IsNative()
-								&& SelectedTemplate->HasAnyFlags(RF_DefaultSubObject);
-							if (bCanPropagateForNativeComponent)
-							{
-								FComponentEditorUtils::PropagateTransformPropertyChange(SelectedTemplate, OldDefaultTransform, NewDefaultTransform, UpdatedComponents);
-							}
-							const bool bCanPropagateForOverridenComponent = !SelectedNodePtr->IsNative()
-								&& SelectedTemplate->HasAnyFlags(RF_ArchetypeObject);
-							if (bCanPropagateForOverridenComponent)
-							{
-								FComponentEditorUtils::PropagateTransformPropertyChangeAmongOverridenTemplates(SelectedTemplate, OldDefaultTransform, NewDefaultTransform, UpdatedComponents);
-							}
-
+							UBlueprint* PreviewBlueprint = UBlueprint::GetBlueprintFromClass(PreviewActor->GetClass());
 							if(PreviewBlueprint != NULL)
 							{
 								// Like PostEditMove(), but we only need to re-run construction scripts
@@ -474,34 +457,35 @@ bool FSCSEditorViewportClient::InputWidgetDelta( FViewport* Viewport, EAxisList:
 									TemplateComp->ConstraintInstance.CopyConstraintGeometryFrom(&ConstraintComp->ConstraintInstance);
 								}
 
-								// Get the Blueprint class default object
-								AActor* CDO = NULL;
-								if(PreviewBlueprint->GeneratedClass)
+								// Iterate over all the active archetype instances and propagate the change(s) to the matching component instance
+								TArray<UObject*> ArchetypeInstances;
+								if(SelectedTemplate->HasAnyFlags(RF_ArchetypeObject))
 								{
-									CDO = Cast<AActor>(PreviewBlueprint->GeneratedClass->ClassDefaultObject);
-								}
-								else if(PreviewBlueprint->ParentClass)
-								{
-									CDO = Cast<AActor>(PreviewBlueprint->ParentClass->ClassDefaultObject);
+									SelectedTemplate->GetArchetypeInstances(ArchetypeInstances);
+									for(int32 InstanceIndex = 0; InstanceIndex < ArchetypeInstances.Num(); ++InstanceIndex)
+									{
+										SceneComp = Cast<USceneComponent>(ArchetypeInstances[InstanceIndex]);
+										if(SceneComp && SceneComp->GetOwner() != PreviewActor)
+										{
+											FComponentEditorUtils::ApplyDefaultValueChange(SceneComp, SceneComp->RelativeLocation, OldRelativeLocation, SelectedTemplate->RelativeLocation);
+											FComponentEditorUtils::ApplyDefaultValueChange(SceneComp, SceneComp->RelativeRotation, OldRelativeRotation, SelectedTemplate->RelativeRotation);
+											FComponentEditorUtils::ApplyDefaultValueChange(SceneComp, SceneComp->RelativeScale3D,  OldRelativeScale3D,  SelectedTemplate->RelativeScale3D);
+										}
+									}
 								}
-
-								if(CDO != NULL)
+								else if(UObject* Outer = SelectedTemplate->GetOuter())
 								{
-									// Iterate over all the active archetype instances and propagate the change(s) to the matching component instance
-									TArray<UObject*> ArchetypeInstances;
-									CDO->GetArchetypeInstances(ArchetypeInstances);
+									Outer->GetArchetypeInstances(ArchetypeInstances);
 									for(int32 InstanceIndex = 0; InstanceIndex < ArchetypeInstances.Num(); ++InstanceIndex)
 									{
-										AActor* ArchetypeInstance = Cast<AActor>(ArchetypeInstances[InstanceIndex]);
-										if(ArchetypeInstance != NULL)
+										if(ArchetypeInstances[InstanceIndex] != PreviewActor)
 										{
-											const bool bIsProcessingPreviewActor = (ArchetypeInstance == PreviewActor);
-											SceneComp = Cast<USceneComponent>(SelectedNodePtr->FindComponentInstanceInActor(ArchetypeInstance));
-											if(!bIsProcessingPreviewActor && SceneComp != nullptr && !UpdatedComponents.Contains(SceneComp))
+											SceneComp = static_cast<USceneComponent*>(FindObjectWithOuter(ArchetypeInstances[InstanceIndex], SelectedTemplate->GetClass(), SelectedTemplate->GetFName()));
+											if(SceneComp)
 											{
-												FComponentEditorUtils::PropagateTransformPropertyChange(SceneComp, SceneComp->RelativeLocation, OldRelativeLocation, SelectedTemplate->RelativeLocation, UpdatedComponents);
-												FComponentEditorUtils::PropagateTransformPropertyChange(SceneComp, SceneComp->RelativeRotation, OldRelativeRotation, SelectedTemplate->RelativeRotation, UpdatedComponents);
-												FComponentEditorUtils::PropagateTransformPropertyChange(SceneComp, SceneComp->RelativeScale3D,  OldRelativeScale3D,  SelectedTemplate->RelativeScale3D,  UpdatedComponents);
+												FComponentEditorUtils::ApplyDefaultValueChange(SceneComp, SceneComp->RelativeLocation, OldRelativeLocation, SelectedTemplate->RelativeLocation);
+												FComponentEditorUtils::ApplyDefaultValueChange(SceneComp, SceneComp->RelativeRotation, OldRelativeRotation, SelectedTemplate->RelativeRotation);
+												FComponentEditorUtils::ApplyDefaultValueChange(SceneComp, SceneComp->RelativeScale3D,  OldRelativeScale3D,  SelectedTemplate->RelativeScale3D);
 											}
 										}
 									}
@@ -539,9 +523,13 @@ void FSCSEditorViewportClient::TrackingStopped()
 	{
 		// Re-run construction scripts if we haven't done so yet (so that the components in the preview actor can update their transforms)
 		AActor* PreviewActor = GetPreviewActor();
-		if(PreviewActor != NULL && PreviewBlueprint != NULL && !PreviewBlueprint->bRunConstructionScriptOnDrag)
+		if(PreviewActor != nullptr)
 		{
-			PreviewActor->RerunConstructionScripts();
+			UBlueprint* PreviewBlueprint = UBlueprint::GetBlueprintFromClass(PreviewActor->GetClass());
+			if(PreviewBlueprint != nullptr && !PreviewBlueprint->bRunConstructionScriptOnDrag)
+			{
+				PreviewActor->RerunConstructionScripts();
+			}
 		}
 
 		// End transaction
@@ -773,7 +761,7 @@ void FSCSEditorViewportClient::ToggleIsSimulateEnabled()
 	TSharedRef<SWidget> Inspector = BlueprintEditorPtr.Pin()->GetInspector();
 
 	// When simulate is enabled, we don't want to allow the user to modify the components
-	BlueprintEditorPtr.Pin()->UpdatePreviewActor(PreviewBlueprint, true);
+	BlueprintEditorPtr.Pin()->UpdatePreviewActor(BlueprintEditorPtr.Pin()->GetBlueprintObj(), true);
 
 	SCSEditor->SetEnabled(!bIsSimulateEnabled);
 	Inspector->SetEnabled(!bIsSimulateEnabled);
@@ -828,14 +816,15 @@ void FSCSEditorViewportClient::BeginTransaction(const FText& Description)
 	{
 		ScopedTransaction = new FScopedTransaction(Description);
 
-		if (PreviewBlueprint != nullptr)
-		{
-			FBlueprintEditorUtils::MarkBlueprintAsModified(PreviewBlueprint);
-		}
-
 		auto BlueprintEditor = BlueprintEditorPtr.Pin();
 		if (BlueprintEditor.IsValid())
 		{
+			UBlueprint* PreviewBlueprint = BlueprintEditor->GetBlueprintObj();
+			if (PreviewBlueprint != nullptr)
+			{
+				FBlueprintEditorUtils::MarkBlueprintAsModified(PreviewBlueprint);
+			}
+
 			TArray<FSCSEditorTreeNodePtrType> SelectedNodes = BlueprintEditor->GetSelectedSCSEditorTreeNodes();
 			for(auto SelectedSCSNodeIter(SelectedNodes.CreateIterator()); SelectedSCSNodeIter; ++SelectedSCSNodeIter)
 			{
@@ -848,7 +837,7 @@ void FSCSEditorViewportClient::BeginTransaction(const FText& Description)
 					}
 
 					// Modify both the component template and the instance in the preview actor (provided there is one)
-					UActorComponent* ComponentTemplate = Node->GetEditableComponentTemplate(BlueprintEditor->GetBlueprintObj());
+					UActorComponent* ComponentTemplate = Node->GetEditableComponentTemplate(PreviewBlueprint);
 					if (ComponentTemplate != nullptr)
 					{
 						ComponentTemplate->SetFlags(RF_Transactional);
diff --git a/Engine/Source/Editor/Kismet/Private/SCSEditorViewportClient.h b/Engine/Source/Editor/Kismet/Private/SCSEditorViewportClient.h
--- a/Engine/Source/Editor/Kismet/Private/SCSEditorViewportClient.h
+++ b/Engine/Source/Editor/Kismet/Private/SCSEditorViewportClient.h
@@ -131,9 +131,6 @@ private:
 	/** Weak reference to the editor hosting the viewport */
 	TWeakPtr<class FBlueprintEditor> BlueprintEditorPtr;
 
-	/** The Blueprint associated with the current preview */
-	UBlueprint* PreviewBlueprint;
-
 	/** The preview actor from last tick so we can tell if it changed. */
 	TWeakObjectPtr<AActor> LastPreviewActor;
 
diff --git a/Engine/Source/Editor/Kismet/Private/SSCSEditorViewport.cpp b/Engine/Source/Editor/Kismet/Private/SSCSEditorViewport.cpp
--- a/Engine/Source/Editor/Kismet/Private/SSCSEditorViewport.cpp
+++ b/Engine/Source/Editor/Kismet/Private/SSCSEditorViewport.cpp
@@ -310,6 +310,8 @@ void SSCSEditorViewport::RequestRefresh(bool bResetCamera, bool bRefreshNow)
 {
 	if(bRefreshNow)
 	{
+		Invalidate();
+
 		if(ViewportClient.IsValid())
 		{
 			ViewportClient->InvalidatePreview(bResetCamera);
@@ -362,6 +364,8 @@ EActiveTimerReturnType SSCSEditorViewport::DeferredUpdatePreview(double InCurren
 {
 	if (ViewportClient.IsValid())
 	{
+		Invalidate();
+
 		ViewportClient->InvalidatePreview(bResetCamera);
 	}
 
diff --git a/Engine/Source/Editor/UnrealEd/Private/Kismet2/ComponentEditorUtils.cpp b/Engine/Source/Editor/UnrealEd/Private/Kismet2/ComponentEditorUtils.cpp
--- a/Engine/Source/Editor/UnrealEd/Private/Kismet2/ComponentEditorUtils.cpp
+++ b/Engine/Source/Editor/UnrealEd/Private/Kismet2/ComponentEditorUtils.cpp
@@ -105,89 +105,6 @@ protected:
 	// FCustomizableTextObjectFactory (end)
 };
 
-USceneComponent* FComponentEditorUtils::GetSceneComponent( UObject* Object, UObject* SubObject /*= NULL*/ )
-{
-	if( Object )
-	{
-		AActor* Actor = Cast<AActor>( Object );
-		if( Actor )
-		{
-			if( SubObject )
-			{
-				if( SubObject->HasAnyFlags( RF_DefaultSubObject ) )
-				{
-					UObject* ClassDefaultObject = SubObject->GetOuter();
-					if( ClassDefaultObject )
-					{
-						for( TFieldIterator<UObjectProperty> ObjPropIt( ClassDefaultObject->GetClass() ); ObjPropIt; ++ObjPropIt )
-						{
-							UObjectProperty* ObjectProperty = *ObjPropIt;
-							if( SubObject == ObjectProperty->GetObjectPropertyValue_InContainer( ClassDefaultObject ) )
-							{
-								return CastChecked<USceneComponent>( ObjectProperty->GetObjectPropertyValue_InContainer( Actor ) );
-							}
-						}
-					}
-				}
-				else if( UBlueprint* Blueprint = Cast<UBlueprint>( SubObject->GetOuter() ) )
-				{
-					if( Blueprint->SimpleConstructionScript )
-					{
-						TArray<USCS_Node*> SCSNodes = Blueprint->SimpleConstructionScript->GetAllNodes();
-						for( int32 SCSNodeIndex = 0; SCSNodeIndex < SCSNodes.Num(); ++SCSNodeIndex )
-						{
-							USCS_Node* SCS_Node = SCSNodes[ SCSNodeIndex ];
-							if( SCS_Node && SCS_Node->ComponentTemplate == SubObject )
-							{
-								return SCS_Node->GetParentComponentTemplate(Blueprint);
-							}
-						}
-					}
-				}
-			}
-
-			return Actor->GetRootComponent();
-		}
-		else if( Object->IsA<USceneComponent>() )
-		{
-			return CastChecked<USceneComponent>( Object );
-		}
-	}
-
-	return NULL;
-}
-
-void FComponentEditorUtils::GetArchetypeInstances( UObject* Object, TArray<UObject*>& ArchetypeInstances )
-{
-	if (Object->HasAnyFlags(RF_ClassDefaultObject))
-	{
-		// Determine if the object is owned by a Blueprint
-		UBlueprint* Blueprint = Cast<UBlueprint>(Object->GetOuter());
-		if(Blueprint != NULL)
-		{
-			if(Blueprint->GeneratedClass != NULL && Blueprint->GeneratedClass->ClassDefaultObject != NULL)
-			{
-				// Collect all instances of the Blueprint
-				Blueprint->GeneratedClass->ClassDefaultObject->GetArchetypeInstances(ArchetypeInstances);
-			}
-		}
-		else
-		{
-			// Object is a default object, collect all instances.
-			Object->GetArchetypeInstances(ArchetypeInstances);
-		}
-	}
-	else if (Object->HasAnyFlags(RF_DefaultSubObject))
-	{
-		UObject* DefaultObject = Object->GetOuter();
-		if(DefaultObject != NULL && DefaultObject->HasAnyFlags(RF_ClassDefaultObject))
-		{
-			// Object is a default subobject, collect all instances of the default object that owns it.
-			DefaultObject->GetArchetypeInstances(ArchetypeInstances);
-		}
-	}
-}
-
 bool FComponentEditorUtils::IsValidVariableNameString(const UActorComponent* InComponent, const FString& InString)
 {
 	// First test to make sure the string is not empty and does not equate to the DefaultSceneRoot node name
@@ -661,66 +578,6 @@ bool FComponentEditorUtils::AttemptApplyMaterialToComponent(USceneComponent* Sce
 	return bResult;
 }
 
-void FComponentEditorUtils::PropagateTransformPropertyChangeAmongOverridenTemplates(
-	class USceneComponent* InSceneComponentTemplate,
-	const FTransformData& OldDefaultTransform,
-	const FTransformData& NewDefaultTransform,
-	TSet<class USceneComponent*>& UpdatedComponents)
-{
-	check(InSceneComponentTemplate != nullptr);
-	TArray<UObject*> ArchetypeInstances;
-	InSceneComponentTemplate->GetArchetypeInstances(ArchetypeInstances);
-	for (auto ArchetypeInstance : ArchetypeInstances)
-	{
-		USceneComponent* InstancedSceneComponent = Cast<USceneComponent>(ArchetypeInstance);
-		if (InstancedSceneComponent && InstancedSceneComponent->HasAnyFlags(RF_InheritableComponentTemplate) && !UpdatedComponents.Contains(InstancedSceneComponent))
-		{
-			PropagateTransformInner(InstancedSceneComponent, OldDefaultTransform, NewDefaultTransform, UpdatedComponents);
-		}
-	}
-}
-
-void FComponentEditorUtils::PropagateTransformPropertyChange(
-	class USceneComponent* InSceneComponentTemplate,
-	const FTransformData& OldDefaultTransform,
-	const FTransformData& NewDefaultTransform,
-	TSet<class USceneComponent*>& UpdatedComponents)
-{
-	check(InSceneComponentTemplate != nullptr);
-
-	TArray<UObject*> ArchetypeInstances;
-	FComponentEditorUtils::GetArchetypeInstances(InSceneComponentTemplate, ArchetypeInstances);
-	for(int32 InstanceIndex = 0; InstanceIndex < ArchetypeInstances.Num(); ++InstanceIndex)
-	{
-		USceneComponent* InstancedSceneComponent = FComponentEditorUtils::GetSceneComponent(ArchetypeInstances[InstanceIndex], InSceneComponentTemplate);
-		if(InstancedSceneComponent != nullptr && !UpdatedComponents.Contains(InstancedSceneComponent))
-		{
-			PropagateTransformInner(InstancedSceneComponent, OldDefaultTransform, NewDefaultTransform, UpdatedComponents);
-		}
-	}
-}
-
-void FComponentEditorUtils::PropagateTransformInner(class USceneComponent* InstancedSceneComponent, const FTransformData& OldDefaultTransform, const FTransformData& NewDefaultTransform, TSet<class USceneComponent*>& UpdatedComponents)
-{
-	static const UProperty* RelativeLocationProperty = FindFieldChecked<UProperty>(USceneComponent::StaticClass(), GET_MEMBER_NAME_CHECKED(USceneComponent, RelativeLocation));
-	if (RelativeLocationProperty != nullptr)
-	{
-		PropagateTransformPropertyChange(InstancedSceneComponent, RelativeLocationProperty, OldDefaultTransform.Trans, NewDefaultTransform.Trans, UpdatedComponents);
-	}
-
-	static const UProperty* RelativeRotationProperty = FindFieldChecked<UProperty>(USceneComponent::StaticClass(), GET_MEMBER_NAME_CHECKED(USceneComponent, RelativeRotation));
-	if (RelativeRotationProperty != nullptr)
-	{
-		PropagateTransformPropertyChange(InstancedSceneComponent, RelativeRotationProperty, OldDefaultTransform.Rot, NewDefaultTransform.Rot, UpdatedComponents);
-	}
-
-	static const UProperty* RelativeScale3DProperty = FindFieldChecked<UProperty>(USceneComponent::StaticClass(), GET_MEMBER_NAME_CHECKED(USceneComponent, RelativeScale3D));
-	if (RelativeScale3DProperty != nullptr)
-	{
-		PropagateTransformPropertyChange(InstancedSceneComponent, RelativeScale3DProperty, OldDefaultTransform.Scale, NewDefaultTransform.Scale, UpdatedComponents);
-	}
-}
-
 FName FComponentEditorUtils::FindVariableNameGivenComponentInstance(UActorComponent* ComponentInstance)
 {
 	check(ComponentInstance != nullptr);
diff --git a/Engine/Source/Editor/UnrealEd/Public/Kismet2/ComponentEditorUtils.h b/Engine/Source/Editor/UnrealEd/Public/Kismet2/ComponentEditorUtils.h
--- a/Engine/Source/Editor/UnrealEd/Public/Kismet2/ComponentEditorUtils.h
+++ b/Engine/Source/Editor/UnrealEd/Public/Kismet2/ComponentEditorUtils.h
@@ -8,11 +8,6 @@
 class UNREALED_API FComponentEditorUtils
 {
 public:
-	
-	static class USceneComponent* GetSceneComponent(class UObject* Object, class UObject* SubObject = NULL);
-
-	static void GetArchetypeInstances(class UObject* Object, TArray<class UObject*>& ArchetypeInstances);
-
 	/** Test whether or not the given string is a valid variable name string for the given component instance */
 	static bool IsValidVariableNameString(const UActorComponent* InComponent, const FString& InString);
 
@@ -95,31 +90,43 @@ public:
 	 */
 	static bool AttemptApplyMaterialToComponent( USceneComponent* SceneComponent, UMaterialInterface* MaterialToApply, int32 OptionalMaterialSlot = -1 );
 
-	struct FTransformData
-	{
-		FTransformData(const class USceneComponent& Component)
-			: Trans(Component.RelativeLocation)
-			, Rot(Component.RelativeRotation)
-			, Scale(Component.RelativeScale3D)
-		{}
-
-		FVector Trans; 
-		FRotator Rot;
-		FVector Scale;
-	};
-
 	/** Potentially transforms the delta to be applied to a component into the appropriate space */
 	static void AdjustComponentDelta(USceneComponent* Component, FVector& Drag, FRotator& Rotation);
 
-	// Given a template, propagates a default transform change to all instances of the template
-	static void PropagateTransformPropertyChange(class USceneComponent* InSceneComponentTemplate, const FTransformData& OldDefaultTransform, const FTransformData& NewDefaultTransform, TSet<class USceneComponent*>& UpdatedComponents);
-
-	static void PropagateTransformPropertyChangeAmongOverridenTemplates(class USceneComponent* InSceneComponentTemplate, const FTransformData& OldDefaultTransform, const FTransformData& NewDefaultTransform, TSet<class USceneComponent*>& UpdatedComponents);
-
+	// Given a template and a property, propagates a default value change to all instances (only if applicable)
+	template<typename T>
+	static void PropagateDefaultValueChange(class USceneComponent* InSceneComponentTemplate, const class UProperty* InProperty, const T& OldDefaultValue, const T& NewDefaultValue, TSet<class USceneComponent*>& UpdatedInstances)
+	{
+		TArray<UObject*> ArchetypeInstances;
+		if(InSceneComponentTemplate->HasAnyFlags(RF_ArchetypeObject))
+		{
+			InSceneComponentTemplate->GetArchetypeInstances(ArchetypeInstances);
+			for(int32 InstanceIndex = 0; InstanceIndex < ArchetypeInstances.Num(); ++InstanceIndex)
+			{
+				USceneComponent* InstancedSceneComponent = static_cast<USceneComponent*>(ArchetypeInstances[InstanceIndex]);
+				if(InstancedSceneComponent != nullptr && !UpdatedInstances.Contains(InstancedSceneComponent) && ApplyDefaultValueChange(InstancedSceneComponent, InProperty, OldDefaultValue, NewDefaultValue))
+				{
+					UpdatedInstances.Add(InstancedSceneComponent);
+				}
+			}
+		}
+		else if(UObject* Outer = InSceneComponentTemplate->GetOuter())
+		{
+			Outer->GetArchetypeInstances(ArchetypeInstances);
+			for(int32 InstanceIndex = 0; InstanceIndex < ArchetypeInstances.Num(); ++InstanceIndex)
+			{
+				USceneComponent* InstancedSceneComponent = static_cast<USceneComponent*>(FindObjectWithOuter(ArchetypeInstances[InstanceIndex], InSceneComponentTemplate->GetClass(), InSceneComponentTemplate->GetFName()));
+				if(InstancedSceneComponent != nullptr && !UpdatedInstances.Contains(InstancedSceneComponent) && ApplyDefaultValueChange(InstancedSceneComponent, InProperty, OldDefaultValue, NewDefaultValue))
+				{
+					UpdatedInstances.Add(InstancedSceneComponent);
+				}
+			}
+		}
+	}
 
-	// Given an instance of a template and a property, propagates a default value change to the instance (only if applicable)
+	// Given an instance of a template and a property, set a default value change to the instance (only if applicable)
 	template<typename T>
-	static void PropagateTransformPropertyChange(class USceneComponent* InSceneComponent, const class UProperty* InProperty, const T& OldDefaultValue, const T& NewDefaultValue, TSet<class USceneComponent*>& UpdatedComponents)
+	static bool ApplyDefaultValueChange(class USceneComponent* InSceneComponent, const class UProperty* InProperty, const T& OldDefaultValue, const T& NewDefaultValue)
 	{
 		check(InProperty != nullptr);
 		check(InSceneComponent != nullptr);
@@ -127,13 +134,15 @@ public:
 		T* CurrentValue = InProperty->ContainerPtrToValuePtr<T>(InSceneComponent);
 		if(CurrentValue != nullptr)
 		{
-			PropagateTransformPropertyChange(InSceneComponent, *CurrentValue, OldDefaultValue, NewDefaultValue, UpdatedComponents);
+			return ApplyDefaultValueChange(InSceneComponent, *CurrentValue, OldDefaultValue, NewDefaultValue);
 		}
+
+		return false;
 	}
 
 	// Given an instance of a template and a current value, propagates a default value change to the instance (only if applicable)
 	template<typename T>
-	static void PropagateTransformPropertyChange(class USceneComponent* InSceneComponent, T& CurrentValue, const T& OldDefaultValue, const T& NewDefaultValue, TSet<class USceneComponent*>& UpdatedComponents)
+	static bool ApplyDefaultValueChange(class USceneComponent* InSceneComponent, T& CurrentValue, const T& OldDefaultValue, const T& NewDefaultValue)
 	{
 		check(InSceneComponent != nullptr);
 
@@ -161,15 +170,14 @@ public:
 			// Re-register the component with the scene so that transforms are updated for display
 			InSceneComponent->ReregisterComponent();
 
-			// Add the component to the set of updated instances
-			UpdatedComponents.Add(InSceneComponent);
+			return true;
 		}
+
+		return false;
 	}
 
 	// Try to find the correct variable name for a given native component template or instance (which can have a mismatch)
 	static FName FindVariableNameGivenComponentInstance(UActorComponent* ComponentInstance);
-private:
-	static void PropagateTransformInner(class USceneComponent* InstancedSceneComponent, const FTransformData& OldDefaultTransform, const FTransformData& NewDefaultTransform, TSet<class USceneComponent*>& UpdatedComponents);
-	
+
 	static USceneComponent* FindClosestParentInList(UActorComponent* ChildComponent, const TArray<UActorComponent*>& ComponentList);
 };
\ No newline at end of file
