[INTEGRATE] Change 2438036 by [REDACTED_PERSON][REDACTED_HANDLE] on 2015/02/09 10:28:56

	Fix BP editor preview not updating on mesh material changes for SCS component templates.

	change summary (FComponentMaterialCategory):
	- cache the NotifyHook for the details panel in material details customization as that's what the BP editor hooks into for property changes (similar to e.g. component transform details and "standard" property edits)
	- switched the "Materials" property name to "OverrideMaterials" to match the updated property name in UMeshComponent
	- invoke the NotifyHook on pre/post edit change so that the BP editor receives those events for material changes

diff --git a/Engine/Source/Editor/DetailCustomizations/Private/ComponentMaterialCategory.cpp b/Engine/Source/Editor/DetailCustomizations/Private/ComponentMaterialCategory.cpp
--- a/Engine/Source/Editor/DetailCustomizations/Private/ComponentMaterialCategory.cpp
+++ b/Engine/Source/Editor/DetailCustomizations/Private/ComponentMaterialCategory.cpp
@@ -162,11 +162,14 @@ private:
 
 FComponentMaterialCategory::FComponentMaterialCategory( TArray< TWeakObjectPtr<USceneComponent> >& InSelectedComponents )
 	: SelectedComponents( InSelectedComponents )
+	, NotifyHook( nullptr )
 {
 }
 
 void FComponentMaterialCategory::Create( IDetailLayoutBuilder& DetailBuilder )
 {
+	NotifyHook = DetailBuilder.GetPropertyUtilities()->GetNotifyHook();
+
 	FMaterialListDelegates MaterialListDelegates;
 	MaterialListDelegates.OnGetMaterials.BindSP( this, &FComponentMaterialCategory::OnGetMaterialsForView );
 	MaterialListDelegates.OnMaterialChanged.BindSP( this, &FComponentMaterialCategory::OnMaterialChanged );
@@ -273,7 +276,7 @@ void FComponentMaterialCategory::OnMaterialChanged( UMaterialInterface* NewMater
 				UObject* EditChangeObject = CurrentComponent;
 				if( CurrentComponent->IsA( UMeshComponent::StaticClass() ) )
 				{
-					MaterialProperty = FindField<UProperty>( UMeshComponent::StaticClass(), "Materials" );
+					MaterialProperty = FindField<UProperty>( UMeshComponent::StaticClass(), "OverrideMaterials" );
 				}
 				else if( CurrentComponent->IsA( UDecalComponent::StaticClass() ) )
 				{
@@ -296,10 +299,20 @@ void FComponentMaterialCategory::OnMaterialChanged( UMaterialInterface* NewMater
 
 				EditChangeObject->PreEditChange( MaterialProperty );
 
+				if( NotifyHook && MaterialProperty )
+				{
+					NotifyHook->NotifyPreChange( MaterialProperty );
+				}
+
 				It.SwapMaterial( NewMaterial );
 
 				FPropertyChangedEvent PropertyChangedEvent( MaterialProperty );
 				EditChangeObject->PostEditChangeProperty( PropertyChangedEvent );
+
+				if( NotifyHook && MaterialProperty )
+				{
+					NotifyHook->NotifyPostChange( PropertyChangedEvent, MaterialProperty );
+				}
 			}
 		}
 	}
diff --git a/Engine/Source/Editor/DetailCustomizations/Private/ComponentMaterialCategory.h b/Engine/Source/Editor/DetailCustomizations/Private/ComponentMaterialCategory.h
--- a/Engine/Source/Editor/DetailCustomizations/Private/ComponentMaterialCategory.h
+++ b/Engine/Source/Editor/DetailCustomizations/Private/ComponentMaterialCategory.h
@@ -43,4 +43,7 @@ private:
 private:
 	/** Reference to the list of selected actors */
 	TArray< TWeakObjectPtr<USceneComponent> > SelectedComponents;
+
+	/** Notify hook to use */
+	FNotifyHook* NotifyHook;
 };
\ No newline at end of file
