[INTEGRATE] Change 2438600 by [REDACTED_PERSON][REDACTED_HANDLE] on 2015/02/09 15:15:59

	Fix selected components disappearing while spinning a rotation spin box

diff --git a/Engine/Source/Editor/DetailCustomizations/Private/ComponentTransformDetails.cpp b/Engine/Source/Editor/DetailCustomizations/Private/ComponentTransformDetails.cpp
--- a/Engine/Source/Editor/DetailCustomizations/Private/ComponentTransformDetails.cpp
+++ b/Engine/Source/Editor/DetailCustomizations/Private/ComponentTransformDetails.cpp
@@ -1204,10 +1204,15 @@ void FComponentTransformDetails::OnSetRotation( float NewValue, bool bCommitted,
 							}
 						}
 
-						AActor* ObjectAsActor = Cast<AActor>( Object );
-						if( ObjectAsActor && !bIsEditingTemplateObject )
+						AActor* EditedActor = Cast<AActor>( Object );
+						if( !EditedActor && SceneComponent )
 						{
-							ObjectAsActor->ReregisterAllComponents();
+							EditedActor = SceneComponent->GetOwner();
+						}
+
+						if( EditedActor && !bIsEditingTemplateObject )
+						{
+							EditedActor->ReregisterAllComponents();
 						}
 
 						// If it's a template, propagate the change out to any current instances of the object
