[INTEGRATE] Change 2438677 by [REDACTED_PERSON][REDACTED_HANDLE] on 2015/02/09 15:54:29

	CIS fix (non-unity)

diff --git a/Engine/Source/Editor/DetailCustomizations/Private/ComponentMaterialCategory.cpp b/Engine/Source/Editor/DetailCustomizations/Private/ComponentMaterialCategory.cpp
--- a/Engine/Source/Editor/DetailCustomizations/Private/ComponentMaterialCategory.cpp
+++ b/Engine/Source/Editor/DetailCustomizations/Private/ComponentMaterialCategory.cpp
@@ -4,6 +4,7 @@
 #include "ComponentMaterialCategory.h"
 #include "AssetThumbnail.h"
 #include "ActorEditorUtils.h"
+#include "IPropertyUtilities.h"
 
 #include "LandscapeProxy.h"
 #include "LandscapeComponent.h"
