Refactor LockedAxis so you can have better control for 6DOF.
Move around some properties in primitive component customization to make it easier to find things.
BodyMass override now uses the default computed to make it more clear you're overriding a default

diff --git a/Engine/Config/BaseEngine.ini b/Engine/Config/BaseEngine.ini
--- a/Engine/Config/BaseEngine.ini
+++ b/Engine/Config/BaseEngine.ini
@@ -287,6 +287,24 @@ bUseStreamingPause=false
 +EnumRedirects=(EnumName="ENavLinkDirection",OldEnumEntry="ESmartNavLinkDir::BothWays",NewEnumEntry="ENavLinkDirection::BothWays")
 +K2FieldRedirects=(OldFieldName="NavigationPath.GetPathLenght",NewFieldName="NavigationPath.GetPathLength")
 
+;Locked axis redirects
++ActiveClassRedirects=(OldClassName="ELockedAxis", NewClassName="EDOFMode")
++EnumRedirects=(EnumName="EDOFMode",OldEnumEntry="EDOFMode::Default",NewEnumEntry="EDOFMode::Default")
++EnumRedirects=(EnumName="EDOFMode",OldEnumEntry="EDOFMode::X",NewEnumEntry="EDOFMode::YZPlane")
++EnumRedirects=(EnumName="EDOFMode",OldEnumEntry="EDOFMode::Y",NewEnumEntry="EDOFMode::XZPlane")
++EnumRedirects=(EnumName="EDOFMode",OldEnumEntry="EDOFMode::Z",NewEnumEntry="EDOFMode::XYPlane")
++EnumRedirects=(EnumName="EDOFMode",OldEnumEntry="EDOFMode::Custom",NewEnumEntry="EDOFMode::CustomPlane")
++EnumRedirects=(EnumName="EDOFMode",OldEnumEntry="EDOFMode::None",NewEnumEntry="EDOFMode::None")
++TaggedPropertyRedirects=(ClassName="BodyInstance",OldPropertyName="LockedAxisMode",NewPropertyName="DOFMode")
++TaggedPropertyRedirects=(ClassName="BodyInstance",OldPropertyName="CustomLockedAxis",NewPropertyName="CustomDOFPlaneNormal")
+
++EnumRedirects=(EnumName="EDOFMode",OldEnumEntry="ELockedAxis::Default",NewEnumEntry="EDOFMode::Default")
++EnumRedirects=(EnumName="EDOFMode",OldEnumEntry="ELockedAxis::X",NewEnumEntry="EDOFMode::YZPlane")
++EnumRedirects=(EnumName="EDOFMode",OldEnumEntry="ELockedAxis::Y",NewEnumEntry="EDOFMode::XZPlane")
++EnumRedirects=(EnumName="EDOFMode",OldEnumEntry="ELockedAxis::Z",NewEnumEntry="EDOFMode::XYPlane")
++EnumRedirects=(EnumName="EDOFMode",OldEnumEntry="ELockedAxis::Custom",NewEnumEntry="EDOFMode::CustomPlane")
++EnumRedirects=(EnumName="EDOFMode",OldEnumEntry="ELockedAxis::None",NewEnumEntry="EDOFMode::None")
+
 ; mesh component
 +TaggedPropertyRedirects=(ClassName="MeshComponent",OldPropertyName="Materials",NewPropertyName="OverrideMaterials")
 
diff --git a/Engine/Source/Editor/DetailCustomizations/Private/BodySetupDetails.cpp b/Engine/Source/Editor/DetailCustomizations/Private/BodySetupDetails.cpp
--- a/Engine/Source/Editor/DetailCustomizations/Private/BodySetupDetails.cpp
+++ b/Engine/Source/Editor/DetailCustomizations/Private/BodySetupDetails.cpp
@@ -41,7 +41,18 @@ void FBodySetupDetails::CustomizeDetails( IDetailLayoutBuilder& DetailBuilder )
 			{
 				TSharedPtr<IPropertyHandle> ChildProperty = BodyInstanceHandler->GetChildHandle(ChildIndex);
 				FString Category = FObjectEditorUtils::GetCategory(ChildProperty->GetProperty());
-				if (ChildProperty->GetProperty()->GetFName() == GET_MEMBER_NAME_CHECKED(FBodyInstance, bSimulatePhysics) || ChildProperty->GetProperty()->GetFName() == GET_MEMBER_NAME_CHECKED(FBodyInstance, bAutoWeld))
+				if (ChildProperty->GetProperty()->GetFName() == GET_MEMBER_NAME_CHECKED(FBodyInstance, bSimulatePhysics)
+					|| ChildProperty->GetProperty()->GetFName() == GET_MEMBER_NAME_CHECKED(FBodyInstance, bAutoWeld)
+					|| ChildProperty->GetProperty()->GetFName() == GET_MEMBER_NAME_CHECKED(FBodyInstance, DOFMode)
+					|| ChildProperty->GetProperty()->GetFName() == GET_MEMBER_NAME_CHECKED(FBodyInstance, CustomDOFPlaneNormal)
+					|| ChildProperty->GetProperty()->GetFName() == GET_MEMBER_NAME_CHECKED(FBodyInstance, bLockTranslation)
+					|| ChildProperty->GetProperty()->GetFName() == GET_MEMBER_NAME_CHECKED(FBodyInstance, bLockXTranslation)
+					|| ChildProperty->GetProperty()->GetFName() == GET_MEMBER_NAME_CHECKED(FBodyInstance, bLockYTranslation)
+					|| ChildProperty->GetProperty()->GetFName() == GET_MEMBER_NAME_CHECKED(FBodyInstance, bLockZTranslation)
+					|| ChildProperty->GetProperty()->GetFName() == GET_MEMBER_NAME_CHECKED(FBodyInstance, bLockRotation)
+					|| ChildProperty->GetProperty()->GetFName() == GET_MEMBER_NAME_CHECKED(FBodyInstance, bLockXRotation)
+					|| ChildProperty->GetProperty()->GetFName() == GET_MEMBER_NAME_CHECKED(FBodyInstance, bLockYRotation)
+					|| ChildProperty->GetProperty()->GetFName() == GET_MEMBER_NAME_CHECKED(FBodyInstance, bLockZRotation))
 				{
 					// skip bSimulatePhysics
 					// this is because we don't want bSimulatePhysics to show up 
diff --git a/Engine/Source/Editor/DetailCustomizations/Private/PrimitiveComponentDetails.cpp b/Engine/Source/Editor/DetailCustomizations/Private/PrimitiveComponentDetails.cpp
--- a/Engine/Source/Editor/DetailCustomizations/Private/PrimitiveComponentDetails.cpp
+++ b/Engine/Source/Editor/DetailCustomizations/Private/PrimitiveComponentDetails.cpp
@@ -12,6 +12,8 @@
 #include "EditorCategoryUtils.h"
 #include "PhysicsEngine/PhysicsSettings.h"
 #include "ComponentMaterialCategory.h"
+#include "SCheckBox.h"
+#include "SNumericEntryBox.h"
 
 #define LOCTEXT_NAMESPACE "PrimitiveComponentDetails"
 
@@ -82,37 +84,22 @@ bool FPrimitiveComponentDetails::IsUseAsyncEditable() const
 	return bEnableUseAsyncScene;
 }
 
-EVisibility FPrimitiveComponentDetails::IsCustomLockedAxisSelected() const
+EVisibility FPrimitiveComponentDetails::IsDOFMode(EDOFMode::Type Mode) const
 {
 	bool bVisible = false;
-	if (LockedAxisProperty.IsValid())
+	if (DOFModeProperty.IsValid())
 	{
-		uint8 LockedAxis;
-		if (LockedAxisProperty->GetValue(LockedAxis) == FPropertyAccess::Success)
+		uint8 CurrentMode;
+		if (DOFModeProperty->GetValue(CurrentMode) == FPropertyAccess::Success)
 		{
-			bVisible = static_cast<ELockedAxis::Type>(LockedAxis) == ELockedAxis::Custom;
+			EDOFMode::Type PropertyDOF = FBodyInstance::ResolveDOFMode(static_cast<EDOFMode::Type>(CurrentMode));
+			bVisible = PropertyDOF == Mode;
 		}
 	}
 
 	return bVisible ? EVisibility::Visible : EVisibility::Collapsed;
 }
 
-EVisibility FPrimitiveComponentDetails::IsLockAxisEnabled() const
-{
-	bool bVisible = false;
-	if (LockedAxisProperty.IsValid())
-	{
-		uint8 LockedAxis;
-		if (LockedAxisProperty->GetValue(LockedAxis) == FPropertyAccess::Success)
-		{
-			bVisible = static_cast<ELockedAxis::Type>(LockedAxis) != ELockedAxis::Default;
-		}
-	}
-
-	return bVisible ? EVisibility::Visible : EVisibility::Collapsed;
-}
-
-
 bool FPrimitiveComponentDetails::IsAutoWeldEditable() const
 {
 	for (int32 i = 0; i < ObjectsCustomized.Num(); ++i)
@@ -158,162 +145,338 @@ EVisibility FPrimitiveComponentDetails::IsMassVisible(bool bOverrideMass) const
 	}
 }
 
-void FPrimitiveComponentDetails::CustomizeDetails( IDetailLayoutBuilder& DetailBuilder )
+void FPrimitiveComponentDetails::AddPhysicsCategory(IDetailLayoutBuilder& DetailBuilder)
 {
-	// Get the objects being customized so we can enable/disable editing of 'Simulate Physics'
-	DetailBuilder.GetObjectsBeingCustomized(ObjectsCustomized);
-
-	// See if we are hiding Physics category
-	TArray<FString> HideCategories;
-	FEditorCategoryUtils::GetClassHideCategories(DetailBuilder.GetDetailsView().GetBaseClass(), HideCategories);
-
-	if(!HideCategories.Contains("Materials"))
-	{
-		AddMaterialCategory(DetailBuilder);
-	}
-
-	TSharedRef<IPropertyHandle> MobilityHandle = DetailBuilder.GetProperty(GET_MEMBER_NAME_CHECKED(UPrimitiveComponent, Mobility), USceneComponent::StaticClass());
-	MobilityHandle->SetToolTipText(LOCTEXT("PrimitiveMobilityTooltip", "Mobility for primitive components controls how they can be modified in game and therefore how they interact with lighting and physics.\n● A movable primitive component can be changed in game, but requires dynamic lighting and shadowing from lights which have a large performance cost.\n● A static primitive component can't be changed in game, but can have its lighting baked, which allows rendering to be very efficient."));
-
-	if ( DetailBuilder.GetProperty(GET_MEMBER_NAME_CHECKED(UPrimitiveComponent, BodyInstance))->IsValidHandle() )
+	if (DetailBuilder.GetProperty(GET_MEMBER_NAME_CHECKED(UPrimitiveComponent, BodyInstance))->IsValidHandle())
 	{
 		TSharedPtr<IPropertyHandle> BodyInstanceHandler = DetailBuilder.GetProperty(GET_MEMBER_NAME_CHECKED(UPrimitiveComponent, BodyInstance));
 		uint32 NumChildren = 0;
 		BodyInstanceHandler->GetNumChildren(NumChildren);
 
 
-		if (!HideCategories.Contains(TEXT("Physics")))
-		{
-			IDetailCategoryBuilder& PhysicsCategory = DetailBuilder.EditCategory("Physics");
+		IDetailCategoryBuilder& PhysicsCategory = DetailBuilder.EditCategory("Physics");
 
+		bool bDisplayMass = true;
+		bool bDisplayMassOverride = true;
+		bool bDisplayConstraints = true;
 
+		for (int32 i = 0; i < ObjectsCustomized.Num(); ++i)
+		{
+			if (ObjectsCustomized[i].IsValid() && ObjectsCustomized[i]->IsA(UDestructibleComponent::StaticClass()))
+			{
+				bDisplayMass = false;
+				bDisplayMassOverride = false;
+				bDisplayConstraints = false;
+			}
 
-			bool bDisplayMass = true;
-			bool bDisplayMassOverride = true;
+			if (ObjectsCustomized[i].IsValid() && ObjectsCustomized[i]->IsA(USkeletalMeshComponent::StaticClass()))
+			{
+				bDisplayMassOverride = false;
+			}
+		}
+
+		bool bConstraintsUIExists = false;
+		bool bMassUIExists = false;
 
-			for (int32 i = 0; i < ObjectsCustomized.Num(); ++i)
+		// add all physics properties now - after adding mass
+		for (uint32 ChildIndex = 0; ChildIndex < NumChildren; ++ChildIndex)
+		{
+			TSharedPtr<IPropertyHandle> ChildProperty = BodyInstanceHandler->GetChildHandle(ChildIndex);
+			FString Category = FObjectEditorUtils::GetCategory(ChildProperty->GetProperty());
+			FName PropName = ChildProperty->GetProperty()->GetFName();
+			if (Category == TEXT("Physics"))
 			{
-				if (ObjectsCustomized[i].IsValid() && ObjectsCustomized[i]->IsA(UDestructibleComponent::StaticClass()))
+				// Only permit modifying bSimulatePhysics when the body has some geometry.
+				if (PropName == GET_MEMBER_NAME_CHECKED(FBodyInstance, bSimulatePhysics))
 				{
-					bDisplayMass = false;
-					bDisplayMassOverride = false;
+					PhysicsCategory.AddProperty(ChildProperty).EditCondition(TAttribute<bool>(this, &FPrimitiveComponentDetails::IsSimulatePhysicsEditable), NULL);
 				}
-
-				if (ObjectsCustomized[i].IsValid() && ObjectsCustomized[i]->IsA(USkeletalMeshComponent::StaticClass()))
+				else if (PropName == FName("bUseAsyncScene")) // Can't use GET_MEMBER_NAME_CHECKED because protected
 				{
-					bDisplayMassOverride = false;
+					//we only enable bUseAsyncScene if the project uses an AsyncScene
+					PhysicsCategory.AddProperty(ChildProperty).EditCondition(TAttribute<bool>(this, &FPrimitiveComponentDetails::IsUseAsyncEditable), NULL);
 				}
-			}
-
-			// add all physics properties now - after adding mass
-			for (uint32 ChildIndex = 0; ChildIndex < NumChildren; ++ChildIndex)
-			{
-				TSharedPtr<IPropertyHandle> ChildProperty = BodyInstanceHandler->GetChildHandle(ChildIndex);
-				FString Category = FObjectEditorUtils::GetCategory(ChildProperty->GetProperty());
-				FName PropName = ChildProperty->GetProperty()->GetFName();
-				if (Category == TEXT("Physics"))
+				else if (PropName == GET_MEMBER_NAME_CHECKED(FBodyInstance, bLockXTranslation)
+					|| PropName == GET_MEMBER_NAME_CHECKED(FBodyInstance, bLockYTranslation)
+					|| PropName == GET_MEMBER_NAME_CHECKED(FBodyInstance, bLockZTranslation)
+					|| PropName == GET_MEMBER_NAME_CHECKED(FBodyInstance, bLockXRotation)
+					|| PropName == GET_MEMBER_NAME_CHECKED(FBodyInstance, bLockYRotation)
+					|| PropName == GET_MEMBER_NAME_CHECKED(FBodyInstance, bLockZRotation)
+					|| PropName == GET_MEMBER_NAME_CHECKED(FBodyInstance, DOFMode)
+					|| PropName == GET_MEMBER_NAME_CHECKED(FBodyInstance, CustomDOFPlaneNormal)
+					|| PropName == GET_MEMBER_NAME_CHECKED(FBodyInstance, bLockTranslation)
+					|| PropName == GET_MEMBER_NAME_CHECKED(FBodyInstance, bLockRotation))
 				{
-					// Only permit modifying bSimulatePhysics when the body has some geometry.
-					if (PropName == GET_MEMBER_NAME_CHECKED(FBodyInstance, bSimulatePhysics))
-					{
-						PhysicsCategory.AddProperty(ChildProperty).EditCondition(TAttribute<bool>(this, &FPrimitiveComponentDetails::IsSimulatePhysicsEditable), NULL);
-					}
-					else if (PropName == FName("bUseAsyncScene")) // Can't use GET_MEMBER_NAME_CHECKED because protected
-					{
-						//we only enable bUseAsyncScene if the project uses an AsyncScene
-						PhysicsCategory.AddProperty(ChildProperty).EditCondition(TAttribute<bool>(this, &FPrimitiveComponentDetails::IsUseAsyncEditable), NULL);
-					}
-					else if (PropName == GET_MEMBER_NAME_CHECKED(FBodyInstance, LockedAxisMode))
-					{
-						LockedAxisProperty = ChildProperty;
-						PhysicsCategory.AddProperty(ChildProperty);
-					}
-					else if (PropName == GET_MEMBER_NAME_CHECKED(FBodyInstance, CustomLockedAxis))
+					// We treat all of these constraint parameters as a big group
+					if (bConstraintsUIExists == false && bDisplayConstraints)
 					{
-						//we only enable bUseAsyncScene if the project uses an AsyncScene
-						PhysicsCategory.AddProperty(ChildProperty).Visibility(TAttribute<EVisibility>(this, &FPrimitiveComponentDetails::IsCustomLockedAxisSelected));
-					}
-					else if (PropName == GET_MEMBER_NAME_CHECKED(FBodyInstance, bLockTranslation) || PropName == GET_MEMBER_NAME_CHECKED(FBodyInstance, bLockRotation))
-					{
-						// Only show the lock translation/rotation options if some lock axis is selected
-						PhysicsCategory.AddProperty(ChildProperty).Visibility(TAttribute<EVisibility>(this, &FPrimitiveComponentDetails::IsLockAxisEnabled));
-					}
-					else if (PropName == GET_MEMBER_NAME_CHECKED(FBodyInstance, bAutoWeld))
-					{
-						PhysicsCategory.AddProperty(ChildProperty).Visibility(TAttribute<EVisibility>(this, &FPrimitiveComponentDetails::IsAutoWeldVisible))
-																  .EditCondition(TAttribute<bool>(this, &FPrimitiveComponentDetails::IsAutoWeldEditable), NULL);
-					}
-					else if (PropName == GET_MEMBER_NAME_CHECKED(FBodyInstance, bOverrideMass))
-					{
-						if (bDisplayMassOverride)
-						{
-							PhysicsCategory.AddProperty(ChildProperty);
-						}
-					}
-					else if (PropName == GET_MEMBER_NAME_CHECKED(FBodyInstance, MassInKg))
-					{
-						
-						if (bDisplayMass)
-						{
-							PhysicsCategory.AddCustomRow(LOCTEXT("Mass", "Mass"), false)
-								.IsEnabled(TAttribute<bool>(this, &FPrimitiveComponentDetails::IsBodyMassEnabled))
-								.NameContent()
+						const float XYZPadding = 5.f;
+						bConstraintsUIExists = true;
+						DOFModeProperty = BodyInstanceHandler->GetChildHandle(TEXT("DOFMode"));
+
+						IDetailGroup& ConstraintsGroup = PhysicsCategory.AddGroup(TEXT("ConstraintsGroup"), LOCTEXT("Constraints", "Constraints"));
+
+						ConstraintsGroup.AddWidgetRow()
+						.Visibility(TAttribute<EVisibility>::Create(TAttribute<EVisibility>::FGetter::CreateSP(this, &FPrimitiveComponentDetails::IsDOFMode, EDOFMode::SixDOF)))
+						.NameContent()
+						[
+							SNew(STextBlock)
+							.Text(LOCTEXT("LockPositionLabel", "Lock Position"))
+							.ToolTipText(LOCTEXT("LockPositionTooltip", "Locks movement along the specified axis"))
+						]
+						.ValueContent()
+						[
+							SNew(SHorizontalBox)
+							+ SHorizontalBox::Slot()
+							.Padding(0.f, 0.f, XYZPadding, 0.f)
+							.AutoWidth()
+							[
+								SNew(SHorizontalBox)
+								+ SHorizontalBox::Slot()
+								.AutoWidth()
 								[
-									ChildProperty->CreatePropertyNameWidget()
+									BodyInstanceHandler->GetChildHandle(TEXT("bLockXTranslation"))->CreatePropertyNameWidget()
 								]
-							.ValueContent()
+								+ SHorizontalBox::Slot()
+								.AutoWidth()
 								[
-									SNew(SVerticalBox)
-									+ SVerticalBox::Slot()
-									.AutoHeight()
+									BodyInstanceHandler->GetChildHandle(TEXT("bLockXTranslation"))->CreatePropertyValueWidget()
+								]
+							]
+
+							+ SHorizontalBox::Slot()
+							.Padding(0.f, 0.f, XYZPadding, 0.f)
+							.AutoWidth()
+							[
+								SNew(SHorizontalBox)
+								+ SHorizontalBox::Slot()
+								.AutoWidth()
+								[
+									BodyInstanceHandler->GetChildHandle(TEXT("bLockYTranslation"))->CreatePropertyNameWidget()
+								]
+								+ SHorizontalBox::Slot()
+									.AutoWidth()
+									[
+										BodyInstanceHandler->GetChildHandle(TEXT("bLockYTranslation"))->CreatePropertyValueWidget()
+									]
+							]
+
+							+ SHorizontalBox::Slot()
+							.Padding(0.f, 0.f, XYZPadding, 0.f)
+							.AutoWidth()
+							[
+								SNew(SHorizontalBox)
+								+ SHorizontalBox::Slot()
+								.AutoWidth()
+								[
+									BodyInstanceHandler->GetChildHandle(TEXT("bLockZTranslation"))->CreatePropertyNameWidget()
+								]
+								+ SHorizontalBox::Slot()
+								.AutoWidth()
+								[
+									BodyInstanceHandler->GetChildHandle(TEXT("bLockZTranslation"))->CreatePropertyValueWidget()
+								]
+							]
+						];
+
+						ConstraintsGroup.AddWidgetRow()
+						.Visibility(TAttribute<EVisibility>::Create(TAttribute<EVisibility>::FGetter::CreateSP(this, &FPrimitiveComponentDetails::IsDOFMode, EDOFMode::SixDOF)))
+						.NameContent()
+						[
+							SNew(STextBlock)
+							.Text(LOCTEXT("LockRotationLabel", "Lock Rotation"))
+							.ToolTipText(LOCTEXT("LockRotationTooltip", "Locks rotation about the specified axis"))
+						]
+						.ValueContent()
+							[
+								SNew(SHorizontalBox)
+								+ SHorizontalBox::Slot()
+								.Padding(0.f, 0.f, XYZPadding, 0.f)
+								.AutoWidth()
+								[
+									SNew(SHorizontalBox)
+									+ SHorizontalBox::Slot()
+									.AutoWidth()
+									[
+										BodyInstanceHandler->GetChildHandle(TEXT("bLockXRotation"))->CreatePropertyNameWidget()
+									]
+									+ SHorizontalBox::Slot()
+									.AutoWidth()
 									[
-										SNew(SEditableTextBox)
-										.Text(this, &FPrimitiveComponentDetails::OnGetBodyMass)
-										.IsReadOnly(this, &FPrimitiveComponentDetails::IsBodyMassReadOnly)
-										.Font(IDetailLayoutBuilder::GetDetailFont())
-										.Visibility(this, &FPrimitiveComponentDetails::IsMassVisible, false)
+										BodyInstanceHandler->GetChildHandle(TEXT("bLockXRotation"))->CreatePropertyValueWidget()
 									]
+								]
 
-									+ SVerticalBox::Slot()
-									.AutoHeight()
+								+ SHorizontalBox::Slot()
+								.Padding(0.f, 0.f, XYZPadding, 0.f)
+								.AutoWidth()
+								[
+									SNew(SHorizontalBox)
+									+ SHorizontalBox::Slot()
+									.AutoWidth()
 									[
-										SNew(SVerticalBox)
-										.Visibility(this, &FPrimitiveComponentDetails::IsMassVisible, true)
-										+ SVerticalBox::Slot()
-										.AutoHeight()
-										[
-											ChildProperty->CreatePropertyValueWidget()
-										]
+										BodyInstanceHandler->GetChildHandle(TEXT("bLockYRotation"))->CreatePropertyNameWidget()
 									]
-									
-								];
-						}
+									+ SHorizontalBox::Slot()
+									.AutoWidth()
+									[
+										BodyInstanceHandler->GetChildHandle(TEXT("bLockYRotation"))->CreatePropertyValueWidget()
+									]
+								]
+
+								+ SHorizontalBox::Slot()
+								.Padding(0.f, 0.f, XYZPadding, 0.f)
+								.AutoWidth()
+								[
+									SNew(SHorizontalBox)
+									+ SHorizontalBox::Slot()
+									.AutoWidth()
+									[
+										BodyInstanceHandler->GetChildHandle(TEXT("bLockZRotation"))->CreatePropertyNameWidget()
+									]
+									+ SHorizontalBox::Slot()
+									.AutoWidth()
+									[
+										BodyInstanceHandler->GetChildHandle(TEXT("bLockZRotation"))->CreatePropertyValueWidget()
+									]
+								]
+
+							];
+
+						//we only show the custom plane normal if we've selected that mode
+						ConstraintsGroup.AddPropertyRow(BodyInstanceHandler->GetChildHandle(TEXT("CustomDOFPlaneNormal")).ToSharedRef()).Visibility(TAttribute<EVisibility>::Create(TAttribute<EVisibility>::FGetter::CreateSP(this, &FPrimitiveComponentDetails::IsDOFMode, EDOFMode::CustomPlane)));
+						ConstraintsGroup.AddPropertyRow(BodyInstanceHandler->GetChildHandle(TEXT("bLockTranslation")).ToSharedRef()).Visibility(TAttribute<EVisibility>::Create(TAttribute<EVisibility>::FGetter::CreateSP(this, &FPrimitiveComponentDetails::IsDOFMode, EDOFMode::CustomPlane)));
+						ConstraintsGroup.AddPropertyRow(BodyInstanceHandler->GetChildHandle(TEXT("bLockRotation")).ToSharedRef()).Visibility(TAttribute<EVisibility>::Create(TAttribute<EVisibility>::FGetter::CreateSP(this, &FPrimitiveComponentDetails::IsDOFMode, EDOFMode::CustomPlane)));
+						ConstraintsGroup.AddPropertyRow(DOFModeProperty.ToSharedRef());
 					}
-					else
+				}
+				else if (PropName == GET_MEMBER_NAME_CHECKED(FBodyInstance, bAutoWeld))
+				{
+					PhysicsCategory.AddProperty(ChildProperty).Visibility(TAttribute<EVisibility>(this, &FPrimitiveComponentDetails::IsAutoWeldVisible))
+						.EditCondition(TAttribute<bool>(this, &FPrimitiveComponentDetails::IsAutoWeldEditable), NULL);
+				}
+				else if (PropName == GET_MEMBER_NAME_CHECKED(FBodyInstance, MassInKg) || PropName == GET_MEMBER_NAME_CHECKED(FBodyInstance, bOverrideMass))
+				{
+					if (bDisplayMass && bMassUIExists == false)
 					{
-						PhysicsCategory.AddProperty(ChildProperty);
+						bMassUIExists = true;
+						TSharedPtr<IPropertyHandle> MassInKGHandle = BodyInstanceHandler->GetChildHandle(TEXT("MassInKg"));
+						TSharedPtr<IPropertyHandle> MassOverrideHandle = BodyInstanceHandler->GetChildHandle(TEXT("bOverrideMass"));
+
+						PhysicsCategory.AddCustomRow(LOCTEXT("Mass", "Mass"), false)
+						.NameContent()
+						[
+							MassInKGHandle->CreatePropertyNameWidget()
+						]
+						.ValueContent()
+						.MinDesiredWidth(125 * 3.f)
+						[
+							SNew(SHorizontalBox)
+							+ SHorizontalBox::Slot()
+							.AutoWidth()
+							[
+								SNew(SVerticalBox)
+								+ SVerticalBox::Slot()
+								.Padding(0.f, 0.f, 10.f, 0.f)
+								[
+									SNew(SNumericEntryBox<float>)
+									.IsEnabled(false)
+									.Value(this, &FPrimitiveComponentDetails::OnGetBodyMass)
+									.Visibility(this, &FPrimitiveComponentDetails::IsMassVisible, false)
+								]
+
+								+ SVerticalBox::Slot()
+								.Padding(0.f, 0.f, 10.f, 0.f)
+								[
+									SNew(SVerticalBox)
+									.Visibility(this, &FPrimitiveComponentDetails::IsMassVisible, true)
+									+ SVerticalBox::Slot()
+									.AutoHeight()
+									[
+										MassInKGHandle->CreatePropertyValueWidget()
+									]
+								]
+							]
+
+							+ SHorizontalBox::Slot()
+							.AutoWidth()
+							[
+								bDisplayMassOverride ? MassOverrideHandle->CreatePropertyValueWidget() : SNew(SSpacer)
+							]
+
+							+ SHorizontalBox::Slot()
+							.AutoWidth()
+							[
+								bDisplayMassOverride ? MassOverrideHandle->CreatePropertyNameWidget() : SNew(SSpacer)
+							]
+						];
 					}
 				}
+				else
+				{
+					PhysicsCategory.AddProperty(ChildProperty);
+				}
 			}
 		}
+	}
+}
 
+void FPrimitiveComponentDetails::AddCollisionCategory(IDetailLayoutBuilder& DetailBuilder)
+{
+	if (DetailBuilder.GetProperty(GET_MEMBER_NAME_CHECKED(UPrimitiveComponent, BodyInstance))->IsValidHandle())
+	{
 		// Collision
-		{
-			IDetailCategoryBuilder& CollisionCategory = DetailBuilder.EditCategory("Collision");
+		TSharedPtr<IPropertyHandle> BodyInstanceHandler = DetailBuilder.GetProperty(GET_MEMBER_NAME_CHECKED(UPrimitiveComponent, BodyInstance));
+		uint32 NumChildren = 0;
+		BodyInstanceHandler->GetNumChildren(NumChildren);
 
-			// add all collision properties
-			for (uint32 ChildIndex = 0; ChildIndex < NumChildren; ++ChildIndex)
+		IDetailCategoryBuilder& CollisionCategory = DetailBuilder.EditCategory("Collision");
+
+		// add all collision properties
+		for (uint32 ChildIndex = 0; ChildIndex < NumChildren; ++ChildIndex)
+		{
+			TSharedPtr<IPropertyHandle> ChildProperty = BodyInstanceHandler->GetChildHandle(ChildIndex);
+			FString Category = FObjectEditorUtils::GetCategory(ChildProperty->GetProperty());
+			if (Category == TEXT("Collision"))
 			{
-				TSharedPtr<IPropertyHandle> ChildProperty = BodyInstanceHandler->GetChildHandle(ChildIndex);
-				FString Category = FObjectEditorUtils::GetCategory(ChildProperty->GetProperty());
-				if (Category == TEXT("Collision"))
-				{
-					CollisionCategory.AddProperty(ChildProperty);
-				}
+				CollisionCategory.AddProperty(ChildProperty);
 			}
-		}		
+		}
+	}
+}
+
+void FPrimitiveComponentDetails::CustomizeDetails( IDetailLayoutBuilder& DetailBuilder )
+{
+	// Get the objects being customized so we can enable/disable editing of 'Simulate Physics'
+	DetailBuilder.GetObjectsBeingCustomized(ObjectsCustomized);
+
+	// See if we are hiding Physics category
+	TArray<FString> HideCategories;
+	FEditorCategoryUtils::GetClassHideCategories(DetailBuilder.GetDetailsView().GetBaseClass(), HideCategories);
+
+	if(!HideCategories.Contains("Materials"))
+	{
+		AddMaterialCategory(DetailBuilder);
+	}
+
+	TSharedRef<IPropertyHandle> MobilityHandle = DetailBuilder.GetProperty(GET_MEMBER_NAME_CHECKED(UPrimitiveComponent, Mobility), USceneComponent::StaticClass());
+	MobilityHandle->SetToolTipText(LOCTEXT("PrimitiveMobilityTooltip", "Mobility for primitive components controls how they can be modified in game and therefore how they interact with lighting and physics.\n● A movable primitive component can be changed in game, but requires dynamic lighting and shadowing from lights which have a large performance cost.\n● A static primitive component can't be changed in game, but can have its lighting baked, which allows rendering to be very efficient."));
+
+
+	if(!HideCategories.Contains("Physics"))
+	{
+		AddPhysicsCategory(DetailBuilder);
 	}
 
+	if (!HideCategories.Contains("Collision"))
+	{
+		AddCollisionCategory(DetailBuilder);
+	}
+
+	if(!HideCategories.Contains("Lighting"))
+	{
+		AddLightingCategory(DetailBuilder);
+	}
 
 	AddAdvancedSubCategory( DetailBuilder, "Rendering", "TextureStreaming" );
 	AddAdvancedSubCategory( DetailBuilder, "Rendering", "LOD");
@@ -335,6 +498,12 @@ void FPrimitiveComponentDetails::AddMaterialCategory( IDetailLayoutBuilder& Deta
 	MaterialCategory = MakeShareable(new FComponentMaterialCategory(Components));
 	MaterialCategory->Create(DetailBuilder);
 }
+
+void FPrimitiveComponentDetails::AddLightingCategory(IDetailLayoutBuilder& DetailBuilder)
+{
+	IDetailCategoryBuilder& CollisionCategory = DetailBuilder.EditCategory("Lighting");
+}
+
 void FPrimitiveComponentDetails::AddAdvancedSubCategory( IDetailLayoutBuilder& DetailBuilder, FName MainCategoryName, FName SubCategoryName)
 {
 	TArray<TSharedRef<IPropertyHandle> > SubCategoryProperties;
@@ -365,7 +534,7 @@ void FPrimitiveComponentDetails::AddAdvancedSubCategory( IDetailLayoutBuilder& D
 
 }
 
-FText FPrimitiveComponentDetails::OnGetBodyMass() const
+TOptional<float> FPrimitiveComponentDetails::OnGetBodyMass() const
 {
 	UPrimitiveComponent* Comp = NULL;
 
@@ -379,6 +548,7 @@ FText FPrimitiveComponentDetails::OnGetBodyMass() const
 			Comp = Cast<UPrimitiveComponent>(ObjectIt->Get());
 
 			float CompMass = Comp->CalculateMass();
+			Comp->BodyInstance.MassInKg = CompMass;	//update the component's override mass to be the calculated one
 			if (Mass == 0.0f || FMath::Abs(Mass - CompMass) < 0.1f)
 			{
 				Mass = CompMass;
@@ -393,10 +563,10 @@ FText FPrimitiveComponentDetails::OnGetBodyMass() const
 
 	if (bMultipleValue)
 	{
-		return LOCTEXT("MultipleValues", "Multiple Values");
+		return TOptional<float>();
 	}
 
-	return FText::AsNumber(Mass);
+	return Mass;
 }
 
 bool FPrimitiveComponentDetails::IsBodyMassReadOnly() const
diff --git a/Engine/Source/Editor/DetailCustomizations/Private/PrimitiveComponentDetails.h b/Engine/Source/Editor/DetailCustomizations/Private/PrimitiveComponentDetails.h
--- a/Engine/Source/Editor/DetailCustomizations/Private/PrimitiveComponentDetails.h
+++ b/Engine/Source/Editor/DetailCustomizations/Private/PrimitiveComponentDetails.h
@@ -13,6 +13,9 @@ public:
 
 private:
 	void AddMaterialCategory( IDetailLayoutBuilder& DetailBuilder);
+	void AddLightingCategory(IDetailLayoutBuilder& DetailBuilder);
+	void AddPhysicsCategory(IDetailLayoutBuilder& DetailBuilder);
+	void AddCollisionCategory(IDetailLayoutBuilder& DetailBuilder);
 
 	ECheckBoxState IsMobilityActive(TSharedRef<IPropertyHandle> MobilityHandle, EComponentMobility::Type InMobility) const;
 
@@ -29,20 +32,20 @@ private:
 	/** Returns whether to enable editing the 'Use Async Scene' checkbox based on the selected objects' mobility and if the project uses an AsyncScene */
 	bool IsUseAsyncEditable() const;
 
-	FText OnGetBodyMass() const;
+	TOptional<float> OnGetBodyMass() const;
 	bool IsBodyMassReadOnly() const;
 	bool IsBodyMassEnabled() const { return !IsBodyMassReadOnly(); }
 
 private:
 	/** Objects being customized so we can update the 'Simulate Physics' state if physics geometry is added/removed */
 	TArray< TWeakObjectPtr<UObject> > ObjectsCustomized;
-	TSharedPtr<IPropertyHandle> LockedAxisProperty;
+	TSharedPtr<IPropertyHandle> DOFModeProperty;
 	TSharedPtr<class FComponentMaterialCategory> MaterialCategory;
-	EVisibility IsCustomLockedAxisSelected() const;
-	EVisibility IsLockAxisEnabled() const;
+	EVisibility IsDOFMode(EDOFMode::Type Mode) const;
 
 	bool IsAutoWeldEditable() const;
 	EVisibility IsAutoWeldVisible() const;
 	EVisibility IsMassVisible(bool bOverrideMass) const;
+
 };
 
diff --git a/Engine/Source/Runtime/Engine/Classes/Components/PrimitiveComponent.h b/Engine/Source/Runtime/Engine/Classes/Components/PrimitiveComponent.h
--- a/Engine/Source/Runtime/Engine/Classes/Components/PrimitiveComponent.h
+++ b/Engine/Source/Runtime/Engine/Classes/Components/PrimitiveComponent.h
@@ -41,7 +41,7 @@ enum ECanBeCharacterBase
 	ECB_MAX,
 };
 
-//UENUM()
+UENUM()
 namespace EHasCustomNavigableGeometry
 {
 	enum Type
@@ -55,7 +55,7 @@ namespace EHasCustomNavigableGeometry
 		// DoCustomNavigableGeometryExport() should be called even if the mesh is non-collidable and wouldn't normally affect the navmesh
 		EvenIfNotCollidable,
 	};
-};
+}
 
 /** Information about the sprite category */
 USTRUCT()
@@ -420,6 +420,7 @@ protected:
 	float LastCheckedAllCollideableDescendantsTime;
 
 	/** If true then DoCustomNavigableGeometryExport will be called to collect navigable geometry of this component. */
+	UPROPERTY()
 	TEnumAsByte<EHasCustomNavigableGeometry::Type> bHasCustomNavigableGeometry;
 
 	/** Next id to be used by a component. */
@@ -681,11 +682,19 @@ public:
 	virtual bool CanEditSimulatePhysics();
 
 	/**
-	 * Sets the specified axis as locked, preventing movement along that axis.
-	 * @param LockedAxis	The axis to lock.
+	* Sets the constraint mode of the component.
+	* @param ConstraintMode	The type of constraint to use.
+	*/
+	DEPRECATED(4.8, "This function is deprecated. Please use SetConstraintMode instead.")
+	UFUNCTION(BlueprintCallable, meta = (DeprecatedFunction, DeprecationMessage = "Use SetConstraintMode instead", Keywords = "set locked axis constraint physics"), Category = Physics)
+	virtual void SetLockedAxis(EDOFMode::Type LockedAxis);
+
+	/**
+	 * Sets the constraint mode of the component.
+	 * @param ConstraintMode	The type of constraint to use.
 	 */
-	UFUNCTION(BlueprintCallable, meta = (FriendlyName = "Set Locked Axis", Keywords = "set locked axis constraint physics"), Category = Physics)
-	virtual void SetLockedAxis(ELockedAxis::Type LockedAxis);
+	UFUNCTION(BlueprintCallable, meta = (FriendlyName = "Set Constraint Mode", Keywords = "set locked axis constraint physics"), Category = Physics)
+	virtual void SetConstraintMode(EDOFMode::Type ConstraintMode);
 
 	/**
 	 *	Add an impulse to a single rigid body. Good for one time instant burst.
diff --git a/Engine/Source/Runtime/Engine/Classes/PhysicsEngine/BodyInstance.h b/Engine/Source/Runtime/Engine/Classes/PhysicsEngine/BodyInstance.h
--- a/Engine/Source/Runtime/Engine/Classes/PhysicsEngine/BodyInstance.h
+++ b/Engine/Source/Runtime/Engine/Classes/PhysicsEngine/BodyInstance.h
@@ -36,21 +36,23 @@ namespace physx
 #endif // WITH_PHYSX
 
 UENUM(BlueprintType)
-namespace ELockedAxis
+namespace EDOFMode
 {
 	enum Type
 	{
-		/*Uses the default locked axis as specified in the project settings.*/
+		/*Inherits the degrees of freedom from the project settings.*/
 		Default,
-		/*Lock movement along the x-axis*/
-		X,
-		/*Lock movement along the y-axis*/
-		Y,
-		/*Lock movement along the z-axis*/
-		Z,
-		/*Lock movement along custom axis*/
-		Custom,
-		/*No axis is locked.*/
+		/*Specifies which axis to freeze rotation and movement along.*/
+		SixDOF,
+		/*Allows 2D movement along the Y-Z plane.*/
+		YZPlane,
+		/*Allows 2D movement along the X-Z plane.*/
+		XZPlane,
+		/*Allows 2D movement along the X-Y plane.*/
+		XYPlane,
+		/*Allows 2D movement along the plane of a given normal*/
+		CustomPlane,
+		/*No constraints.*/
 		None
 	};
 }
@@ -191,42 +193,42 @@ public:
 	UPROPERTY(EditAnywhere, BlueprintReadWrite, Category=Physics)
 	uint32 bSimulatePhysics:1;
 
-	/** If true and is attached to a parent, the two bodies will be joined into a single rigid body. Physical settings like collision profile and body settings are determined by the root */
+	/** If true, mass will not be automatically computed and you must set it directly */
+	UPROPERTY(EditAnywhere, BlueprintReadOnly, Category = Physics, meta=(DisplayName="Override"))
+	uint32 bOverrideMass : 1;
+
+	/**Mass of the body in KG. By default we compute this based on physical material and mass scale.
+	*@see bOverrideMass to set this directly */
+	UPROPERTY(EditAnywhere, BlueprintReadOnly, Category = Physics, meta = (editcondition = "bOverrideMass", ClampMin = "0.001", UIMin = "0.001"))
+	float MassInKg;
+
+	/** 'Drag' force added to reduce linear movement */
+	UPROPERTY(EditAnywhere, BlueprintReadWrite, Category = Physics)
+	float LinearDamping;
+
+	/** 'Drag' force added to reduce angular movement */
 	UPROPERTY(EditAnywhere, BlueprintReadWrite, Category = Physics)
+	float AngularDamping;
+
+	/** If object should have the force of gravity applied */
+	UPROPERTY(EditAnywhere, BlueprintReadOnly, Category = Physics)
+	uint32 bEnableGravity : 1;
+
+	/** If true and is attached to a parent, the two bodies will be joined into a single rigid body. Physical settings like collision profile and body settings are determined by the root */
+	UPROPERTY(EditAnywhere, AdvancedDisplay, BlueprintReadWrite, Category = Physics)
 	uint32 bAutoWeld : 1;
 
 	/** determines if the body is currently welded */
 	uint32 bWelded : 1;
 
 	/** If object should start awake, or if it should initially be sleeping */
-	UPROPERTY(EditAnywhere, BlueprintReadOnly, Category=Physics, meta=(editcondition = "bSimulatePhysics"))
+	UPROPERTY(EditAnywhere, AdvancedDisplay, BlueprintReadOnly, Category = Physics, meta = (editcondition = "bSimulatePhysics"))
 	uint32 bStartAwake:1;
 
-	/** If object should have the force of gravity applied */
-	UPROPERTY(EditAnywhere, AdvancedDisplay, BlueprintReadOnly, Category=Physics)
-	uint32 bEnableGravity:1;
-
 	/** If true, it will update mass when scale changes **/
 	UPROPERTY()
 	uint32 bUpdateMassWhenScaleChanges:1;
 
-	/** If true, mass will not be automatically computed and you must set it directly */
-	UPROPERTY(EditAnywhere, BlueprintReadOnly, Category=Physics)
-	uint32 bOverrideMass : 1;
-
-	/**Mass of the body in KG. By default we compute this based on physical material and mass scale.
-	   *@see bOverrideMass to set this directly */
-	UPROPERTY(EditAnywhere, BlueprintReadOnly, Category = Physics, meta = (editcondition = "bOverrideMass", ClampMin = "0.001", UIMin = "0.001"))
-	float MassInKg;
-
-	/** Locks physical movement along specified axis.*/
-	UPROPERTY(EditAnywhere, Category = Physics, meta=(DisplayName="Locked Axis"))
-	TEnumAsByte<ELockedAxis::Type> LockedAxisMode;
-	
-	/** Locks physical movement along custom axis. (0,0,0) indicates no lock*/
-	UPROPERTY(EditAnywhere, Category = Physics)
-	FVector CustomLockedAxis;
-
 	/** When a Locked Axis Mode is selected, will lock translation on the specified axis*/
 	UPROPERTY(EditAnywhere, Category = Physics, meta=(DisplayName = "Lock Axis Translation"))
 	uint32 bLockTranslation : 1;
@@ -235,12 +237,59 @@ public:
 	UPROPERTY(EditAnywhere, Category = Physics, meta=(DisplayName = "Lock Axis Rotation"))
 	uint32 bLockRotation : 1;
 
+	/** Lock translation along the X-axis*/
+	UPROPERTY(EditAnywhere, Category = Physics, meta = (DisplayName = "X"))
+	uint32 bLockXTranslation : 1;
+
+	/** Lock translation along the Y-axis*/
+	UPROPERTY(EditAnywhere, Category = Physics, meta = (DisplayName = "Y"))
+	uint32 bLockYTranslation : 1;
+
+	/** Lock translation along the Z-axis*/
+	UPROPERTY(EditAnywhere, Category = Physics, meta = (DisplayName = "Z"))
+	uint32 bLockZTranslation : 1;
+
+	/** Lock rotation about the X-axis*/
+	UPROPERTY(EditAnywhere, Category = Physics, meta = (DisplayName = "X"))
+	uint32 bLockXRotation : 1;
+
+	/** Lock rotation about the Y-axis*/
+	UPROPERTY(EditAnywhere, Category = Physics, meta = (DisplayName = "Y"))
+	uint32 bLockYRotation : 1;
+
+	/** Lock rotation about the Z-axis*/
+	UPROPERTY(EditAnywhere, Category = Physics, meta = (DisplayName = "Z"))
+	uint32 bLockZRotation : 1;
+
+	/** Locks physical movement along specified axis.*/
+	UPROPERTY(EditAnywhere, Category = Physics, meta = (DisplayName = "Mode"))
+	TEnumAsByte<EDOFMode::Type> DOFMode;
+
+	/** Locks physical movement along a custom plane for a given normal.*/
+	UPROPERTY(EditAnywhere, Category = Physics, meta = (DisplayName = "Plane Normal"))
+	FVector CustomDOFPlaneNormal;
+
+	/** User specified offset for the center of mass of this object, from the calculated location */
+	UPROPERTY(EditAnywhere, AdvancedDisplay, BlueprintReadWrite, Category = Physics, meta = (DisplayName = "Center Of Mass Offset"))
+	FVector COMNudge;
+
+	/** Per-instance scaling of mass */
+	UPROPERTY(EditAnywhere, AdvancedDisplay, BlueprintReadWrite, Category = Physics)
+	float MassScale;
+
+	/** The maximum angular velocity for this instance */
+	UPROPERTY(EditAnywhere, AdvancedDisplay, BlueprintReadOnly, Category = Physics, meta = (editcondition = "bSimulatePhysics"))
+	float MaxAngularVelocity;
+
+
 	/** Locks physical movement along axis. */
-	void SetDOFLock(ELockedAxis::Type NewAxisMode);
+	void SetDOFLock(EDOFMode::Type NewDOFMode);
 
 	FVector GetLockedAxis() const;
 	void CreateDOFLock();
 
+	static EDOFMode::Type ResolveDOFMode(EDOFMode::Type DOFMode);
+
 	/** Constraint used to allow for easy DOF setup per bodyinstance */
 	FConstraintInstance * DOFConstraint;
 
@@ -266,13 +315,6 @@ protected:
 	UPROPERTY(EditAnywhere, AdvancedDisplay, BlueprintReadOnly, Category=Physics)
 	uint32 bOverrideWalkableSlopeOnInstance:1;
 
-	/**
-	 * Custom walkable slope override setting for this instance.
-	 * @see GetWalkableSlopeOverride(), SetWalkableSlopeOverride()
-	 */
-	UPROPERTY(EditAnywhere, AdvancedDisplay, BlueprintReadOnly, Category=Physics, meta=(editcondition="bOverrideWalkableSlopeOnInstance"))
-	struct FWalkableSlopeOverride WalkableSlopeOverride;
-
 	/** Whether this body instance has its own custom MaxDepenetrationVelocity*/
 	UPROPERTY(EditAnywhere, AdvancedDisplay, BlueprintReadOnly, Category = Physics)
 	uint32 bOverrideMaxDepenetrationVelocity : 1;
@@ -281,35 +323,22 @@ protected:
 	UPROPERTY(EditAnywhere, AdvancedDisplay, BlueprintReadOnly, Category = Physics, meta = (editcondition = "bOverrideMaxDepenetrationVelocity", ClampMin = "0.0", UIMin = "0.0"))
 	float MaxDepenetrationVelocity;
 
+	/**
+	* Custom walkable slope override setting for this instance.
+	* @see GetWalkableSlopeOverride(), SetWalkableSlopeOverride()
+	*/
+	UPROPERTY(EditAnywhere, AdvancedDisplay, BlueprintReadOnly, Category = Physics, meta = (editcondition = "bOverrideWalkableSlopeOnInstance"))
+	struct FWalkableSlopeOverride WalkableSlopeOverride;
+
 	/**	Allows you to override the PhysicalMaterial to use for simple collision on this body. */
-	UPROPERTY(EditAnywhere, AdvancedDisplay, BlueprintReadOnly, Category=Collision)
+	UPROPERTY(EditAnywhere, BlueprintReadOnly, Category=Collision)
 	class UPhysicalMaterial* PhysMaterialOverride;
 
 public:
-	/** User specified offset for the center of mass of this object, from the calculated location */
-	UPROPERTY(EditAnywhere, AdvancedDisplay, BlueprintReadWrite, Category=Physics, meta=(DisplayName="Center Of Mass Offset"))
-	FVector COMNudge;
-
 	/** The set of values used in considering when put this body to sleep. */
 	UPROPERTY(EditAnywhere, AdvancedDisplay, BlueprintReadWrite, Category=Physics)
 	TEnumAsByte<enum ESleepFamily> SleepFamily;
 
-	/** Per-instance scaling of mass */
-	UPROPERTY(EditAnywhere, AdvancedDisplay, BlueprintReadWrite, Category=Physics)
-	float MassScale;
-
-	/** 'Drag' force added to reduce angular movement */
-	UPROPERTY(EditAnywhere, AdvancedDisplay, BlueprintReadWrite, Category=Physics)
-	float AngularDamping;
-
-	/** 'Drag' force added to reduce linear movement */
-	UPROPERTY(EditAnywhere, AdvancedDisplay, BlueprintReadWrite, Category=Physics)
-	float LinearDamping;
-
-	/** The maximum angular velocity for this instance */
-	UPROPERTY(EditAnywhere, AdvancedDisplay, BlueprintReadOnly, Category = Physics, meta = (editcondition = "bSimulatePhysics"))
-	float MaxAngularVelocity;
-
 	/**	Influence of rigid body physics (blending) on the mesh's pose (0.0 == use only animation, 1.0 == use only physics) */
 	/** Provide appropriate interface for doing this instead of allowing BlueprintReadWrite **/
 	UPROPERTY()
diff --git a/Engine/Source/Runtime/Engine/Classes/PhysicsEngine/PhysicsSettings.h b/Engine/Source/Runtime/Engine/Classes/PhysicsEngine/PhysicsSettings.h
--- a/Engine/Source/Runtime/Engine/Classes/PhysicsEngine/PhysicsSettings.h
+++ b/Engine/Source/Runtime/Engine/Classes/PhysicsEngine/PhysicsSettings.h
@@ -48,6 +48,22 @@ namespace EFrictionCombineMode
 	};
 }
 
+UENUM()
+namespace ESettingsDOF
+{
+	enum Type
+	{
+		/*Allows for full 3D movement and rotation*/
+		Full3D,
+		/*Allows 2D movement along the Y-Z plane.*/
+		YZPlane,
+		/*Allows 2D movement along the X-Z plane.*/
+		XZPlane,
+		/*Allows 2D movement along the X-Y plane.*/
+		XYPlane,
+	};
+}
+
 UENUM()
 namespace ESettingsLockedAxis
 {
@@ -102,9 +118,12 @@ class ENGINE_API UPhysicsSettings
 	UPROPERTY(config, EditAnywhere, Category = Simulation)
 	bool bEnable2DPhysics;
 
-	/** Locks axis of movement for physical objects. (Useful for making 2D games in a 3D setting) */
+	UPROPERTY(config)
+	TEnumAsByte<ESettingsLockedAxis::Type> LockedAxis_DEPRECATED;
+
+	/** Useful for constraining all objects in the world, for example if you are making a 2D game using 3D environments.*/
 	UPROPERTY(config, EditAnywhere, Category = Simulation)
-	TEnumAsByte<ESettingsLockedAxis::Type> LockedAxis;
+	TEnumAsByte<ESettingsDOF::Type> DefaultDegreesOfFreedom;
 	
 	/** Minimum relative velocity required for an object to bounce. A typical value for simulation stability is about 0.2 * gravity*/
 	UPROPERTY(config, EditAnywhere, Category = Simulation, meta = (ClampMin = "0", UIMin = "0"))
diff --git a/Engine/Source/Runtime/Engine/Private/MovementComponent.cpp b/Engine/Source/Runtime/Engine/Private/MovementComponent.cpp
--- a/Engine/Source/Runtime/Engine/Private/MovementComponent.cpp
+++ b/Engine/Source/Runtime/Engine/Private/MovementComponent.cpp
@@ -333,13 +333,13 @@ FVector UMovementComponent::GetPlaneConstraintNormalFromAxisSetting(EPlaneConstr
 {
 	if (AxisSetting == EPlaneConstraintAxisSetting::UseGlobalPhysicsSetting)
 	{
-		ESettingsLockedAxis::Type GlobalSetting = UPhysicsSettings::Get()->LockedAxis;
+		ESettingsDOF::Type GlobalSetting = UPhysicsSettings::Get()->DefaultDegreesOfFreedom;
 		switch (GlobalSetting)
 		{
-		case ESettingsLockedAxis::None:	return FVector::ZeroVector;
-		case ESettingsLockedAxis::X:	return FVector(1.f, 0.f, 0.f);
-		case ESettingsLockedAxis::Y:	return FVector(0.f, 1.f, 0.f);
-		case ESettingsLockedAxis::Z:	return FVector(0.f, 0.f, 1.f);
+		case ESettingsDOF::Full3D:	return FVector::ZeroVector;
+		case ESettingsDOF::YZPlane:	return FVector(1.f, 0.f, 0.f);
+		case ESettingsDOF::XZPlane:	return FVector(0.f, 1.f, 0.f);
+		case ESettingsDOF::XYPlane:	return FVector(0.f, 0.f, 1.f);
 		default:
 			checkf(false, TEXT("GetPlaneConstraintNormalFromAxisSetting: Unknown global axis setting %d for %s"), int32(GlobalSetting), *GetNameSafe(GetOwner()));
 			return FVector::ZeroVector;
diff --git a/Engine/Source/Runtime/Engine/Private/PhysicsEngine/BodyInstance.cpp b/Engine/Source/Runtime/Engine/Private/PhysicsEngine/BodyInstance.cpp
--- a/Engine/Source/Runtime/Engine/Private/PhysicsEngine/BodyInstance.cpp
+++ b/Engine/Source/Runtime/Engine/Private/PhysicsEngine/BodyInstance.cpp
@@ -216,8 +216,14 @@ FBodyInstance::FBodyInstance()
 , bUpdateMassWhenScaleChanges(false)
 , bOverrideMass(false)
 , MassInKg(100.f)
-, LockedAxisMode(0)
-, CustomLockedAxis(FVector::ZeroVector)
+, DOFMode(0)
+, CustomDOFPlaneNormal(FVector::ZeroVector)
+, bLockXTranslation(false)
+, bLockYTranslation(false)
+, bLockZTranslation(false)
+, bLockXRotation(false)
+, bLockYRotation(false)
+, bLockZRotation(false)
 , bLockTranslation(true)
 , bLockRotation(true)
 , DOFConstraint(NULL)
@@ -508,34 +514,44 @@ void FBodyInstance::SetCollisionEnabled(ECollisionEnabled::Type NewType, bool bU
 	}
 }
 
-FVector FBodyInstance::GetLockedAxis() const
+
+
+EDOFMode::Type FBodyInstance::ResolveDOFMode(EDOFMode::Type DOFMode)
 {
-	ELockedAxis::Type MyLockedAxis = LockedAxisMode;
-	if (MyLockedAxis == ELockedAxis::Default)
+	EDOFMode::Type ResultDOF = DOFMode;
+	if (DOFMode == EDOFMode::Default)
 	{
-		ESettingsLockedAxis::Type SettingLockedAxis = UPhysicsSettings::Get()->LockedAxis;
-		if (SettingLockedAxis == ESettingsLockedAxis::X) MyLockedAxis = ELockedAxis::X;
-		if (SettingLockedAxis == ESettingsLockedAxis::Y) MyLockedAxis = ELockedAxis::Y;
-		if (SettingLockedAxis == ESettingsLockedAxis::Z) MyLockedAxis = ELockedAxis::Z;
-		if (SettingLockedAxis == ESettingsLockedAxis::None) MyLockedAxis = ELockedAxis::None;
+		ESettingsDOF::Type SettingDefaultPlane = UPhysicsSettings::Get()->DefaultDegreesOfFreedom;
+		if (SettingDefaultPlane == ESettingsDOF::XYPlane) ResultDOF = EDOFMode::XYPlane;
+		if (SettingDefaultPlane == ESettingsDOF::XZPlane) ResultDOF = EDOFMode::XZPlane;
+		if (SettingDefaultPlane == ESettingsDOF::YZPlane) ResultDOF = EDOFMode::YZPlane;
+		if (SettingDefaultPlane == ESettingsDOF::Full3D) ResultDOF  = EDOFMode::SixDOF;
 	}
 
-	switch (MyLockedAxis)
+	return ResultDOF;
+}
+
+FVector FBodyInstance::GetLockedAxis() const
+{
+	EDOFMode::Type MyDOFMode = ResolveDOFMode(DOFMode);
+
+	switch (MyDOFMode)
 	{
-	case ELockedAxis::None: return FVector::ZeroVector;
-	case ELockedAxis::X: return FVector(1, 0, 0);
-	case ELockedAxis::Y: return FVector(0, 1, 0);
-	case ELockedAxis::Z: return FVector(0, 0, 1);
-	case ELockedAxis::Custom: return CustomLockedAxis;
+	case EDOFMode::None: return FVector::ZeroVector;
+	case EDOFMode::YZPlane: return FVector(1, 0, 0);
+	case EDOFMode::XZPlane: return FVector(0, 1, 0);
+	case EDOFMode::XYPlane: return FVector(0, 0, 1);
+	case EDOFMode::CustomPlane: return CustomDOFPlaneNormal;
+	case EDOFMode::SixDOF: return FVector::ZeroVector;
 	default:	check(0);	//unsupported locked axis type
 	}
 
 	return FVector::ZeroVector;
 }
 
-void FBodyInstance::SetDOFLock(ELockedAxis::Type NewAxisMode)
+void FBodyInstance::SetDOFLock(EDOFMode::Type NewAxisMode)
 {
-	LockedAxisMode = NewAxisMode;
+	DOFMode = NewAxisMode;
 
 	CreateDOFLock();
 }
@@ -549,10 +565,16 @@ void FBodyInstance::CreateDOFLock()
 		DOFConstraint = NULL;
 	}
 
-	//setup constraint based on DOF
-	FVector LockedAxis = GetLockedAxis();
+	const FVector LockedAxis = GetLockedAxis();
+	const EDOFMode::Type DOF = ResolveDOFMode(DOFMode);
+
+	if (IsDynamic() == false || (LockedAxis.IsNearlyZero() && DOF != EDOFMode::SixDOF))
+	{
+		return;
+	}
 
-	if (IsDynamic() == false || LockedAxis.IsNearlyZero())
+	//if we're using SixDOF make sure we have at least one constraint
+	if (DOF == EDOFMode::SixDOF && !bLockXTranslation && !bLockYTranslation && !bLockZTranslation && !bLockXRotation && !bLockYRotation && !bLockZRotation)
 	{
 		return;
 	}
@@ -562,22 +584,35 @@ void FBodyInstance::CreateDOFLock()
 		DOFConstraint->bSwingLimitSoft = false;
 		DOFConstraint->bTwistLimitSoft = false;
 		DOFConstraint->bLinearLimitSoft = false;
-		//set all rotation to free
-		DOFConstraint->AngularSwing1Motion = bLockRotation ? EAngularConstraintMotion::ACM_Locked : EAngularConstraintMotion::ACM_Free;
-		DOFConstraint->AngularSwing2Motion = bLockRotation ? EAngularConstraintMotion::ACM_Locked : EAngularConstraintMotion::ACM_Free;
-		DOFConstraint->AngularTwistMotion = EAngularConstraintMotion::ACM_Free;
 
-		DOFConstraint->LinearXMotion = bLockTranslation ? ELinearConstraintMotion::LCM_Locked : ELinearConstraintMotion::LCM_Free;
-		DOFConstraint->LinearYMotion = ELinearConstraintMotion::LCM_Free;
-		DOFConstraint->LinearZMotion = ELinearConstraintMotion::LCM_Free;
+		const FTransform TM = GetUnrealWorldTransform();
+		FVector Normal = FVector(1, 0, 0);
+		FVector Sec = FVector(0, 1, 0);
 
-		FVector Normal = LockedAxis.GetSafeNormal();
-		FVector Sec;
-		FVector Garbage;
-		Normal.FindBestAxisVectors(Garbage, Sec);
 
+		if(DOF != EDOFMode::SixDOF)
+		{
+			DOFConstraint->AngularSwing1Motion = (bLockRotation || DOFMode != EDOFMode::CustomPlane) ? EAngularConstraintMotion::ACM_Locked : EAngularConstraintMotion::ACM_Free;
+			DOFConstraint->AngularSwing2Motion = (bLockRotation || DOFMode != EDOFMode::CustomPlane) ? EAngularConstraintMotion::ACM_Locked : EAngularConstraintMotion::ACM_Free;
+			DOFConstraint->AngularTwistMotion = EAngularConstraintMotion::ACM_Free;
 
-		FTransform TM = GetUnrealWorldTransform();
+			DOFConstraint->LinearXMotion = (bLockTranslation || DOFMode != EDOFMode::CustomPlane) ? ELinearConstraintMotion::LCM_Locked : ELinearConstraintMotion::LCM_Free;
+			DOFConstraint->LinearYMotion = ELinearConstraintMotion::LCM_Free;
+			DOFConstraint->LinearZMotion = ELinearConstraintMotion::LCM_Free;
+
+			Normal = LockedAxis.GetSafeNormal();
+			FVector Garbage;
+			Normal.FindBestAxisVectors(Garbage, Sec);
+		}else
+		{
+			DOFConstraint->AngularTwistMotion = bLockXRotation ? EAngularConstraintMotion::ACM_Locked : EAngularConstraintMotion::ACM_Free;
+			DOFConstraint->AngularSwing2Motion = bLockYRotation ? EAngularConstraintMotion::ACM_Locked : EAngularConstraintMotion::ACM_Free;
+			DOFConstraint->AngularSwing1Motion = bLockZRotation ? EAngularConstraintMotion::ACM_Locked : EAngularConstraintMotion::ACM_Free;
+
+			DOFConstraint->LinearXMotion = bLockXTranslation ? ELinearConstraintMotion::LCM_Locked : ELinearConstraintMotion::LCM_Free;
+			DOFConstraint->LinearYMotion = bLockYTranslation ? ELinearConstraintMotion::LCM_Locked : ELinearConstraintMotion::LCM_Free;
+			DOFConstraint->LinearZMotion = bLockZTranslation ? ELinearConstraintMotion::LCM_Locked : ELinearConstraintMotion::LCM_Free;
+		}
 
 		DOFConstraint->PriAxis1 = TM.InverseTransformVectorNoScale(Normal);
 		DOFConstraint->SecAxis1 = TM.InverseTransformVectorNoScale(Sec);
@@ -3325,6 +3360,8 @@ void FBodyInstance::UpdateMassProperties()
 			
 			// Apply user-defined mass scaling.
 			NewMass *= FMath::Clamp<float>(MassScale, 0.01f, 100.0f);
+
+			MassInKg = NewMass;	//update the override mass to be the default mass
 		}
 		else
 		{
diff --git a/Engine/Source/Runtime/Engine/Private/PhysicsEngine/PhysicsSettings.cpp b/Engine/Source/Runtime/Engine/Private/PhysicsEngine/PhysicsSettings.cpp
--- a/Engine/Source/Runtime/Engine/Private/PhysicsEngine/PhysicsSettings.cpp
+++ b/Engine/Source/Runtime/Engine/Private/PhysicsEngine/PhysicsSettings.cpp
@@ -14,6 +14,7 @@ UPhysicsSettings::UPhysicsSettings(const FObjectInitializer& ObjectInitializer)
 	, TriangleMeshTriangleMinAreaThreshold(5.0f)
 	, bEnableAsyncScene(false)
 	, bEnable2DPhysics(false)
+	, LockedAxis_DEPRECATED(static_cast<ESettingsLockedAxis::Type>(-1))
 	, BounceThresholdVelocity(200.f)
 	, bSimulateSkeletalMeshOnDedicatedServer(true)
 	, MaxPhysicsDeltaTime(1.f / 30.f)
@@ -33,6 +34,28 @@ void UPhysicsSettings::PostInitProperties()
 #if WITH_EDITOR
 	LoadSurfaceType();
 #endif
+
+	if (LockedAxis_DEPRECATED != static_cast<ESettingsLockedAxis::Type>(-1))
+	{
+		if (LockedAxis_DEPRECATED == ESettingsLockedAxis::None)
+		{
+			DefaultDegreesOfFreedom = ESettingsDOF::Full3D;
+		}
+		else if (LockedAxis_DEPRECATED == ESettingsLockedAxis::X)
+		{
+			DefaultDegreesOfFreedom = ESettingsDOF::YZPlane;
+		}
+		else if (LockedAxis_DEPRECATED == ESettingsLockedAxis::Y)
+		{
+			DefaultDegreesOfFreedom = ESettingsDOF::XZPlane;
+		}
+		else if (LockedAxis_DEPRECATED == ESettingsLockedAxis::Z)
+		{
+			DefaultDegreesOfFreedom = ESettingsDOF::XYPlane;
+		}
+
+		LockedAxis_DEPRECATED = static_cast<ESettingsLockedAxis::Type>(-1);
+	}
 }
 
 #if WITH_EDITOR
@@ -60,7 +83,7 @@ void UPhysicsSettings::PostEditChangeProperty(struct FPropertyChangedEvent& Prop
 	{
 		UPhysicalMaterial::RebuildPhysicalMaterials();
 	}
-	else if (PropertyName == GET_MEMBER_NAME_CHECKED(UPhysicsSettings, LockedAxis))
+	else if (PropertyName == GET_MEMBER_NAME_CHECKED(UPhysicsSettings, DefaultDegreesOfFreedom))
 	{
 		UMovementComponent::PhysicsLockedAxisSettingChanged();
 	}
diff --git a/Engine/Source/Runtime/Engine/Private/PrimitiveComponentPhysics.cpp b/Engine/Source/Runtime/Engine/Private/PrimitiveComponentPhysics.cpp
--- a/Engine/Source/Runtime/Engine/Private/PrimitiveComponentPhysics.cpp
+++ b/Engine/Source/Runtime/Engine/Private/PrimitiveComponentPhysics.cpp
@@ -171,7 +171,7 @@ void UPrimitiveComponent::SetSimulatePhysics(bool bSimulate)
 	BodyInstance.SetInstanceSimulatePhysics(bSimulate);
 }
 
-void UPrimitiveComponent::SetLockedAxis(ELockedAxis::Type LockedAxis)
+void UPrimitiveComponent::SetConstraintMode(EDOFMode::Type ConstraintMode)
 {
 	FBodyInstance * RootBI = GetBodyInstance(NAME_None, false);
 
@@ -180,7 +180,12 @@ void UPrimitiveComponent::SetLockedAxis(ELockedAxis::Type LockedAxis)
 		return;
 	}
 
-	RootBI->SetDOFLock(LockedAxis);
+	RootBI->SetDOFLock(ConstraintMode);
+}
+
+void UPrimitiveComponent::SetLockedAxis(EDOFMode::Type LockedAxis)
+{
+	SetConstraintMode(LockedAxis);
 }
 
 void UPrimitiveComponent::AddImpulse(FVector Impulse, FName BoneName, bool bVelChange)
diff --git a/Engine/Source/Runtime/Engine/Private/SkeletalMeshComponentPhysics.cpp b/Engine/Source/Runtime/Engine/Private/SkeletalMeshComponentPhysics.cpp
--- a/Engine/Source/Runtime/Engine/Private/SkeletalMeshComponentPhysics.cpp
+++ b/Engine/Source/Runtime/Engine/Private/SkeletalMeshComponentPhysics.cpp
@@ -779,14 +779,22 @@ void USkeletalMeshComponent::InitArticulated(FPhysScene* PhysScene)
 
 			if (i == RootBodyIndex)
 			{
-				BodyInst->LockedAxisMode = BodyInstance.LockedAxisMode;
-				BodyInst->CustomLockedAxis = BodyInstance.CustomLockedAxis;
+				BodyInst->DOFMode = BodyInstance.DOFMode;
+				BodyInst->CustomDOFPlaneNormal = BodyInstance.CustomDOFPlaneNormal;
+				BodyInst->bLockXTranslation = BodyInstance.bLockXTranslation;
+				BodyInst->bLockYTranslation = BodyInstance.bLockYTranslation;
+				BodyInst->bLockZTranslation = BodyInstance.bLockZTranslation;
+				BodyInst->bLockXRotation = BodyInstance.bLockXRotation;
+				BodyInst->bLockYRotation = BodyInstance.bLockYRotation;
+				BodyInst->bLockZRotation = BodyInstance.bLockZRotation;
+				BodyInst->bLockTranslation = BodyInstance.bLockTranslation;
+				BodyInst->bLockRotation = BodyInstance.bLockRotation;
 
 				BodyInst->COMNudge = BodyInstance.COMNudge;
 			}
 			else
 			{
-				BodyInst->LockedAxisMode = ELockedAxis::None;
+				BodyInst->DOFMode = EDOFMode::None;
 			}
 
 #if WITH_PHYSX
