Fix Mac compile

diff --git a/Engine/Source/Editor/DetailCustomizations/Private/PrimitiveComponentDetails.cpp b/Engine/Source/Editor/DetailCustomizations/Private/PrimitiveComponentDetails.cpp
--- a/Engine/Source/Editor/DetailCustomizations/Private/PrimitiveComponentDetails.cpp
+++ b/Engine/Source/Editor/DetailCustomizations/Private/PrimitiveComponentDetails.cpp
@@ -401,13 +401,13 @@ void FPrimitiveComponentDetails::AddPhysicsCategory(IDetailLayoutBuilder& Detail
 							+ SHorizontalBox::Slot()
 							.AutoWidth()
 							[
-								bDisplayMassOverride ? MassOverrideHandle->CreatePropertyValueWidget() : SNew(SSpacer)
+								bDisplayMassOverride ? MassOverrideHandle->CreatePropertyValueWidget() : StaticCastSharedRef<SWidget>(SNew(SSpacer))
 							]
 
 							+ SHorizontalBox::Slot()
 							.AutoWidth()
 							[
-								bDisplayMassOverride ? MassOverrideHandle->CreatePropertyNameWidget() : SNew(SSpacer)
+								bDisplayMassOverride ? MassOverrideHandle->CreatePropertyNameWidget() : StaticCastSharedRef<SWidget>(SNew(SSpacer))
 							]
 						];
 					}
