[INTEGRATE] Change 2440322 by [REDACTED_PERSON][REDACTED_HANDLE] on 2015/02/10 15:40:12

	Fix material list in details not visible on a blueprinted component

diff --git a/Engine/Source/Editor/DetailCustomizations/Private/ComponentMaterialCategory.cpp b/Engine/Source/Editor/DetailCustomizations/Private/ComponentMaterialCategory.cpp
--- a/Engine/Source/Editor/DetailCustomizations/Private/ComponentMaterialCategory.cpp
+++ b/Engine/Source/Editor/DetailCustomizations/Private/ComponentMaterialCategory.cpp
@@ -55,17 +55,14 @@ public:
 				UPrimitiveComponent* PrimitiveComp = Cast<UPrimitiveComponent>( CurComponent );
 				UDecalComponent* DecalComponent = (PrimitiveComp ? NULL : Cast<UDecalComponent>( CurComponent ));
 
-				if (!CurComponent->IsCreatedByConstructionScript())
+				if( PrimitiveComp )
 				{
-					if( PrimitiveComp )
-					{
-						NumMaterials = PrimitiveComp->GetNumMaterials();
-					}
-					else if( DecalComponent )
-					{
-						// DecalComponent isn't a primitive component so we must get the materials directly from it
-						NumMaterials = DecalComponent->GetNumMaterials();
-					}
+					NumMaterials = PrimitiveComp->GetNumMaterials();
+				}
+				else if( DecalComponent )
+				{
+					// DecalComponent isn't a primitive component so we must get the materials directly from it
+					NumMaterials = DecalComponent->GetNumMaterials();
 				}
 
 				// Check materials
@@ -183,7 +180,7 @@ void FComponentMaterialCategory::Create( IDetailLayoutBuilder& DetailBuilder )
 	{	
 		UActorComponent* CurrentComponent = It.GetComponent();
 
-		if( !bAnyMaterialsToDisplay && !CurrentComponent->IsCreatedByConstructionScript() )
+		if( !bAnyMaterialsToDisplay )
 		{
 			bAnyMaterialsToDisplay = true;
 			break;
@@ -212,7 +209,7 @@ void FComponentMaterialCategory::OnGetMaterialsForView( IMaterialListBuilder& Ma
 
 		UActorComponent* CurrentComponent = It.GetComponent();
 
-		if( CurrentComponent && !CurrentComponent->IsCreatedByConstructionScript() )
+		if( CurrentComponent )
 		{
 			UMaterialInterface* Material = It.GetMaterial();
 
@@ -254,7 +251,7 @@ void FComponentMaterialCategory::OnMaterialChanged( UMaterialInterface* NewMater
 			AActor* Actor = CurrentComponent->GetOwner();
 
 			// Component materials can be replaced if they are not created from a blueprint (not exposed to the user) and have material overrides on the component
-			bool bCanBeReplaced = (!Actor || Actor->GetClass()->ClassGeneratedBy == NULL) && 
+			bool bCanBeReplaced = 
 				( CurrentComponent->IsA( UMeshComponent::StaticClass() ) ||
 				CurrentComponent->IsA( UDecalComponent::StaticClass() ) ||
 				CurrentComponent->IsA( UTextRenderComponent::StaticClass() ) ||
