[INTEGRATE] Change 2444726 by [REDACTED_PERSON][REDACTED_HANDLE] on 2015/02/13 10:16:29

	Fix crash when creating a blueprint from a skylight actor

	https://jira.ol.epicgames.net/browse/UE-9700

diff --git a/Engine/Source/Editor/DetailCustomizations/Private/SkyLightComponentDetails.cpp b/Engine/Source/Editor/DetailCustomizations/Private/SkyLightComponentDetails.cpp
--- a/Engine/Source/Editor/DetailCustomizations/Private/SkyLightComponentDetails.cpp
+++ b/Engine/Source/Editor/DetailCustomizations/Private/SkyLightComponentDetails.cpp
@@ -22,9 +22,13 @@ void FSkyLightComponentDetails::CustomizeDetails( IDetailLayoutBuilder& DetailLa
 	MobilityHandle->SetToolTipText(LOCTEXT("SkyLightMobilityTooltip", "Mobility for sky light components determines what rendering methods will be used.  A Stationary sky light has its shadowing baked into Bent Normal AO by Lightmass, but its lighting can be changed in game."));
 
 	TSharedPtr<IPropertyHandle> LightIntensityProperty = DetailLayout.GetProperty(GET_MEMBER_NAME_CHECKED(USkyLightComponent, Intensity), ULightComponentBase::StaticClass());
-	// Point lights need to override the ui min and max for units of lumens, so we have to undo that
-	LightIntensityProperty->GetProperty()->SetMetaData("UIMin",TEXT("0.0f"));
-	LightIntensityProperty->GetProperty()->SetMetaData("UIMax",TEXT("20.0f"));
+
+	if( LightIntensityProperty->IsValidHandle() )
+	{
+		// Point lights need to override the ui min and max for units of lumens, so we have to undo that
+		LightIntensityProperty->GetProperty()->SetMetaData("UIMin", TEXT("0.0f"));
+		LightIntensityProperty->GetProperty()->SetMetaData("UIMax", TEXT("20.0f"));
+	}
 
 	const TArray< TWeakObjectPtr<UObject> >& SelectedObjects = DetailLayout.GetDetailsView().GetSelectedObjects();
 
