Editor: Remove dead declaration FActorDetails::AddMaterialCategory (had no definition)

diff --git a/Engine/Source/Editor/DetailCustomizations/Private/ActorDetails.h b/Engine/Source/Editor/DetailCustomizations/Private/ActorDetails.h
--- a/Engine/Source/Editor/DetailCustomizations/Private/ActorDetails.h
+++ b/Engine/Source/Editor/DetailCustomizations/Private/ActorDetails.h
@@ -73,11 +73,6 @@ private:
 
 	void AddExperimentalWarningCategory( IDetailLayoutBuilder& DetailBuilder );
 
-	/**
-	 * Creates a category for displaying materials used by the selected actors
-	 */
-	void AddMaterialCategory( IDetailLayoutBuilder& DetailBuilder );
-
 	const TArray< TWeakObjectPtr<AActor> >& GetSelectedActors() const;
 
 	/** Creates a category to display blutility functionality in a separate space from the standard blueprint details */
