Fixed FBX import UI sometimes getting stuck in animation only import when importing static meshes as skeletal meshes
UE-4857

diff --git a/Engine/Source/Editor/DetailCustomizations/Private/FbxImportUIDetails.cpp b/Engine/Source/Editor/DetailCustomizations/Private/FbxImportUIDetails.cpp
--- a/Engine/Source/Editor/DetailCustomizations/Private/FbxImportUIDetails.cpp
+++ b/Engine/Source/Editor/DetailCustomizations/Private/FbxImportUIDetails.cpp
@@ -71,13 +71,6 @@ void FFbxImportUIDetails::CustomizeDetails( IDetailLayoutBuilder& DetailBuilder
 	}
 	EFBXImportType ImportType = ImportUI->MeshTypeToImport;
 
-	if(ImportUI->OriginalImportType == FBXIT_SkeletalMesh)
-	{
-		TSharedRef<IPropertyHandle> ImportMeshProp = DetailBuilder.GetProperty(GET_MEMBER_NAME_CHECKED(UFbxImportUI, bImportMesh));
-		ImportMeshProp->SetOnPropertyValueChanged(FSimpleDelegate::CreateSP(this, &FFbxImportUIDetails::ImportMeshToggleChanged));
-		MeshCategory.AddProperty(ImportMeshProp);
-	}
-
 	if(ImportType != FBXIT_Animation)
 	{
 		TSharedRef<IPropertyHandle> ImportSkeletalProp = DetailBuilder.GetProperty(GET_MEMBER_NAME_CHECKED(UFbxImportUI, bImportAsSkeletal));
@@ -85,6 +78,17 @@ void FFbxImportUIDetails::CustomizeDetails( IDetailLayoutBuilder& DetailBuilder
 		MeshCategory.AddProperty(ImportSkeletalProp);
 	}
 
+	TSharedRef<IPropertyHandle> ImportMeshProp = DetailBuilder.GetProperty(GET_MEMBER_NAME_CHECKED(UFbxImportUI, bImportMesh));
+	if(ImportUI->OriginalImportType == FBXIT_SkeletalMesh && ImportType != FBXIT_StaticMesh)
+	{
+		ImportMeshProp->SetOnPropertyValueChanged(FSimpleDelegate::CreateSP(this, &FFbxImportUIDetails::ImportMeshToggleChanged));
+		MeshCategory.AddProperty(ImportMeshProp);
+	}
+	else
+	{
+		DetailBuilder.HideProperty(ImportMeshProp);
+	}
+
 	for(TSharedRef<IPropertyHandle> Handle : CategoryDefaultProperties)
 	{
 		FString MetaData = Handle->GetMetaData(TEXT("ImportType"));
@@ -320,6 +324,7 @@ void FFbxImportUIDetails::MeshImportModeChanged()
 	if(CachedDetailBuilder)
 	{
 		ImportUI->MeshTypeToImport = ImportUI->MeshTypeToImport == FBXIT_SkeletalMesh ? FBXIT_StaticMesh : FBXIT_SkeletalMesh;
+		ImportUI->bImportMesh = true;
 		CachedDetailBuilder->ForceRefreshDetails();
 	}
 }
diff --git a/Engine/Source/Editor/UnrealEd/Classes/Factories/FbxImportUI.h b/Engine/Source/Editor/UnrealEd/Classes/Factories/FbxImportUI.h
--- a/Engine/Source/Editor/UnrealEd/Classes/Factories/FbxImportUI.h
+++ b/Engine/Source/Editor/UnrealEd/Classes/Factories/FbxImportUI.h
@@ -52,7 +52,7 @@ class UFbxImportUI : public UObject
 	bool bImportAsSkeletal;
 
 	/** Whether to import the mesh. Allows animation only import when importing a skeletal mesh. */
-	UPROPERTY(EditAnywhere, config, Category=Mesh, meta=(ImportType="SkeletalMesh"))
+	UPROPERTY(EditAnywhere, Category=Mesh, meta=(ImportType="SkeletalMesh"))
 	bool bImportMesh;
 
 	/** For static meshes, enabling this option will combine all meshes in the FBX into a single monolithic mesh in Unreal */
diff --git a/Engine/Source/Editor/UnrealEd/Private/Fbx/FbxMainImport.cpp b/Engine/Source/Editor/UnrealEd/Private/Fbx/FbxMainImport.cpp
--- a/Engine/Source/Editor/UnrealEd/Private/Fbx/FbxMainImport.cpp
+++ b/Engine/Source/Editor/UnrealEd/Private/Fbx/FbxMainImport.cpp
@@ -92,6 +92,7 @@ FBXImportOptions* GetImportOptions( UnFbx::FFbxImporter* FbxImporter, UFbxImport
 		}
 
 		ImportUI->bImportAsSkeletal = ImportUI->MeshTypeToImport == FBXIT_SkeletalMesh;
+		ImportUI->bImportMesh = ImportUI->MeshTypeToImport != FBXIT_Animation;
 		ImportUI->bIsObjImport = bIsObjFormat;
 
 		TSharedPtr<SWindow> ParentWindow;
