Fix crash accessing timespan property after a blueprint is compiled

diff --git a/Engine/Source/Editor/DetailCustomizations/Private/TimespanStructCustomization.cpp b/Engine/Source/Editor/DetailCustomizations/Private/TimespanStructCustomization.cpp
--- a/Engine/Source/Editor/DetailCustomizations/Private/TimespanStructCustomization.cpp
+++ b/Engine/Source/Editor/DetailCustomizations/Private/TimespanStructCustomization.cpp
@@ -36,6 +36,7 @@ void FTimespanStructCustomization::CustomizeHeader( TSharedRef<IPropertyHandle>
 				.OnTextChanged(this, &FTimespanStructCustomization::HandleTextBoxTextChanged)
 				.OnTextCommitted(this, &FTimespanStructCustomization::HandleTextBoxTextCommited)
 				.SelectAllTextOnCommit(true)
+				.Font( IPropertyTypeCustomizationUtils::GetRegularFont() )
 				.Text(this, &FTimespanStructCustomization::HandleTextBoxText)
 		];
 }
@@ -65,8 +66,14 @@ FText FTimespanStructCustomization::HandleTextBoxText( ) const
 	{
 		return LOCTEXT("MultipleValues", "Multiple Values");
 	}
-
-	return FText::FromString(((FTimespan*)RawData[0])->ToString());
+	else if( RawData[0] == nullptr )
+	{
+		return FText::GetEmpty();
+	}
+	else
+	{
+		return FText::FromString(((FTimespan*)RawData[0])->ToString());
+	}
 }
 
 
