Can hide the "Tick" category in Actor classes.

Fixed a crash with attempting to use the hidden category.

diff --git a/Engine/Source/Editor/DetailCustomizations/Private/ActorDetails.cpp b/Engine/Source/Editor/DetailCustomizations/Private/ActorDetails.cpp
--- a/Engine/Source/Editor/DetailCustomizations/Private/ActorDetails.cpp
+++ b/Engine/Source/Editor/DetailCustomizations/Private/ActorDetails.cpp
@@ -46,6 +46,10 @@ FActorDetails::~FActorDetails()
 
 void FActorDetails::CustomizeDetails( IDetailLayoutBuilder& DetailLayout )
 {
+	// Get the list of hidden categories
+	TArray<FString> HideCategories;
+	FEditorCategoryUtils::GetClassHideCategories(DetailLayout.GetDetailsView().GetBaseClass(), HideCategories); 
+
 	// These details only apply when adding an instance of the actor in a level
 	if( !DetailLayout.GetDetailsView().HasClassDefaultObject() && DetailLayout.GetDetailsView().GetSelectedActorInfo().NumSelected > 0 )
 	{
@@ -99,10 +103,6 @@ void FActorDetails::CustomizeDetails( IDetailLayoutBuilder& DetailLayout )
 
 		AddExperimentalWarningCategory(DetailLayout);
 
-		// Get the list of hidden categories
-		TArray<FString> HideCategories;
-		FEditorCategoryUtils::GetClassHideCategories(DetailLayout.GetDetailsView().GetBaseClass(), HideCategories); 
-
 		if (!HideCategories.Contains(TEXT("Transform")))
 		{
 			AddTransformCategory(DetailLayout);
@@ -125,7 +125,7 @@ void FActorDetails::CustomizeDetails( IDetailLayoutBuilder& DetailLayout )
 	TSharedPtr<IPropertyHandle> PrimaryTickProperty = DetailLayout.GetProperty(GET_MEMBER_NAME_CHECKED(AActor, PrimaryActorTick));
 
 	// Defaults only show tick properties
-	if (DetailLayout.GetDetailsView().HasClassDefaultObject())
+	if (DetailLayout.GetDetailsView().HasClassDefaultObject() && !HideCategories.Contains(TEXT("Tick")))
 	{
 		IDetailCategoryBuilder& TickCategory = DetailLayout.EditCategory("Tick");
 
