Turn auto weld on for shape components. UE-8812

diff --git a/Engine/Source/Editor/DetailCustomizations/Private/PrimitiveComponentDetails.cpp b/Engine/Source/Editor/DetailCustomizations/Private/PrimitiveComponentDetails.cpp
--- a/Engine/Source/Editor/DetailCustomizations/Private/PrimitiveComponentDetails.cpp
+++ b/Engine/Source/Editor/DetailCustomizations/Private/PrimitiveComponentDetails.cpp
@@ -123,7 +123,7 @@ EVisibility FPrimitiveComponentDetails::IsAutoWeldVisible() const
 {
 	for (int32 i = 0; i < ObjectsCustomized.Num(); ++i)
 	{
-		if (ObjectsCustomized[i].IsValid() && !ObjectsCustomized[i]->IsA(UStaticMeshComponent::StaticClass()))
+		if (ObjectsCustomized[i].IsValid() && !( ObjectsCustomized[i]->IsA(UStaticMeshComponent::StaticClass()) || ObjectsCustomized[i]->IsA(UShapeComponent::StaticClass())))
 		{
 			return EVisibility::Collapsed;
 		}
diff --git a/Engine/Source/Runtime/Engine/Private/ShapeComponent.cpp b/Engine/Source/Runtime/Engine/Private/ShapeComponent.cpp
--- a/Engine/Source/Runtime/Engine/Private/ShapeComponent.cpp
+++ b/Engine/Source/Runtime/Engine/Private/ShapeComponent.cpp
@@ -14,6 +14,7 @@ UShapeComponent::UShapeComponent(const FObjectInitializer& ObjectInitializer)
 	// since shape component is so popular for BP and so on, I'm adding manual support for compatibility
 	// this only works since this variable is getting deprecated. 
 	BodyInstance.ResponseToChannels_DEPRECATED.SetAllChannels(ECR_Block);
+	BodyInstance.bAutoWeld = true;	//UShapeComponent by default has auto welding
 
 	bHiddenInGame = true;
 	bCastDynamicShadow = false;
