Media: Fixed non-file based URLs (http://, rtsp:// etc.) getting mangled in MediaPlayer editor

diff --git a/Engine/Source/Editor/DetailCustomizations/Private/MediaPlayerCustomization.cpp b/Engine/Source/Editor/DetailCustomizations/Private/MediaPlayerCustomization.cpp
--- a/Engine/Source/Editor/DetailCustomizations/Private/MediaPlayerCustomization.cpp
+++ b/Engine/Source/Editor/DetailCustomizations/Private/MediaPlayerCustomization.cpp
@@ -388,7 +388,7 @@ FString FMediaPlayerCustomization::HandleUrlPickerFileTypeFilter() const
 
 void FMediaPlayerCustomization::HandleUrlPickerPathPicked( const FString& PickedPath )
 {
-	if (PickedPath.StartsWith(TEXT("./")))
+	if (PickedPath.StartsWith(TEXT("./")) || PickedPath.Contains(TEXT("://")))
 	{
 		UrlProperty->SetValue(PickedPath);
 	}
