Fixed aligment issues on Mass property in primitive component details

diff --git a/Engine/Source/Editor/DetailCustomizations/Private/PrimitiveComponentDetails.cpp b/Engine/Source/Editor/DetailCustomizations/Private/PrimitiveComponentDetails.cpp
--- a/Engine/Source/Editor/DetailCustomizations/Private/PrimitiveComponentDetails.cpp
+++ b/Engine/Source/Editor/DetailCustomizations/Private/PrimitiveComponentDetails.cpp
@@ -381,6 +381,7 @@ void FPrimitiveComponentDetails::AddPhysicsCategory(IDetailLayoutBuilder& Detail
 								[
 									SNew(SNumericEntryBox<float>)
 									.IsEnabled(false)
+									.Font( IDetailLayoutBuilder::GetDetailFont() )
 									.Value(this, &FPrimitiveComponentDetails::OnGetBodyMass)
 									.Visibility(this, &FPrimitiveComponentDetails::IsMassVisible, false)
 								]
@@ -400,6 +401,7 @@ void FPrimitiveComponentDetails::AddPhysicsCategory(IDetailLayoutBuilder& Detail
 
 							+ SHorizontalBox::Slot()
 							.AutoWidth()
+							.VAlign( VAlign_Center )
 							[
 								bDisplayMassOverride ? MassOverrideHandle->CreatePropertyValueWidget() : StaticCastSharedRef<SWidget>(SNew(SSpacer))
 							]
