UE-11523 DateTime variables will cause a crash on compile

diff --git a/Engine/Source/Editor/DetailCustomizations/Private/DateTimeStructCustomization.cpp b/Engine/Source/Editor/DetailCustomizations/Private/DateTimeStructCustomization.cpp
--- a/Engine/Source/Editor/DetailCustomizations/Private/DateTimeStructCustomization.cpp
+++ b/Engine/Source/Editor/DetailCustomizations/Private/DateTimeStructCustomization.cpp
@@ -66,7 +66,13 @@ FText FDateTimeStructCustomization::HandleTextBoxText( ) const
 		return LOCTEXT("MultipleValues", "Multiple Values");
 	}
 
-	return FText::FromString(((FDateTime*)RawData[0])->ToString());
+	auto DateTimePtr = static_cast<FDateTime*>(RawData[0]);
+	if (!DateTimePtr)
+	{
+		return FText::GetEmpty();
+	}
+
+	return FText::FromString(DateTimePtr->ToString());
 }
 
 
