Added customization to improve display of crossfade params

diff --git a/Engine/Source/Editor/DetailCustomizations/Private/DetailCustomizations.cpp b/Engine/Source/Editor/DetailCustomizations/Private/DetailCustomizations.cpp
--- a/Engine/Source/Editor/DetailCustomizations/Private/DetailCustomizations.cpp
+++ b/Engine/Source/Editor/DetailCustomizations/Private/DetailCustomizations.cpp
@@ -97,6 +97,7 @@
 #include "SceneCaptureDetails.h"
 #include "CurveColorCustomization.h"
 #include "ActorComponentDetails.h"
+#include "DistanceDatumStructCustomization.h"
 
 IMPLEMENT_MODULE( FDetailCustomizationsModule, DetailCustomizations );
 
@@ -194,6 +195,7 @@ void FDetailCustomizationsModule::RegisterPropertyTypeCustomizations()
 	RegisterCustomPropertyTypeLayout("ParticleSysParam", FOnGetPropertyTypeCustomizationInstance::CreateStatic(&FParticleSysParamStructCustomization::MakeInstance));
 	RegisterCustomPropertyTypeLayout("RawDistributionVector", FOnGetPropertyTypeCustomizationInstance::CreateStatic(&FRawDistributionVectorStructCustomization::MakeInstance));
 	RegisterCustomPropertyTypeLayout("CollisionProfileName", FOnGetPropertyTypeCustomizationInstance::CreateStatic(&FCollisionProfileNameCustomization::MakeInstance));
+	RegisterCustomPropertyTypeLayout("DistanceDatum", FOnGetPropertyTypeCustomizationInstance::CreateStatic(&FDistanceDatumStructCustomization::MakeInstance));
 }
 
 
diff --git a/Engine/Source/Editor/DetailCustomizations/Private/DistanceDatumStructCustomization.cpp b/Engine/Source/Editor/DetailCustomizations/Private/DistanceDatumStructCustomization.cpp
new file mode 100644
--- /dev/null
+++ b/Engine/Source/Editor/DetailCustomizations/Private/DistanceDatumStructCustomization.cpp
@@ -0,0 +1,92 @@
+// Copyright 1998-2015 Epic Games, Inc. All Rights Reserved.
+
+#include "DetailCustomizationsPrivatePCH.h"
+#include "DistanceDatumStructCustomization.h"
+#include "Sound/SoundNodeParamCrossFade.h"
+
+#define LOCTEXT_NAMESPACE "DistanceDatumStructCustomization"
+
+TSharedRef<IPropertyTypeCustomization> FDistanceDatumStructCustomization::MakeInstance()
+{
+	return MakeShareable(new FDistanceDatumStructCustomization);
+}
+
+FDistanceDatumStructCustomization::~FDistanceDatumStructCustomization()
+{
+}
+
+FDistanceDatumStructCustomization::FDistanceDatumStructCustomization()
+{
+}
+
+void FDistanceDatumStructCustomization::CustomizeHeader(TSharedRef<IPropertyHandle> InStructPropertyHandle, FDetailWidgetRow& HeaderRow, IPropertyTypeCustomizationUtils& StructCustomizationUtils)
+{
+	HeaderRow
+		.NameContent()
+		[
+			InStructPropertyHandle->CreatePropertyNameWidget()
+		]
+		.ValueContent()
+		[
+			InStructPropertyHandle->CreatePropertyValueWidget()
+		];
+}
+
+void FDistanceDatumStructCustomization::CustomizeChildren(TSharedRef<IPropertyHandle> InStructPropertyHandle, IDetailChildrenBuilder& StructBuilder, IPropertyTypeCustomizationUtils& StructCustomizationUtils)
+{
+	uint32 NumChildren;
+	InStructPropertyHandle->GetNumChildren(NumChildren);
+
+	TArray<UObject*> OuterObjects;
+	InStructPropertyHandle->GetOuterObjects(OuterObjects);
+
+	// Check whether we are dealing with param cross fade nodes
+	bool bAllParamNodes = true;
+	for (UObject* OuterObject : OuterObjects)
+	{
+		if (!OuterObject->IsA(USoundNodeParamCrossFade::StaticClass()))
+		{
+			bAllParamNodes = false;
+		}
+	}
+
+	// How to describe the parameter depending on how many are selected
+	FText ParamDesc = LOCTEXT("Param", "Param");
+	if (bAllParamNodes)
+	{
+		if (OuterObjects.Num() == 1)
+		{
+			FName ParamName = CastChecked<USoundNodeParamCrossFade>(OuterObjects[0])->ParamName;
+			if (ParamName != NAME_None)
+			{
+				ParamDesc = FText::FromName(ParamName);
+			}
+		}
+	}
+
+	for (uint32 ChildIndex = 0; ChildIndex < NumChildren; ++ChildIndex)
+	{
+		TSharedRef<IPropertyHandle> Child = InStructPropertyHandle->GetChildHandle(ChildIndex).ToSharedRef();
+		IDetailPropertyRow& ChildRow = StructBuilder.AddChildProperty(Child);
+
+		if (bAllParamNodes)
+		{
+			if (Child->GetProperty()->GetFName() == GET_MEMBER_NAME_CHECKED(FDistanceDatum, FadeInDistanceStart))
+			{
+				ChildRow.DisplayName(FText::Format(LOCTEXT("FadeInStart", "Fade In {0} Value Start"), ParamDesc));
+			}
+			else if (Child->GetProperty()->GetFName() == GET_MEMBER_NAME_CHECKED(FDistanceDatum, FadeInDistanceEnd))
+			{
+				ChildRow.DisplayName(FText::Format(LOCTEXT("FadeInEnd", "Fade In {0} Value End"), ParamDesc));
+			}
+			else if (Child->GetProperty()->GetFName() == GET_MEMBER_NAME_CHECKED(FDistanceDatum, FadeOutDistanceStart))
+			{
+				ChildRow.DisplayName(FText::Format(LOCTEXT("FadeOutStart", "Fade Out {0} Value Start"), ParamDesc));
+			}
+			else if (Child->GetProperty()->GetFName() == GET_MEMBER_NAME_CHECKED(FDistanceDatum, FadeOutDistanceEnd))
+			{
+				ChildRow.DisplayName(FText::Format(LOCTEXT("FadeOutEnd", "Fade Out {0} Value End"), ParamDesc));
+			}
+		}
+	}
+}
diff --git a/Engine/Source/Editor/DetailCustomizations/Private/DistanceDatumStructCustomization.h b/Engine/Source/Editor/DetailCustomizations/Private/DistanceDatumStructCustomization.h
new file mode 100644
--- /dev/null
+++ b/Engine/Source/Editor/DetailCustomizations/Private/DistanceDatumStructCustomization.h
@@ -0,0 +1,29 @@
+// Copyright 1998-2015 Epic Games, Inc. All Rights Reserved.
+
+#pragma once
+
+#include "IPropertyTypeCustomization.h"
+
+/**
+ * Customizes a Distance Datum struct to improve naming when used as a parameter
+ */
+class FDistanceDatumStructCustomization : public IPropertyTypeCustomization
+{
+public:
+	static TSharedRef<IPropertyTypeCustomization> MakeInstance();
+
+	/**
+	 * Destructor
+	 */
+	virtual ~FDistanceDatumStructCustomization();
+
+	/** IPropertyTypeCustomization interface */
+	virtual void CustomizeHeader( TSharedRef<IPropertyHandle> InStructPropertyHandle, FDetailWidgetRow& HeaderRow, IPropertyTypeCustomizationUtils& StructCustomizationUtils ) override;
+	virtual void CustomizeChildren( TSharedRef<IPropertyHandle> InStructPropertyHandle, IDetailChildrenBuilder& StructBuilder, IPropertyTypeCustomizationUtils& StructCustomizationUtils ) override;
+
+private:
+	/**
+	 * Constructor
+	 */
+	FDistanceDatumStructCustomization();
+};
