Added missing #undef LOCTEXT_NAMESPACE in DistanceDatumStructCustomization.cpp

diff --git a/Engine/Source/Editor/DetailCustomizations/Private/DistanceDatumStructCustomization.cpp b/Engine/Source/Editor/DetailCustomizations/Private/DistanceDatumStructCustomization.cpp
--- a/Engine/Source/Editor/DetailCustomizations/Private/DistanceDatumStructCustomization.cpp
+++ b/Engine/Source/Editor/DetailCustomizations/Private/DistanceDatumStructCustomization.cpp
@@ -90,3 +90,5 @@ void FDistanceDatumStructCustomization::CustomizeChildren(TSharedRef<IPropertyHa
 		}
 	}
 }
+
+#undef LOCTEXT_NAMESPACE
