Missing files from the Texture LOD Settings changelist/

diff --git a/Engine/Source/Editor/DetailCustomizations/Private/TextureLODSettingsDetails.cpp b/Engine/Source/Editor/DetailCustomizations/Private/TextureLODSettingsDetails.cpp
new file mode 100644
--- /dev/null
+++ b/Engine/Source/Editor/DetailCustomizations/Private/TextureLODSettingsDetails.cpp
@@ -0,0 +1,387 @@
+// Copyright 1998-2015 Epic Games, Inc. All Rights Reserved.
+
+#include "DetailCustomizationsPrivatePCH.h"
+#include "TextureLODSettingsDetails.h"
+
+#define LOCTEXT_NAMESPACE "TextureLODSettingsDetails"
+
+
+////////////////////////////////////////////////
+// FDeviceProfileTextureLODSettingsDetails
+
+
+FDeviceProfileTextureLODSettingsDetails::FDeviceProfileTextureLODSettingsDetails(IDetailLayoutBuilder* InDetailBuilder)
+	: DetailBuilder(InDetailBuilder)
+{
+	TextureLODSettingsPropertyNameHandle = DetailBuilder->GetProperty("TextureLODGroups", UTextureLODSettings::StaticClass());
+	LODGroupsArrayHandle = TextureLODSettingsPropertyNameHandle->AsArray();
+
+	TArray<UObject*> OuterObjects;
+	TextureLODSettingsPropertyNameHandle->GetOuterObjects(OuterObjects);
+	if (OuterObjects.Num() == 1)
+	{
+		DeviceProfile = CastChecked<UDeviceProfile>(OuterObjects[0]);
+	}
+}
+
+
+void FDeviceProfileTextureLODSettingsDetails::CreateTextureGroupEntryRow(int32 GroupId, IDetailCategoryBuilder& DetailCategoryBuilder)
+{
+	TSharedRef<IPropertyHandle> LODGroupElementHandle = LODGroupsArrayHandle->GetElement(GroupId);
+
+	DetailCategoryBuilder.AddCustomBuilder(MakeShareable(new FTextureLODGroupLayout(DeviceProfile, (TextureGroup)GroupId)));
+}
+
+
+void FDeviceProfileTextureLODSettingsDetails::CreateTextureLODSettingsPropertyView()
+{
+	DetailBuilder->HideProperty(TextureLODSettingsPropertyNameHandle);
+
+	IDetailCategoryBuilder& TextureLODSettingsDetailCategory = DetailBuilder->EditCategory("Texture LOD Settings");
+
+#define SETUPLODGROUP(GroupId) CreateTextureGroupEntryRow(GroupId, TextureLODSettingsDetailCategory);
+	FOREACH_ENUM_TEXTUREGROUP(SETUPLODGROUP)
+#undef SETUPLODGROUP
+}
+
+
+////////////////////////////////////////////////
+// FTextureLODGroupLayout
+
+FTextureLODGroupLayout::FTextureLODGroupLayout(const UDeviceProfile* InDeviceProfile, TextureGroup InGroupId)
+{
+	LodGroup = &InDeviceProfile->GetTextureLODSettings()->GetTextureLODGroup(InGroupId);
+
+#define POPULATE_MIPGENSETTINGS(MipGenSettingsId) AddToAvailableMipGenSettings(MipGenSettingsId);
+	FOREACH_ENUM_TEXTUREMIPGENSETTINGS(POPULATE_MIPGENSETTINGS)
+#undef POPULATE_MIPGENSETTINGS
+}
+
+
+FTextureLODGroupLayout::~FTextureLODGroupLayout()
+{
+}
+
+
+void FTextureLODGroupLayout::AddToAvailableMipGenSettings(TextureMipGenSettings MipGenSettingsId)
+{
+	MipGenSettingsComboList.Add(MakeShareable(new TextureMipGenSettings(MipGenSettingsId)));
+}
+
+void FTextureLODGroupLayout::GenerateHeaderRowContent(FDetailWidgetRow& NodeRow)
+{
+	UEnum* TextureGroupEnum = FindObject<UEnum>(NULL, TEXT("/Script/Engine.TextureGroup"));
+	const FString LODGroupName = TextureGroupEnum->GetMetaData(TEXT("DisplayName"), LodGroup->Group);
+
+	NodeRow.NameContent()
+		[
+			SNew(STextBlock)
+			.Text(FText::FromString(LODGroupName))
+			.Font(IDetailLayoutBuilder::GetDetailFont())
+		];
+}
+
+
+BEGIN_SLATE_FUNCTION_BUILD_OPTIMIZATION
+void FTextureLODGroupLayout::GenerateChildContent(IDetailChildrenBuilder& ChildrenBuilder)
+{
+	// Min and Max LOD properties
+	{
+		ChildrenBuilder.AddChildContent(LOCTEXT("MinLODSize", "Min LOD Size"))
+			.NameContent()
+			[
+				SNew(STextBlock)
+				.Font(IDetailLayoutBuilder::GetDetailFont())
+				.Text(LOCTEXT("MinLODSize", "Min LOD Size"))
+			]
+			.ValueContent()
+			[
+				SNew(SSpinBox<uint32>)
+				.Font(IDetailLayoutBuilder::GetDetailFont())
+				.MinValue(2)
+				.MaxValue(4096)
+				.Value(this, &FTextureLODGroupLayout::GetMinLODSize)
+				.OnValueChanged(this, &FTextureLODGroupLayout::OnMinLODSizeChanged)
+				.OnValueCommitted(this, &FTextureLODGroupLayout::OnMinLODSizeCommitted)
+			];
+	}
+
+	{
+		ChildrenBuilder.AddChildContent(LOCTEXT("MaxLODSize", "Max LOD Size"))
+			.NameContent()
+			[
+				SNew(STextBlock)
+				.Font(IDetailLayoutBuilder::GetDetailFont())
+				.Text(LOCTEXT("MaxLODSize", "Max LOD Size"))
+			]
+			.ValueContent()
+			[
+				SNew(SSpinBox<uint32>)
+				.Font(IDetailLayoutBuilder::GetDetailFont())
+				.MinValue(2)
+				.MaxValue(4096)
+				.Value(this, &FTextureLODGroupLayout::GetMaxLODSize)
+				.OnValueChanged(this, &FTextureLODGroupLayout::OnMaxLODSizeChanged)
+				.OnValueCommitted(this, &FTextureLODGroupLayout::OnMaxLODSizeCommitted)
+			];
+	}
+
+
+	// LOD Bias
+	{
+		ChildrenBuilder.AddChildContent(LOCTEXT("LODBias", "LOD Bias"))
+			.NameContent()
+			[
+				SNew(STextBlock)
+				.Font(IDetailLayoutBuilder::GetDetailFont())
+				.Text(LOCTEXT("LODBias", "LOD Bias"))
+			]
+			.ValueContent()
+			[
+				SNew(SSpinBox<uint32>)
+				.Font(IDetailLayoutBuilder::GetDetailFont())
+				.MinValue(1)
+				.MaxValue(100)
+				.Value(this, &FTextureLODGroupLayout::GetLODBias)
+				.OnValueChanged(this, &FTextureLODGroupLayout::OnLODBiasChanged)
+				.OnValueCommitted(this, &FTextureLODGroupLayout::OnLODBiasCommitted)
+			];
+	}
+
+
+	// Filter properties
+
+	TSharedPtr<FName> NamePointPtr = MakeShareable(new FName(NAME_Point));
+	FilterComboList.Add(NamePointPtr);
+	TSharedPtr<FName> NameLinearPtr = MakeShareable(new FName(NAME_Linear));
+	FilterComboList.Add(NameLinearPtr);
+	TSharedPtr<FName> NameAnisoPtr = MakeShareable(new FName(NAME_Aniso));
+	FilterComboList.Add(NameAnisoPtr);
+	{
+		ChildrenBuilder.AddChildContent(LOCTEXT("MinMagFilter", "MinMag Filter"))
+			.NameContent()
+			[
+				SNew(STextBlock)
+				.Font(IDetailLayoutBuilder::GetDetailFont())
+				.Text(LOCTEXT("MinMagFilter", "MinMag Filter"))
+			]
+			.ValueContent()
+			[
+				SNew(SComboBox< TSharedPtr<FName> >)
+				.OptionsSource(&FilterComboList)
+				.OnGenerateWidget(this, &FTextureLODGroupLayout::MakeMinMagFilterComboWidget)
+				.OnSelectionChanged(this, &FTextureLODGroupLayout::OnMinMagFilterChanged)
+				.InitiallySelectedItem(FilterComboList[0])
+				.ContentPadding(0)
+				.Content()
+				[
+					SNew(STextBlock)
+					.Text(this, &FTextureLODGroupLayout::GetMinMagFilterComboBoxContent)
+					.Font(IDetailLayoutBuilder::GetDetailFont())
+					.ToolTipText(this, &FTextureLODGroupLayout::GetMinMagFilterComboBoxToolTip)
+				]
+			];
+	}
+
+	{
+		ChildrenBuilder.AddChildContent(LOCTEXT("MipFilter", "Mip Filter"))
+			.NameContent()
+			[
+				SNew(STextBlock)
+				.Font(IDetailLayoutBuilder::GetDetailFont())
+				.Text(LOCTEXT("MipFilter", "Mip Filter"))
+			]
+		.ValueContent()
+			[
+				SNew(SComboBox< TSharedPtr<FName> >)
+				.OptionsSource(&FilterComboList)
+				.OnGenerateWidget(this, &FTextureLODGroupLayout::MakeMipFilterComboWidget)
+				.OnSelectionChanged(this, &FTextureLODGroupLayout::OnMipFilterChanged)
+				.InitiallySelectedItem(FilterComboList[0])
+				.ContentPadding(0)
+				.Content()
+				[
+					SNew(STextBlock)
+					.Text(this, &FTextureLODGroupLayout::GetMipFilterComboBoxContent)
+					.Font(IDetailLayoutBuilder::GetDetailFont())
+					.ToolTipText(this, &FTextureLODGroupLayout::GetMipFilterComboBoxToolTip)
+				]
+			];
+	}
+
+
+	// Mip Gen Settings
+	{
+		ChildrenBuilder.AddChildContent(LOCTEXT("MipGenSettings", "Mip Gen Settings"))
+			.NameContent()
+			[
+				SNew(STextBlock)
+				.Font(IDetailLayoutBuilder::GetDetailFont())
+				.Text(LOCTEXT("MipGenSettings", "Mip Gen Settings"))
+			]
+		.ValueContent()
+			[
+				SNew(SComboBox< TSharedPtr<TextureMipGenSettings> >)
+				.OptionsSource(&MipGenSettingsComboList)
+				.OnGenerateWidget(this, &FTextureLODGroupLayout::MakeMipGenSettingsComboWidget)
+				.OnSelectionChanged(this, &FTextureLODGroupLayout::OnMipGenSettingsChanged)
+				.InitiallySelectedItem(MipGenSettingsComboList[0])
+				.ContentPadding(0)
+				.Content()
+				[
+					SNew(STextBlock)
+					.Text(this, &FTextureLODGroupLayout::GetMipGenSettingsComboBoxContent)
+					.Font(IDetailLayoutBuilder::GetDetailFont())
+					.ToolTipText(this, &FTextureLODGroupLayout::GetMipGenSettingsComboBoxToolTip)
+				]
+			];
+	}
+
+}
+END_SLATE_FUNCTION_BUILD_OPTIMIZATION
+
+
+uint32 FTextureLODGroupLayout::GetMinLODSize() const
+{
+	return LodGroup->MinLODSize;
+}
+
+void FTextureLODGroupLayout::OnMinLODSizeChanged(uint32 NewValue)
+{
+	LodGroup->MinLODSize = NewValue;
+}
+
+void FTextureLODGroupLayout::OnMinLODSizeCommitted(uint32 NewValue, ETextCommit::Type TextCommitType)
+{
+	if (FEngineAnalytics::IsAvailable())
+	{
+//		FEngineAnalytics::GetProvider().RecordEvent(TEXT("Editor.Usage.StaticMesh.ReductionSettings"), TEXT("MinLODSize"), FString::Printf(TEXT("%.1f"), NewValue));
+	}
+	OnMinLODSizeChanged(NewValue);
+}
+
+
+uint32 FTextureLODGroupLayout::GetMaxLODSize() const
+{
+	return LodGroup->MaxLODSize;
+}
+
+void FTextureLODGroupLayout::OnMaxLODSizeChanged(uint32 NewValue)
+{
+	LodGroup->MaxLODSize = NewValue;
+}
+
+void FTextureLODGroupLayout::OnMaxLODSizeCommitted(uint32 NewValue, ETextCommit::Type TextCommitType)
+{
+	if (FEngineAnalytics::IsAvailable())
+	{
+		//		FEngineAnalytics::GetProvider().RecordEvent(TEXT("Editor.Usage.StaticMesh.ReductionSettings"), TEXT("MaxLODSize"), FString::Printf(TEXT("%.1f"), NewValue));
+	}
+	OnMaxLODSizeChanged(NewValue);
+}
+
+
+uint32 FTextureLODGroupLayout::GetLODBias() const
+{
+	return LodGroup->LODBias;
+}
+
+void FTextureLODGroupLayout::OnLODBiasChanged(uint32 NewValue)
+{
+	LodGroup->LODBias = NewValue;
+}
+
+void FTextureLODGroupLayout::OnLODBiasCommitted(uint32 NewValue, ETextCommit::Type TextCommitType)
+{
+	if (FEngineAnalytics::IsAvailable())
+	{
+		//		FEngineAnalytics::GetProvider().RecordEvent(TEXT("Editor.Usage.StaticMesh.ReductionSettings"), TEXT("LODBias"), FString::Printf(TEXT("%.1f"), NewValue));
+	}
+	OnLODBiasChanged(NewValue);
+}
+
+
+TSharedRef<SWidget> FTextureLODGroupLayout::MakeMinMagFilterComboWidget(TSharedPtr<FName> InItem)
+{
+	return SNew(STextBlock).Text(FText::FromName(*InItem)).Font(IDetailLayoutBuilder::GetDetailFont());
+}
+
+void FTextureLODGroupLayout::OnMinMagFilterChanged(TSharedPtr<FName> NewSelection, ESelectInfo::Type SelectInfo)
+{
+	// if it's set from code, we did that on purpose
+	if (SelectInfo != ESelectInfo::Direct)
+	{
+		FName NewValue = *NewSelection.Get();
+		LodGroup->MinMagFilter = NewValue;
+	}
+}
+
+FText FTextureLODGroupLayout::GetMinMagFilterComboBoxToolTip() const
+{
+	return LOCTEXT("MinMagFilterComboToolTip", "");
+}
+
+FText FTextureLODGroupLayout::GetMinMagFilterComboBoxContent() const
+{
+	return FText::FromName(LodGroup->MinMagFilter);
+}
+
+
+TSharedRef<SWidget> FTextureLODGroupLayout::MakeMipFilterComboWidget(TSharedPtr<FName> InItem)
+{
+	return SNew(STextBlock).Text(FText::FromName(*InItem)).Font(IDetailLayoutBuilder::GetDetailFont());
+}
+
+void FTextureLODGroupLayout::OnMipFilterChanged(TSharedPtr<FName> NewSelection, ESelectInfo::Type SelectInfo)
+{
+	// if it's set from code, we did that on purpose
+	if (SelectInfo != ESelectInfo::Direct)
+	{
+		FName NewValue = *NewSelection.Get();
+		LodGroup->MipFilter = NewValue;
+	}
+}
+
+FText FTextureLODGroupLayout::GetMipFilterComboBoxToolTip() const
+{
+	return LOCTEXT("MipFilterComboToolTip", "");
+}
+
+FText FTextureLODGroupLayout::GetMipFilterComboBoxContent() const
+{
+	return FText::FromName(LodGroup->MipFilter);
+}
+
+
+TSharedRef<SWidget> FTextureLODGroupLayout::MakeMipGenSettingsComboWidget(TSharedPtr<TextureMipGenSettings> InItem)
+{
+	UEnum* TextureGroupEnum = FindObject<UEnum>(NULL, TEXT("/Script/Engine.TextureMipGenSettings"));
+	const FString MipGenSettingsName = TextureGroupEnum->GetMetaData(TEXT("DisplayName"), *InItem.Get());
+
+	return SNew(STextBlock).Text(FText::FromString(MipGenSettingsName)).Font(IDetailLayoutBuilder::GetDetailFont());
+}
+
+void FTextureLODGroupLayout::OnMipGenSettingsChanged(TSharedPtr<TextureMipGenSettings> NewSelection, ESelectInfo::Type SelectInfo)
+{
+	// if it's set from code, we did that on purpose
+	if (SelectInfo != ESelectInfo::Direct)
+	{
+		TextureMipGenSettings NewValue = *NewSelection.Get();
+		LodGroup->MipGenSettings = NewValue;
+	}
+}
+
+FText FTextureLODGroupLayout::GetMipGenSettingsComboBoxToolTip() const
+{
+	return LOCTEXT("MipGenSettingsComboToolTip", "");
+}
+
+FText FTextureLODGroupLayout::GetMipGenSettingsComboBoxContent() const
+{
+	UEnum* TextureGroupEnum = FindObject<UEnum>(NULL, TEXT("/Script/Engine.TextureMipGenSettings"));
+	const FString MipGenSettingsName = TextureGroupEnum->GetMetaData(TEXT("DisplayName"), LodGroup->MipGenSettings);
+	
+	return FText::FromString(MipGenSettingsName);
+}
+
+#undef LOCTEXT_NAMESPACE
\ No newline at end of file
diff --git a/Engine/Source/Editor/DetailCustomizations/Private/TextureLODSettingsDetails.h b/Engine/Source/Editor/DetailCustomizations/Private/TextureLODSettingsDetails.h
new file mode 100644
--- /dev/null
+++ b/Engine/Source/Editor/DetailCustomizations/Private/TextureLODSettingsDetails.h
@@ -0,0 +1,123 @@
+// Copyright 1998-2015 Epic Games, Inc. All Rights Reserved.
+
+#pragma once
+
+
+/**
+ * Texture Group layout for property editor views.
+ */
+class FTextureLODGroupLayout : public IDetailCustomNodeBuilder, public TSharedFromThis < FTextureLODGroupLayout >
+{
+public:
+	FTextureLODGroupLayout(const class UDeviceProfile* InDeviceProfile, TextureGroup GroupId);
+	virtual ~FTextureLODGroupLayout();
+private:
+	/** Begin IDetailCustomNodeBuilder Interface*/
+	virtual void SetOnRebuildChildren(FSimpleDelegate InOnRegenerateChildren) override {}
+	virtual void GenerateHeaderRowContent(FDetailWidgetRow& NodeRow) override;
+	virtual void GenerateChildContent(IDetailChildrenBuilder& ChildrenBuilder) override;
+	virtual void Tick(float DeltaTime) override{}
+	virtual bool RequiresTick() const override { return false; }
+	virtual FName GetName() const override { static FName TextureLODGroup("TextureLODGroup"); return TextureLODGroup; }
+	virtual bool InitiallyCollapsed() const override { return true; }
+	/** End IDetailCustomNodeBuilder Interface*/
+
+private:
+
+	// Controls for the Min LOD Size property editing
+	uint32 GetMinLODSize() const;
+	void OnMinLODSizeChanged(uint32 NewValue);
+	void OnMinLODSizeCommitted(uint32 NewValue, ETextCommit::Type TextCommitType);
+
+	// Controls for the Max LOD Size property editing
+	uint32 GetMaxLODSize() const;
+	void OnMaxLODSizeChanged(uint32 NewValue);
+	void OnMaxLODSizeCommitted(uint32 NewValue, ETextCommit::Type TextCommitType);
+
+	// Controls for the LOD Bias property editing
+	uint32 GetLODBias() const;
+	void OnLODBiasChanged(uint32 NewValue);
+	void OnLODBiasCommitted(uint32 NewValue, ETextCommit::Type TextCommitType);
+
+	// Controls for the MinMag Filter property editing
+	TSharedRef<SWidget> MakeMinMagFilterComboWidget(TSharedPtr<FName> InItem);
+	void OnMinMagFilterChanged(TSharedPtr<FName> NewSelection, ESelectInfo::Type SelectInfo);
+	FText GetMinMagFilterComboBoxContent() const;
+	FText GetMinMagFilterComboBoxToolTip() const;
+
+	// Controls for the Mip Filter property editing
+	TSharedRef<SWidget> MakeMipFilterComboWidget(TSharedPtr<FName> InItem);
+	void OnMipFilterChanged(TSharedPtr<FName> NewSelection, ESelectInfo::Type SelectInfo);
+	FText GetMipFilterComboBoxContent() const;
+	FText GetMipFilterComboBoxToolTip() const;
+
+	// Controls for the MipGenSettings property editing
+	TSharedRef<SWidget> MakeMipGenSettingsComboWidget(TSharedPtr<TextureMipGenSettings> InItem);
+	void OnMipGenSettingsChanged(TSharedPtr<TextureMipGenSettings> NewSelection, ESelectInfo::Type SelectInfo);
+	FText GetMipGenSettingsComboBoxContent() const;
+	FText GetMipGenSettingsComboBoxToolTip() const;
+
+private:
+	// The LOD group we are creating an editor for.
+	FTextureLODGroup* LodGroup;
+
+	// The available filter names for the MinMag and Mip Filter selection
+	TArray<TSharedPtr<FName>> FilterComboList;
+
+private:
+
+	// Populate the Mip Gen Settings Combo Options
+	void AddToAvailableMipGenSettings(TextureMipGenSettings MipGenSettingsId);
+
+	// The available MipGenSettingsavailable for selection
+	TArray<TSharedPtr<TextureMipGenSettings>> MipGenSettingsComboList;
+};
+
+/**
+* Details panel for Texture LOD Settings
+*/
+class FDeviceProfileTextureLODSettingsDetails
+	: public TSharedFromThis<FDeviceProfileTextureLODSettingsDetails>
+{
+
+public:
+
+	/**
+	 * Constructor for the parent property details view
+	 *
+	 * @param InDetailsBuilder Where we are adding our property view to
+	 */
+	FDeviceProfileTextureLODSettingsDetails(IDetailLayoutBuilder* InDetailBuilder);
+
+
+	/**
+	 * Create the parent property view for the device profile
+	 */
+	void CreateTextureLODSettingsPropertyView();
+
+	/**
+	 * Create an editor for the provided LOD Group
+	 */
+	void CreateTextureGroupEntryRow(int32 GroupId, IDetailCategoryBuilder& DetailGroup);
+	
+private:
+
+	/**
+	 * Delegate used when the device profiles parent is updated from any source.
+	 */
+	void OnTextureLODSettingsPropertyChanged();
+
+private:
+
+	/** A handle to the detail view builder */
+	IDetailLayoutBuilder* DetailBuilder;
+
+	/** Access to the Parent Property */
+	TSharedPtr<IPropertyHandle> TextureLODSettingsPropertyNameHandle;
+
+	/** Access to the LOD Groups array */
+	TSharedPtr<IPropertyHandleArray> LODGroupsArrayHandle;
+
+	/** A reference to the object we are showing these properties for */
+	class UDeviceProfile* DeviceProfile;
+};
\ No newline at end of file
