Details customizations for auto reimport structures

diff --git a/Engine/Source/Editor/DetailCustomizations/Private/AutoReimportDirectoryCustomization.cpp b/Engine/Source/Editor/DetailCustomizations/Private/AutoReimportDirectoryCustomization.cpp
new file mode 100644
--- /dev/null
+++ b/Engine/Source/Editor/DetailCustomizations/Private/AutoReimportDirectoryCustomization.cpp
@@ -0,0 +1,293 @@
+// Copyright 1998-2015 Epic Games, Inc. All Rights Reserved.
+
+#include "DetailCustomizationsPrivatePCH.h"
+#include "AutoReimportDirectoryCustomization.h"
+
+#include "DesktopPlatformModule.h"
+
+#define LOCTEXT_NAMESPACE "FAutoReimportDirectoryCustomization"
+
+void FAutoReimportWildcardCustomization::CustomizeHeader(TSharedRef<IPropertyHandle> InPropertyHandle, FDetailWidgetRow& InHeaderRow, IPropertyTypeCustomizationUtils& InStructCustomizationUtils)
+{
+	PropertyHandle = InPropertyHandle;
+
+	WildcardProperty = PropertyHandle->GetChildHandle(GET_MEMBER_NAME_CHECKED(FAutoReimportWildcard, Wildcard));
+	IncludeProperty = PropertyHandle->GetChildHandle(GET_MEMBER_NAME_CHECKED(FAutoReimportWildcard, bInclude));
+	
+	InHeaderRow
+	.NameContent()
+	[
+		InPropertyHandle->CreatePropertyNameWidget()
+	]
+	.ValueContent()
+	.MaxDesiredWidth(150.0f)
+	.HAlign(HAlign_Fill)
+	.VAlign(VAlign_Center)
+	[
+		SNew(SHorizontalBox)
+
+		+SHorizontalBox::Slot()
+		.VAlign(VAlign_Center)
+		.Padding(FMargin(0,0,4.f,0))
+		[
+			SNew(SEditableTextBox)
+			.Text(this, &FAutoReimportWildcardCustomization::GetWildcardText)
+			.OnTextChanged(this, &FAutoReimportWildcardCustomization::OnWildcardChanged)
+			.OnTextCommitted(this, &FAutoReimportWildcardCustomization::OnWildcardCommitted)
+		]
+
+		+SHorizontalBox::Slot()
+		.AutoWidth()
+		[
+			SNew(SCheckBox)
+			.IsChecked(this, &FAutoReimportWildcardCustomization::GetCheckState)
+			.OnCheckStateChanged(this, &FAutoReimportWildcardCustomization::OnCheckStateChanged)
+			.Content()
+			[
+				SNew(STextBlock)
+				.Text(LOCTEXT("Include_Label", "Include?"))
+				.ToolTipText(LOCTEXT("Include_ToolTip", "When checked, files that match the wildcard will be included, otherwise files that match will be excluded from the monitor."))
+			]
+		]
+	];
+}
+
+void FAutoReimportWildcardCustomization::CustomizeChildren(TSharedRef<IPropertyHandle> InStructPropertyHandle, IDetailChildrenBuilder& InStructBuilder, IPropertyTypeCustomizationUtils& InStructCustomizationUtils)
+{
+}
+
+FText FAutoReimportWildcardCustomization::GetWildcardText() const
+{
+	FText Ret;
+	WildcardProperty->GetValueAsFormattedText(Ret);
+	return Ret;
+}
+
+void FAutoReimportWildcardCustomization::OnWildcardCommitted(const FText& InValue, ETextCommit::Type CommitType)
+{
+	WildcardProperty->SetValue(InValue.ToString());
+}
+
+void FAutoReimportWildcardCustomization::OnWildcardChanged(const FText& InValue)
+{
+	WildcardProperty->SetValue(InValue.ToString());
+}
+
+ECheckBoxState FAutoReimportWildcardCustomization::GetCheckState() const
+{
+	bool bChecked = true;
+	IncludeProperty->GetValue(bChecked);
+	return bChecked ? ECheckBoxState::Checked : ECheckBoxState::Unchecked;
+}
+
+void FAutoReimportWildcardCustomization::OnCheckStateChanged(ECheckBoxState InState)
+{
+	IncludeProperty->SetValue(InState == ECheckBoxState::Checked);
+}
+
+void FAutoReimportDirectoryCustomization::CustomizeHeader(TSharedRef<IPropertyHandle> InPropertyHandle, FDetailWidgetRow& InHeaderRow, IPropertyTypeCustomizationUtils& InStructCustomizationUtils)
+{
+	PropertyHandle = InPropertyHandle;
+
+	SourceDirProperty = PropertyHandle->GetChildHandle(GET_MEMBER_NAME_CHECKED(FAutoReimportDirectoryConfig, SourceDirectory));
+	MountPointProperty = PropertyHandle->GetChildHandle(GET_MEMBER_NAME_CHECKED(FAutoReimportDirectoryConfig, MountPoint));
+	WildcardsProperty = PropertyHandle->GetChildHandle(GET_MEMBER_NAME_CHECKED(FAutoReimportDirectoryConfig, Wildcards));
+
+	FString ExistingMountPath;
+	MountPointProperty->GetValueAsFormattedString(ExistingMountPath);
+	MountPathVisibility = ExistingMountPath.IsEmpty() ? EVisibility::Collapsed : EVisibility::Visible;
+	
+	InHeaderRow
+	.NameContent()
+	[
+		InPropertyHandle->CreatePropertyNameWidget()
+	]
+	.ValueContent()
+	.MaxDesiredWidth(150.0f)
+	.HAlign(HAlign_Fill)
+	.VAlign(VAlign_Center)
+	[
+		SNew(SHorizontalBox)
+
+		+SHorizontalBox::Slot()
+		.VAlign(VAlign_Center)
+		.Padding(FMargin(0,0,4.f,0))
+		[
+			SNew(SEditableTextBox)
+			.Text(this, &FAutoReimportDirectoryCustomization::GetDirectoryText)
+			.OnTextChanged(this, &FAutoReimportDirectoryCustomization::OnDirectoryChanged)
+			.OnTextCommitted(this, &FAutoReimportDirectoryCustomization::OnDirectoryCommitted)
+		]
+
+		+SHorizontalBox::Slot()
+		.AutoWidth()
+		[
+			SNew(SButton)
+			.OnClicked(this, &FAutoReimportDirectoryCustomization::BrowseForFolder)
+			.ToolTipText(LOCTEXT("BrowseForDirectory", "Browse for a directory"))
+			.Text(LOCTEXT("Browse", "Browse"))
+		]
+	];
+}
+
+void FAutoReimportDirectoryCustomization::CustomizeChildren(TSharedRef<IPropertyHandle> InStructPropertyHandle, IDetailChildrenBuilder& InStructBuilder, IPropertyTypeCustomizationUtils& InStructCustomizationUtils)
+{
+	{
+		FDetailWidgetRow& DetailRow = InStructBuilder.AddChildContent(LOCTEXT("MountPathName", "Mount Path"));
+
+		DetailRow.Visibility(TAttribute<EVisibility>(this, &FAutoReimportDirectoryCustomization::GetMountPathVisibility));
+
+		DetailRow.NameContent()
+		[
+			SNew(STextBlock)
+			.Font(IDetailLayoutBuilder::GetDetailFont())
+			.Text(LOCTEXT("MountPath_Label", "Map Directory To"))
+			.ToolTipText(LOCTEXT("MountPathToolTip", "Specify a mount path to which this physical directory relates. Any new files added on disk will be imported into this virtual path."))
+		];
+
+		DetailRow.ValueContent()
+		.MaxDesiredWidth(150.0f)
+		.HAlign(HAlign_Fill)
+		.VAlign(VAlign_Center)
+		[
+			SNew(SHorizontalBox)
+
+			+SHorizontalBox::Slot()
+			[
+				SNew(SOverlay)
+
+				+ SOverlay::Slot()
+				[
+					SNew(SEditableTextBox)
+					.Text(this, &FAutoReimportDirectoryCustomization::GetMountPointText)
+					.OnTextChanged(this, &FAutoReimportDirectoryCustomization::OnMountPointChanged)
+					.OnTextCommitted(this, &FAutoReimportDirectoryCustomization::OnMountPointCommitted)
+				]
+			]
+		];
+	}
+
+	InStructBuilder.AddChildProperty(WildcardsProperty.ToSharedRef());
+}
+
+EVisibility FAutoReimportDirectoryCustomization::GetMountPathVisibility() const
+{
+	return MountPathVisibility;
+}
+
+FText FAutoReimportDirectoryCustomization::GetDirectoryText() const
+{
+	FText Ret;
+	SourceDirProperty->GetValueAsFormattedText(Ret);
+	return Ret;
+}
+
+void FAutoReimportDirectoryCustomization::OnDirectoryCommitted(const FText& InValue, ETextCommit::Type CommitType)
+{
+	SetSourcePath(InValue.ToString());
+}
+
+void FAutoReimportDirectoryCustomization::OnDirectoryChanged(const FText& InValue)
+{
+	SetSourcePath(InValue.ToString());
+}
+
+FText FAutoReimportDirectoryCustomization::GetMountPointText() const
+{
+	FText Ret;
+	MountPointProperty->GetValueAsFormattedText(Ret);
+	return Ret;
+}
+
+void FAutoReimportDirectoryCustomization::OnMountPointCommitted(const FText& InValue, ETextCommit::Type CommitType)
+{
+	MountPointProperty->SetValue(InValue.ToString());
+}
+
+void FAutoReimportDirectoryCustomization::OnMountPointChanged(const FText& InValue)
+{
+	MountPointProperty->SetValue(InValue.ToString());
+}
+
+FReply FAutoReimportDirectoryCustomization::BrowseForFolder()
+{
+	FString InitialDir;
+	SourceDirProperty->GetValue(InitialDir);
+
+	if (InitialDir.IsEmpty())
+	{
+		InitialDir = FPaths::ConvertRelativePathToFull(FPaths::GameContentDir());
+	}
+	else if (!FPackageName::GetPackageMountPoint(InitialDir).IsNone())
+	{
+		InitialDir = FPaths::ConvertRelativePathToFull(FPackageName::LongPackageNameToFilename(InitialDir));
+	}
+
+	IDesktopPlatform* DesktopPlatform = FDesktopPlatformModule::Get();
+	if ( DesktopPlatform )
+	{
+		void* Window = nullptr;
+
+		// IMainFrameModule& MainFrameModule = FModuleManager::LoadModuleChecked<IMainFrameModule>(TEXT("MainFrame"));
+		// const TSharedPtr<SWindow>& MainFrameParentWindow = MainFrameModule.GetParentWindow();
+		// if ( MainFrameParentWindow.IsValid() && MainFrameParentWindow->GetNativeWindow().IsValid() )
+		// {
+		// 	ParentWindowWindowHandle = MainFrameParentWindow->GetNativeWindow()->GetOSWindowHandle();
+		// }
+
+		FString FolderName;
+		const FString Title = LOCTEXT("BrowseForFolderTitle", "Choose a directory to monitor").ToString();
+		const bool bFolderSelected = DesktopPlatform->OpenDirectoryDialog(Window, Title, InitialDir, FolderName);
+		if (bFolderSelected)
+		{
+			FolderName /= TEXT("");
+			SetSourcePath(FolderName);
+		}
+	}
+
+	return FReply::Handled();
+}
+
+void FAutoReimportDirectoryCustomization::SetSourcePath(FString InSourceDir)
+{
+	MountPathVisibility = EVisibility::Visible;
+
+	// Don't log errors and warnings
+	FAutoReimportDirectoryConfig::FParseContext Context(false);
+
+	FString MountPoint;
+	if (FAutoReimportDirectoryConfig::ParseSourceDirectoryAndMountPoint(InSourceDir, MountPoint, Context))
+	{
+		// If we have a mount point from the source path, we use that as the source path
+		if (!MountPoint.IsEmpty())
+		{
+			SourceDirProperty->SetValue(MountPoint);
+			MountPointProperty->SetValue(FString());
+			MountPathVisibility = EVisibility::Collapsed;
+		}
+		else
+		{
+			FString ExistingMountPath;
+			MountPointProperty->GetValue(ExistingMountPath);
+
+			if (ExistingMountPath.IsEmpty())
+			{
+				MountPointProperty->SetValue(MountPoint);	
+			}
+			SourceDirProperty->SetValue(InSourceDir);
+		}
+	}
+	else
+	{
+		SourceDirProperty->SetValue(InSourceDir);
+		MountPathVisibility = InSourceDir.IsEmpty() ? EVisibility::Collapsed : EVisibility::Visible;
+	}
+	
+}
+
+void FAutoReimportDirectoryCustomization::UpdateMountPath()
+{
+
+}
+
+#undef LOCTEXT_NAMESPACE
diff --git a/Engine/Source/Editor/DetailCustomizations/Private/AutoReimportDirectoryCustomization.h b/Engine/Source/Editor/DetailCustomizations/Private/AutoReimportDirectoryCustomization.h
new file mode 100644
--- /dev/null
+++ b/Engine/Source/Editor/DetailCustomizations/Private/AutoReimportDirectoryCustomization.h
@@ -0,0 +1,78 @@
+// Copyright 1998-2015 Epic Games, Inc. All Rights Reserved.
+
+#pragma once
+
+
+/**
+ * Implements a details view customization for the FAutoReimportDirectoryConfig struct.
+ */
+class FAutoReimportWildcardCustomization : public IPropertyTypeCustomization
+{
+public:
+
+	virtual void CustomizeHeader(TSharedRef<class IPropertyHandle> InStructPropertyHandle, class FDetailWidgetRow& InHeaderRow, IPropertyTypeCustomizationUtils& InStructCustomizationUtils) override;
+	virtual void CustomizeChildren(TSharedRef<class IPropertyHandle> InStructPropertyHandle, class IDetailChildrenBuilder& InStructBuilder, IPropertyTypeCustomizationUtils& InStructCustomizationUtils) override;
+
+public:
+
+	/** Creates an instance of this class. */
+	static TSharedRef<IPropertyTypeCustomization> MakeInstance()
+	{
+		return MakeShareable(new FAutoReimportWildcardCustomization());
+	}
+
+private:
+
+	FText GetWildcardText() const;
+	void OnWildcardCommitted(const FText& InValue, ETextCommit::Type CommitType);
+	void OnWildcardChanged(const FText& InValue);
+
+	void OnCheckStateChanged(ECheckBoxState InState);
+	ECheckBoxState GetCheckState() const;
+	
+	TSharedPtr<IPropertyHandle> PropertyHandle;
+	TSharedPtr<IPropertyHandle> WildcardProperty;
+	TSharedPtr<IPropertyHandle> IncludeProperty;
+};
+
+/**
+ * Implements a details view customization for the FAutoReimportDirectoryConfig struct.
+ */
+class FAutoReimportDirectoryCustomization : public IPropertyTypeCustomization
+{
+public:
+
+	virtual void CustomizeHeader(TSharedRef<class IPropertyHandle> InStructPropertyHandle, class FDetailWidgetRow& InHeaderRow, IPropertyTypeCustomizationUtils& InStructCustomizationUtils) override;
+	virtual void CustomizeChildren(TSharedRef<class IPropertyHandle> InStructPropertyHandle, class IDetailChildrenBuilder& InStructBuilder, IPropertyTypeCustomizationUtils& InStructCustomizationUtils) override;
+
+public:
+
+	/** Creates an instance of this class. */
+	static TSharedRef<IPropertyTypeCustomization> MakeInstance()
+	{
+		return MakeShareable(new FAutoReimportDirectoryCustomization());
+	}
+
+private:
+
+
+	EVisibility GetMountPathVisibility() const;
+
+	FText GetDirectoryText() const;
+	void OnDirectoryCommitted(const FText& InValue, ETextCommit::Type CommitType);
+	void OnDirectoryChanged(const FText& InValue);
+
+	FText GetMountPointText() const;
+	void OnMountPointCommitted(const FText& InValue, ETextCommit::Type CommitType);
+	void OnMountPointChanged(const FText& InValue);
+
+	void SetSourcePath(FString InSourceDir);
+	void UpdateMountPath();
+	FReply BrowseForFolder();
+
+	EVisibility MountPathVisibility;
+	TSharedPtr<IPropertyHandle> PropertyHandle;
+	TSharedPtr<IPropertyHandle> SourceDirProperty;
+	TSharedPtr<IPropertyHandle> MountPointProperty;
+	TSharedPtr<IPropertyHandle> WildcardsProperty;
+};
diff --git a/Engine/Source/Editor/DetailCustomizations/Private/DetailCustomizations.cpp b/Engine/Source/Editor/DetailCustomizations/Private/DetailCustomizations.cpp
--- a/Engine/Source/Editor/DetailCustomizations/Private/DetailCustomizations.cpp
+++ b/Engine/Source/Editor/DetailCustomizations/Private/DetailCustomizations.cpp
@@ -97,6 +97,7 @@
 #include "SceneCaptureDetails.h"
 #include "CurveColorCustomization.h"
 #include "ActorComponentDetails.h"
+#include "AutoReimportDirectoryCustomization.h"
 #include "DistanceDatumStructCustomization.h"
 
 IMPLEMENT_MODULE( FDetailCustomizationsModule, DetailCustomizations );
@@ -195,6 +196,8 @@ void FDetailCustomizationsModule::RegisterPropertyTypeCustomizations()
 	RegisterCustomPropertyTypeLayout("ParticleSysParam", FOnGetPropertyTypeCustomizationInstance::CreateStatic(&FParticleSysParamStructCustomization::MakeInstance));
 	RegisterCustomPropertyTypeLayout("RawDistributionVector", FOnGetPropertyTypeCustomizationInstance::CreateStatic(&FRawDistributionVectorStructCustomization::MakeInstance));
 	RegisterCustomPropertyTypeLayout("CollisionProfileName", FOnGetPropertyTypeCustomizationInstance::CreateStatic(&FCollisionProfileNameCustomization::MakeInstance));
+	RegisterCustomPropertyTypeLayout("AutoReimportDirectoryConfig", FOnGetPropertyTypeCustomizationInstance::CreateStatic(&FAutoReimportDirectoryCustomization::MakeInstance));
+	RegisterCustomPropertyTypeLayout("AutoReimportWildcard", FOnGetPropertyTypeCustomizationInstance::CreateStatic(&FAutoReimportWildcardCustomization::MakeInstance));
 	RegisterCustomPropertyTypeLayout("DistanceDatum", FOnGetPropertyTypeCustomizationInstance::CreateStatic(&FDistanceDatumStructCustomization::MakeInstance));
 }
 
