- HLOD
  - Do not display if simplygon isn't found
  - Do not allow it to use CreateProxyMesh if simplygon isn't found

https://jira.ol.epicgames.net/browse/UE-8454

diff --git a/Engine/Source/Editor/DetailCustomizations/DetailCustomizations.Build.cs b/Engine/Source/Editor/DetailCustomizations/DetailCustomizations.Build.cs
--- a/Engine/Source/Editor/DetailCustomizations/DetailCustomizations.Build.cs
+++ b/Engine/Source/Editor/DetailCustomizations/DetailCustomizations.Build.cs
@@ -41,7 +41,8 @@ public class DetailCustomizations : ModuleRules
                 "MediaAssets",
                 "HardwareTargeting",
 				"SharedSettingsWidgets",
-                "AIModule"
+                "AIModule", 
+                "MeshUtilities"
 			}
 		);
 
diff --git a/Engine/Source/Editor/DetailCustomizations/Private/DetailCustomizations.cpp b/Engine/Source/Editor/DetailCustomizations/Private/DetailCustomizations.cpp
--- a/Engine/Source/Editor/DetailCustomizations/Private/DetailCustomizations.cpp
+++ b/Engine/Source/Editor/DetailCustomizations/Private/DetailCustomizations.cpp
@@ -99,6 +99,7 @@
 #include "ActorComponentDetails.h"
 #include "AutoReimportDirectoryCustomization.h"
 #include "DistanceDatumStructCustomization.h"
+#include "HierarchicalSimplificationCustomizations.h"
 
 IMPLEMENT_MODULE( FDetailCustomizationsModule, DetailCustomizations );
 
@@ -199,6 +200,7 @@ void FDetailCustomizationsModule::RegisterPropertyTypeCustomizations()
 	RegisterCustomPropertyTypeLayout("AutoReimportDirectoryConfig", FOnGetPropertyTypeCustomizationInstance::CreateStatic(&FAutoReimportDirectoryCustomization::MakeInstance));
 	RegisterCustomPropertyTypeLayout("AutoReimportWildcard", FOnGetPropertyTypeCustomizationInstance::CreateStatic(&FAutoReimportWildcardCustomization::MakeInstance));
 	RegisterCustomPropertyTypeLayout("DistanceDatum", FOnGetPropertyTypeCustomizationInstance::CreateStatic(&FDistanceDatumStructCustomization::MakeInstance));
+	RegisterCustomPropertyTypeLayout("HierarchicalSimplification", FOnGetPropertyTypeCustomizationInstance::CreateStatic(&FHierarchicalSimplificationCustomizations::MakeInstance));
 }
 
 
diff --git a/Engine/Source/Editor/DetailCustomizations/Private/HierarchicalSimplificationCustomizations.cpp b/Engine/Source/Editor/DetailCustomizations/Private/HierarchicalSimplificationCustomizations.cpp
new file mode 100644
--- /dev/null
+++ b/Engine/Source/Editor/DetailCustomizations/Private/HierarchicalSimplificationCustomizations.cpp
@@ -0,0 +1,55 @@
+// Copyright 1998-2015 Epic Games, Inc. All Rights Reserved.
+
+#include "DetailCustomizationsPrivatePCH.h"
+#include "HierarchicalSimplificationCustomizations.h"
+#include "Developer/MeshUtilities/Public/MeshUtilities.h"
+
+TSharedRef<IPropertyTypeCustomization> FHierarchicalSimplificationCustomizations::MakeInstance() 
+{
+	return MakeShareable( new FHierarchicalSimplificationCustomizations );
+}
+
+void FHierarchicalSimplificationCustomizations::CustomizeHeader(TSharedRef<IPropertyHandle> StructPropertyHandle, FDetailWidgetRow& HeaderRow, IPropertyTypeCustomizationUtils& StructCustomizationUtils)
+{
+}
+
+void FHierarchicalSimplificationCustomizations::CustomizeChildren( TSharedRef<IPropertyHandle> StructPropertyHandle, class IDetailChildrenBuilder& ChildBuilder, IPropertyTypeCustomizationUtils& StructCustomizationUtils )
+{
+	uint32 NumChildren;
+	StructPropertyHandle->GetNumChildren( NumChildren );
+
+	TMap<FName, TSharedPtr< IPropertyHandle > > PropertyHandles;
+
+	for( uint32 ChildIndex = 0; ChildIndex < NumChildren; ++ChildIndex )
+	{
+		TSharedRef<IPropertyHandle> ChildHandle = StructPropertyHandle->GetChildHandle( ChildIndex ).ToSharedRef();
+		const FName PropertyName = ChildHandle->GetProperty()->GetFName();
+
+		PropertyHandles.Add(PropertyName, ChildHandle);
+	}
+
+	TSharedPtr< IPropertyHandle > SimplifyMeshPropertyHandle = PropertyHandles.FindChecked(GET_MEMBER_NAME_CHECKED(FHierarchicalSimplification, bSimplifyMesh));
+
+	IDetailPropertyRow& SimplifyMeshRow = ChildBuilder.AddChildProperty(SimplifyMeshPropertyHandle.ToSharedRef());
+	SimplifyMeshRow.Visibility(TAttribute<EVisibility>(this, &FHierarchicalSimplificationCustomizations::IsSimplifyMeshVisible));
+
+	for( auto Iter(PropertyHandles.CreateConstIterator()); Iter; ++Iter  )
+	{
+		if (Iter.Value() != SimplifyMeshPropertyHandle)
+		{
+			ChildBuilder.AddChildProperty(Iter.Value().ToSharedRef());
+		}
+	}
+}
+
+EVisibility FHierarchicalSimplificationCustomizations::IsSimplifyMeshVisible() const
+{
+	IMeshUtilities& MeshUtilities = FModuleManager::Get().LoadModuleChecked<IMeshUtilities>("MeshUtilities");
+	if (MeshUtilities.GetMeshMergingInterface() != nullptr)
+	{
+		return EVisibility::Visible;
+	}
+
+	return EVisibility::Hidden;
+}
+
diff --git a/Engine/Source/Editor/DetailCustomizations/Private/HierarchicalSimplificationCustomizations.h b/Engine/Source/Editor/DetailCustomizations/Private/HierarchicalSimplificationCustomizations.h
new file mode 100644
--- /dev/null
+++ b/Engine/Source/Editor/DetailCustomizations/Private/HierarchicalSimplificationCustomizations.h
@@ -0,0 +1,17 @@
+// Copyright 1998-2015 Epic Games, Inc. All Rights Reserved.
+
+#pragma once
+
+class FHierarchicalSimplificationCustomizations : public IPropertyTypeCustomization
+{
+public:
+	static TSharedRef<IPropertyTypeCustomization> MakeInstance();
+
+	/** IPropertyTypeCustomization instance */
+	virtual void CustomizeHeader(TSharedRef<IPropertyHandle> StructPropertyHandle, FDetailWidgetRow& HeaderRow, IPropertyTypeCustomizationUtils& StructCustomizationUtils) override;
+	virtual void CustomizeChildren( TSharedRef<IPropertyHandle> StructPropertyHandle, class IDetailChildrenBuilder& ChildBuilder, IPropertyTypeCustomizationUtils& StructCustomizationUtils ) override;
+
+protected:
+	EVisibility IsSimplifyMeshVisible() const;
+	TSharedPtr< IPropertyHandle > SimplifyMeshHandle;
+};
diff --git a/Engine/Source/Editor/UnrealEd/Private/HierarchicalLOD.cpp b/Engine/Source/Editor/UnrealEd/Private/HierarchicalLOD.cpp
--- a/Engine/Source/Editor/UnrealEd/Private/HierarchicalLOD.cpp
+++ b/Engine/Source/Editor/UnrealEd/Private/HierarchicalLOD.cpp
@@ -568,7 +568,7 @@ void FLODCluster::BuildActor(class UWorld* InWorld, class ULevel* InLevel, const
 				IMeshUtilities& MeshUtilities = FModuleManager::Get().LoadModuleChecked<IMeshUtilities>("MeshUtilities");
 				// should give unique name, so use level + actor name
 				const FString PackageName = FString::Printf(TEXT("LOD_%s_%s"), *InLevel->GetName(), *FirstActor->GetName());
-				if (LODSetup.bSimplifyMesh)
+				if (MeshUtilities.GetMeshMergingInterface() && LODSetup.bSimplifyMesh)
 				{
 					MeshUtilities.CreateProxyMesh(Actors, ProxySetting, AssetsOuter, PackageName, OutAssets, OutProxyLocation);
 				}
