Before displaying the new SPATIALIZATION_Algorithm enum in FSoundAttenuation, now checking to see if a spatialization plugin is enabled

- this will prevent it being confusing for people who don't have a spatialization plugin enabled and wondering why the enum doesn't seem to do anything

- implementation introduces new enum in audio.h for plugin-types for audio. The idea is that people in the future may make more plugin types, so might as well make a "way" for them to do it and query if its enabled. Will be useful for other plugins that modify existing UClasses.

- also fixed a bug with a boolean being false when it should be true (bSpatializationExtensionEnabled). This bool is used to toggle off/on the fancy spatialization plugin from console command to "diff" between it being on/off for sound designers

diff --git a/Engine/Source/Editor/DetailCustomizations/Private/AttenuationSettingsCustomizations.cpp b/Engine/Source/Editor/DetailCustomizations/Private/AttenuationSettingsCustomizations.cpp
--- a/Engine/Source/Editor/DetailCustomizations/Private/AttenuationSettingsCustomizations.cpp
+++ b/Engine/Source/Editor/DetailCustomizations/Private/AttenuationSettingsCustomizations.cpp
@@ -78,7 +78,12 @@ void FAttenuationSettingsCustomization::CustomizeChildren( TSharedRef<IPropertyH
 	ChildBuilder.AddChildProperty(PropertyHandles.FindChecked(GET_MEMBER_NAME_CHECKED(FAttenuationSettings, bAttenuate)).ToSharedRef());
 	ChildBuilder.AddChildProperty(PropertyHandles.FindChecked(GET_MEMBER_NAME_CHECKED(FAttenuationSettings, bSpatialize)).ToSharedRef());
 	ChildBuilder.AddChildProperty(DistanceAlgorithmHandle.ToSharedRef() );
-	ChildBuilder.AddChildProperty(SpatializationAlgorithmHandle.ToSharedRef());
+
+	// Check to see if a spatialization plugin is enabled
+	if (IsAudioPluginEnabled(EAudioPlugin::SPATIALIZATION))
+	{
+		ChildBuilder.AddChildProperty(SpatializationAlgorithmHandle.ToSharedRef());
+	}
 
 	IDetailPropertyRow& dbAttenuationRow = ChildBuilder.AddChildProperty(PropertyHandles.FindChecked(GET_MEMBER_NAME_CHECKED(FAttenuationSettings, dBAttenuationAtMax)).ToSharedRef());
 	dbAttenuationRow.Visibility(TAttribute<EVisibility>(this, &FAttenuationSettingsCustomization::IsNaturalSoundSelected));
diff --git a/Engine/Source/Runtime/Engine/Private/Audio.cpp b/Engine/Source/Runtime/Engine/Private/Audio.cpp
--- a/Engine/Source/Runtime/Engine/Private/Audio.cpp
+++ b/Engine/Source/Runtime/Engine/Private/Audio.cpp
@@ -17,6 +17,7 @@
 #include "EngineAnalytics.h"
 #include "Runtime/Analytics/Analytics/Public/Interfaces/IAnalyticsProvider.h"
 #include "ContentStreaming.h"
+#include "IAudioExtensionPlugin.h"
 
 DEFINE_LOG_CATEGORY(LogAudio);
 
@@ -52,6 +53,19 @@ DEFINE_STAT(STAT_AudioGatherWaveInstances);
 DEFINE_STAT(STAT_AudioUpdateTime);
 DEFINE_STAT(STAT_AudioFindNearestLocation);
 
+
+bool IsAudioPluginEnabled(EAudioPlugin::Type PluginType)
+{
+	if (PluginType == EAudioPlugin::SPATIALIZATION)
+	{
+		TArray<IAudioSpatializationPlugin *> SpatializationPlugins = IModularFeatures::Get().GetModularFeatureImplementations<IAudioSpatializationPlugin>(IAudioSpatializationPlugin::GetModularFeatureName());
+		return SpatializationPlugins.Num() > 0;
+	}
+
+	return false;
+}
+
+
 /*-----------------------------------------------------------------------------
 	FSoundBuffer implementation.
 -----------------------------------------------------------------------------*/
diff --git a/Engine/Source/Runtime/Engine/Private/AudioDevice.cpp b/Engine/Source/Runtime/Engine/Private/AudioDevice.cpp
--- a/Engine/Source/Runtime/Engine/Private/AudioDevice.cpp
+++ b/Engine/Source/Runtime/Engine/Private/AudioDevice.cpp
@@ -96,6 +96,7 @@ bool FAudioDevice::Init()
 
 		Plugin->Initialize();
 		SpatializeProcessor = Plugin->GetNewSpatializationAlgorithm(this);
+		bSpatializationExtensionEnabled = true;
 
 		// There should only ever be 0 or 1 spatialization plugin at the moment
 		check(SpatializationPlugins.Num() == 1);
diff --git a/Engine/Source/Runtime/Engine/Public/Audio.h b/Engine/Source/Runtime/Engine/Public/Audio.h
--- a/Engine/Source/Runtime/Engine/Public/Audio.h
+++ b/Engine/Source/Runtime/Engine/Public/Audio.h
@@ -149,6 +149,23 @@ private:
 	TArray<FNotifyBufferDetails> Notifies;
 };
 
+
+/**
+* Enumeration of audio plugin types
+*
+*/
+struct EAudioPlugin
+{
+	enum Type
+	{
+		SPATIALIZATION,
+	};
+};
+
+/** Queries if a plugin of the given type is enabled. */
+ENGINE_API bool IsAudioPluginEnabled(EAudioPlugin::Type PluginType);
+
+
 /**
  * Structure encapsulating all information required to play a USoundWave on a channel/source. This is required
  * as a single USoundWave object can be used in multiple active cues or multiple times in the same cue.
@@ -559,3 +576,5 @@ ENGINE_API void SetCompressedAudioFormatsToBuild(const TCHAR* Platform = NULL);
  * @param	Platform				Name of platform to cook for, or NULL if all platforms
  */
 ENGINE_API const TArray<FName>& GetCompressedAudioFormatsToBuild();
+
+
