Fix for crash in  FComponentTransformDetails::CacheCommonLocationUnits

diff --git a/Engine/Source/Editor/DetailCustomizations/Private/ComponentTransformDetails.cpp b/Engine/Source/Editor/DetailCustomizations/Private/ComponentTransformDetails.cpp
--- a/Engine/Source/Editor/DetailCustomizations/Private/ComponentTransformDetails.cpp
+++ b/Engine/Source/Editor/DetailCustomizations/Private/ComponentTransformDetails.cpp
@@ -601,10 +601,14 @@ void FComponentTransformDetails::CacheCommonLocationUnits()
 	{
 		VectorUnits = LargestUnit.GetValue();
 	}
-	else
+	else if (DefaultInputUnits.IsSet())
 	{
 		VectorUnits	= DefaultInputUnits.GetValue();
 	}
+	else
+	{
+		VectorUnits = EUnit::Unspecified;
+	}
 
 	VectorUnits = FUnitConversion::ConvertToGlobalDisplayRange(VectorUnits);
 }
