Fix for crash running actor blutility when the event has parameters
#UE-12605 - Enabling Call In Editor on an event with parameters causes an assert
#Branch UE4
#Proj DetailsCustomizations, CoreUObject

diff --git a/Engine/Source/Editor/DetailCustomizations/Private/ActorDetails.cpp b/Engine/Source/Editor/DetailCustomizations/Private/ActorDetails.cpp
--- a/Engine/Source/Editor/DetailCustomizations/Private/ActorDetails.cpp
+++ b/Engine/Source/Editor/DetailCustomizations/Private/ActorDetails.cpp
@@ -711,7 +711,8 @@ FReply FActorDetails::CallBlutilityFunction()
 	if( ActorWeakPtr.IsValid() && ActiveBlutilityFunction.IsValid() )
 	{
 		AActor* Actor = ActorWeakPtr.Get();
-		Actor->ProcessEvent( ActiveBlutilityFunction.Get(), NULL );
+		UFunction* ActiveFunction = ActiveBlutilityFunction.Get();
+		Actor->ProcessEvent( ActiveFunction, ActiveFunction->Children );
 	}
 
 	return FReply::Handled();
