Fix crash when setting the mobility property on scene components due to re-entrancy.
https://jira.ol.epicgames.net/browse/UE-13374

diff --git a/Engine/Source/Editor/DetailCustomizations/Private/SceneComponentDetails.cpp b/Engine/Source/Editor/DetailCustomizations/Private/SceneComponentDetails.cpp
--- a/Engine/Source/Editor/DetailCustomizations/Private/SceneComponentDetails.cpp
+++ b/Engine/Source/Editor/DetailCustomizations/Private/SceneComponentDetails.cpp
@@ -351,12 +351,12 @@ TSharedRef<IDetailCustomization> FSceneComponentDetails::MakeInstance()
 	return MakeShareable( new FSceneComponentDetails );
 }
 
-ECheckBoxState FSceneComponentDetails::IsMobilityActive(TWeakPtr<IPropertyHandle> MobilityHandle, EComponentMobility::Type InMobility) const
+ECheckBoxState FSceneComponentDetails::IsMobilityActive(EComponentMobility::Type InMobility) const
 {
 	if ( MobilityHandle.IsValid() )
 	{
 		uint8 MobilityByte;
-		MobilityHandle.Pin()->GetValue(MobilityByte);
+		MobilityHandle->GetValue(MobilityByte);
 
 		return MobilityByte == InMobility ? ECheckBoxState::Checked : ECheckBoxState::Unchecked;
 	}
@@ -364,12 +364,12 @@ ECheckBoxState FSceneComponentDetails::IsMobilityActive(TWeakPtr<IPropertyHandle
 	return ECheckBoxState::Unchecked;
 }
 
-FSlateColor FSceneComponentDetails::GetMobilityTextColor(TWeakPtr<IPropertyHandle> MobilityHandle, EComponentMobility::Type InMobility) const
+FSlateColor FSceneComponentDetails::GetMobilityTextColor(EComponentMobility::Type InMobility) const
 {
 	if ( MobilityHandle.IsValid() )
 	{
 		uint8 MobilityByte;
-		MobilityHandle.Pin()->GetValue(MobilityByte);
+		MobilityHandle->GetValue(MobilityByte);
 
 		return MobilityByte == InMobility ? FSlateColor(FLinearColor(0, 0, 0)) : FSlateColor(FLinearColor(0.72f, 0.72f, 0.72f, 1.f));
 	}
@@ -377,19 +377,19 @@ FSlateColor FSceneComponentDetails::GetMobilityTextColor(TWeakPtr<IPropertyHandl
 	return FSlateColor(FLinearColor(0.72f, 0.72f, 0.72f, 1.f));
 }
 
-void FSceneComponentDetails::OnMobilityChanged(ECheckBoxState InCheckedState, TWeakPtr<IPropertyHandle> MobilityHandle, EComponentMobility::Type InMobility)
+void FSceneComponentDetails::OnMobilityChanged(ECheckBoxState InCheckedState, EComponentMobility::Type InMobility)
 {
 	if ( MobilityHandle.IsValid() && InCheckedState == ECheckBoxState::Checked)
 	{
-		MobilityHandle.Pin()->SetValue((uint8)InMobility);
+		MobilityHandle->SetValue((uint8)InMobility);
 	}
 }
 
-FText FSceneComponentDetails::GetMobilityToolTip(TWeakPtr<IPropertyHandle> MobilityHandle) const
+FText FSceneComponentDetails::GetMobilityToolTip() const
 {
 	if ( MobilityHandle.IsValid() )
 	{
-		return MobilityHandle.Pin()->GetToolTipText();
+		return MobilityHandle->GetToolTipText();
 	}
 
 	return FText::GetEmpty();
@@ -401,7 +401,7 @@ void FSceneComponentDetails::CustomizeDetails( IDetailLayoutBuilder& DetailBuild
 
 	// Put mobility property in Transform section
 	IDetailCategoryBuilder& TransformCategory = DetailBuilder.EditCategory( "TransformCommon", LOCTEXT("TransformCommonCategory", "Transform"), ECategoryPriority::Transform );
-	TSharedPtr<IPropertyHandle> MobilityProperty = DetailBuilder.GetProperty("Mobility");
+	MobilityHandle = DetailBuilder.GetProperty("Mobility");
 
 	uint8 RestrictedMobilityBits = 0u;
 	enum 
@@ -433,7 +433,7 @@ void FSceneComponentDetails::CustomizeDetails( IDetailLayoutBuilder& DetailBuild
 			{
 				 TSharedPtr<FPropertyRestriction> StaticRestriction = MakeShareable(new FPropertyRestriction(RestrictReason));
 				 StaticRestriction->AddValue(TEXT("Static"));
-				 MobilityProperty->AddRestriction(StaticRestriction.ToSharedRef());
+				 MobilityHandle->AddRestriction(StaticRestriction.ToSharedRef());
 
 				 RestrictedMobilityBits |= StaticMobilityBitMask;
 			}			
@@ -447,7 +447,7 @@ void FSceneComponentDetails::CustomizeDetails( IDetailLayoutBuilder& DetailBuild
 			{
 				TSharedPtr<FPropertyRestriction> StationaryRestriction = MakeShareable(new FPropertyRestriction(RestrictReason));
 				StationaryRestriction->AddValue(TEXT("Stationary"));
-				MobilityProperty->AddRestriction(StationaryRestriction.ToSharedRef());
+				MobilityHandle->AddRestriction(StationaryRestriction.ToSharedRef());
 
 				RestrictedMobilityBits |= StationaryMobilityBitMask;
 			}
@@ -460,22 +460,15 @@ void FSceneComponentDetails::CustomizeDetails( IDetailLayoutBuilder& DetailBuild
 		}
 	}
 
-	FSimpleDelegate OnMobilityChangedDelegate = FSimpleDelegate::CreateSP(this, &FSceneComponentDetails::OnMobilityChanged, MobilityProperty);
-	MobilityProperty->SetOnPropertyValueChanged(OnMobilityChangedDelegate);
-
-
-	// Build Editor for Mobility
-	TWeakPtr<IPropertyHandle> MobilityPropertyWeak(MobilityProperty);
-
 	TSharedPtr<SUniformGridPanel> ButtonOptionsPanel;
 
-	IDetailPropertyRow& MobilityRow = TransformCategory.AddProperty(MobilityProperty);
+	IDetailPropertyRow& MobilityRow = TransformCategory.AddProperty(MobilityHandle);
 	MobilityRow.CustomWidget()
 	.NameContent()
 	[
 		SNew(STextBlock)
 		.Text(LOCTEXT("Mobility", "Mobility"))
-		.ToolTipText(this, &FSceneComponentDetails::GetMobilityToolTip, MobilityPropertyWeak)
+		.ToolTipText(this, &FSceneComponentDetails::GetMobilityToolTip)
 		.Font(IDetailLayoutBuilder::GetDetailFont())
 	]
 	.ValueContent()
@@ -496,8 +489,8 @@ void FSceneComponentDetails::CustomizeDetails( IDetailLayoutBuilder& DetailBuild
 		[
 			SNew(SCheckBox)
 			.Style(FEditorStyle::Get(), "Property.ToggleButton.Start")
-			.IsChecked(this, &FSceneComponentDetails::IsMobilityActive, MobilityPropertyWeak, EComponentMobility::Static)
-			.OnCheckStateChanged(this, &FSceneComponentDetails::OnMobilityChanged, MobilityPropertyWeak, EComponentMobility::Static)
+			.IsChecked(this, &FSceneComponentDetails::IsMobilityActive, EComponentMobility::Static)
+			.OnCheckStateChanged(this, &FSceneComponentDetails::OnMobilityChanged, EComponentMobility::Static)
 			.ToolTipText(LOCTEXT("Mobility_Static_Tooltip", "A static object can't be changed in game.\n● Allows Baked Lighting\n● Fastest Rendering"))
 			[
 				SNew(SHorizontalBox)
@@ -520,7 +513,7 @@ void FSceneComponentDetails::CustomizeDetails( IDetailLayoutBuilder& DetailBuild
 					SNew(STextBlock)
 					.Text(LOCTEXT("Static", "Static"))
 					.Font(IDetailLayoutBuilder::GetDetailFont())
-					.ColorAndOpacity(this, &FSceneComponentDetails::GetMobilityTextColor, MobilityPropertyWeak, EComponentMobility::Static)
+					.ColorAndOpacity(this, &FSceneComponentDetails::GetMobilityTextColor, EComponentMobility::Static)
 				]
 			]
 		];
@@ -534,9 +527,9 @@ void FSceneComponentDetails::CustomizeDetails( IDetailLayoutBuilder& DetailBuild
 		ButtonOptionsPanel->AddSlot(ColumnIndex, 0)
 		[
 			SNew(SCheckBox)
-			.IsChecked(this, &FSceneComponentDetails::IsMobilityActive, MobilityPropertyWeak, EComponentMobility::Stationary)
+			.IsChecked(this, &FSceneComponentDetails::IsMobilityActive, EComponentMobility::Stationary)
 			.Style(FEditorStyle::Get(), ( ColumnIndex == 0 ) ? "Property.ToggleButton.Start" : "Property.ToggleButton.Middle")
-			.OnCheckStateChanged(this, &FSceneComponentDetails::OnMobilityChanged, MobilityPropertyWeak, EComponentMobility::Stationary)
+			.OnCheckStateChanged(this, &FSceneComponentDetails::OnMobilityChanged, EComponentMobility::Stationary)
 			.ToolTipText(LOCTEXT("Mobility_Stationary_Tooltip", "A stationary light will only have its shadowing and bounced lighting from static geometry baked by Lightmass, all other lighting will be dynamic.  It can change color and intensity in game.\n● Can't Move\n● Allows Partial Baked Lighting\n● Dynamic Shadows"))
 			[
 				SNew(SHorizontalBox)
@@ -559,7 +552,7 @@ void FSceneComponentDetails::CustomizeDetails( IDetailLayoutBuilder& DetailBuild
 					SNew(STextBlock)
 					.Text(LOCTEXT("Stationary", "Stationary"))
 					.Font(IDetailLayoutBuilder::GetDetailFont())
-					.ColorAndOpacity(this, &FSceneComponentDetails::GetMobilityTextColor, MobilityPropertyWeak, EComponentMobility::Stationary)
+					.ColorAndOpacity(this, &FSceneComponentDetails::GetMobilityTextColor, EComponentMobility::Stationary)
 				]
 			]
 		];
@@ -571,9 +564,9 @@ void FSceneComponentDetails::CustomizeDetails( IDetailLayoutBuilder& DetailBuild
 	ButtonOptionsPanel->AddSlot(ColumnIndex, 0)
 	[
 		SNew(SCheckBox)
-		.IsChecked(this, &FSceneComponentDetails::IsMobilityActive, MobilityPropertyWeak, EComponentMobility::Movable)
+		.IsChecked(this, &FSceneComponentDetails::IsMobilityActive, EComponentMobility::Movable)
 		.Style(FEditorStyle::Get(), ( ColumnIndex == 0 ) ? "Property.ToggleButton" : "Property.ToggleButton.End")
-		.OnCheckStateChanged(this, &FSceneComponentDetails::OnMobilityChanged, MobilityPropertyWeak, EComponentMobility::Movable)
+		.OnCheckStateChanged(this, &FSceneComponentDetails::OnMobilityChanged, EComponentMobility::Movable)
 		.ToolTipText(LOCTEXT("Mobility_Movable_Tooltip", "Movable objects can be moved and changed in game.\n● Totally Dynamic\n● Allows Dynamic Shadows\n● Slowest Rendering"))
 		[
 			SNew(SHorizontalBox)
@@ -596,7 +589,7 @@ void FSceneComponentDetails::CustomizeDetails( IDetailLayoutBuilder& DetailBuild
 				SNew(STextBlock)
 				.Text(LOCTEXT("Movable", "Movable"))
 				.Font(IDetailLayoutBuilder::GetDetailFont())
-				.ColorAndOpacity(this, &FSceneComponentDetails::GetMobilityTextColor, MobilityPropertyWeak, EComponentMobility::Movable)
+				.ColorAndOpacity(this, &FSceneComponentDetails::GetMobilityTextColor, EComponentMobility::Movable)
 			]
 		]
 	];
diff --git a/Engine/Source/Editor/DetailCustomizations/Private/SceneComponentDetails.h b/Engine/Source/Editor/DetailCustomizations/Private/SceneComponentDetails.h
--- a/Engine/Source/Editor/DetailCustomizations/Private/SceneComponentDetails.h
+++ b/Engine/Source/Editor/DetailCustomizations/Private/SceneComponentDetails.h
@@ -14,13 +14,13 @@ public:
 private:
 	void MakeTransformDetails( IDetailLayoutBuilder& DetailBuilder );
 
-	FSlateColor GetMobilityTextColor(TWeakPtr<IPropertyHandle> MobilityHandle, EComponentMobility::Type InMobility) const;
+	FSlateColor GetMobilityTextColor(EComponentMobility::Type InMobility) const;
 
-	ECheckBoxState IsMobilityActive(TWeakPtr<IPropertyHandle> MobilityHandle, EComponentMobility::Type InMobility) const;
+	ECheckBoxState IsMobilityActive(EComponentMobility::Type InMobility) const;
 
-	void OnMobilityChanged(ECheckBoxState InCheckedState, TWeakPtr<IPropertyHandle> MobilityHandle, EComponentMobility::Type InMobility);
+	void OnMobilityChanged(ECheckBoxState InCheckedState, EComponentMobility::Type InMobility);
 
-	FText GetMobilityToolTip(TWeakPtr<IPropertyHandle> MobilityHandle) const;
+	FText GetMobilityToolTip() const;
 	
 	/**
 	 * When a scene component's Mobility is altered, we need to make sure the scene hierarchy is
@@ -43,4 +43,6 @@ private:
 
 	/** A collection of USceneComponents that were selected last time CustomizeDetails() was ran */
 	TArray< TWeakObjectPtr<UObject> > CachedSelectedSceneComponents;
+private:
+	TSharedPtr<IPropertyHandle> MobilityHandle;
 };
