Fix crash setting mobility on components.

This is due to recent changes that rerun construction scripts when component properties change.  Moved the propagation of mobility to child components to USceneComponent::PostEditChange methods.

diff --git a/Engine/Source/Editor/DetailCustomizations/Private/SceneComponentDetails.cpp b/Engine/Source/Editor/DetailCustomizations/Private/SceneComponentDetails.cpp
--- a/Engine/Source/Editor/DetailCustomizations/Private/SceneComponentDetails.cpp
+++ b/Engine/Source/Editor/DetailCustomizations/Private/SceneComponentDetails.cpp
@@ -10,223 +10,13 @@
 #include "SNotificationList.h"
 #include "NotificationManager.h"
 #include "Engine/BlueprintGeneratedClass.h"
-#include "Engine/SCS_Node.h"
-#include "Engine/SimpleConstructionScript.h"
 #include "Engine/InheritableComponentHandler.h"
 #include "Components/LightComponentBase.h"
+#include "ComponentUtils.h"
 
 #define LOCTEXT_NAMESPACE "SceneComponentDetails"
 
 
-/**
- * A helper for retrieving the simple-construction-script that this component belongs in.
- * 
- * @param  Component	The component you want the SCS for.
- * @return The component's blueprint SCS (NULL if one wasn't found).
- */
-static USimpleConstructionScript const* GetSimpleConstructionScript(USceneComponent const* Component)
-{
-	USimpleConstructionScript const* BlueprintSCS = NULL;
-
-	check(Component);
-	UObject const* ComponentOuter = Component->GetOuter();
-
-	if (UBlueprint const* const OuterBlueprint = Cast<UBlueprint const>(ComponentOuter))
-	{
-		BlueprintSCS = OuterBlueprint->SimpleConstructionScript;
-	}
-	else if (UBlueprintGeneratedClass const* const GeneratedClass = Cast<UBlueprintGeneratedClass>(ComponentOuter))
-	{
-		BlueprintSCS = GeneratedClass->SimpleConstructionScript;
-	}
-
-	return BlueprintSCS;
-}
-
-/**
- * A static helper function for retrieving the simple-construction-script node that
- * corresponds to the specified scene component template.
- * 
- * @param  ComponentObj	The component you want to find a USCS_Node for
- * @return A USCS_Node pointer corresponding to the specified component (NULL if we didn't find one)
- */
-static USCS_Node* FindCorrespondingSCSNode(USceneComponent const* ComponentObj)
-{
-	USimpleConstructionScript const* BlueprintSCS = GetSimpleConstructionScript(ComponentObj);
-	if (BlueprintSCS == NULL)
-	{
-		return NULL;
-	}
-
-	TArray<USCS_Node*> AllSCSNodes = BlueprintSCS->GetAllNodes();
-	for (int32 SCSNodeIndex = 0; SCSNodeIndex < AllSCSNodes.Num(); ++SCSNodeIndex)
-	{
-		USCS_Node* SCSNode = AllSCSNodes[SCSNodeIndex];
-		if (SCSNode->ComponentTemplate == ComponentObj) 
-		{
-			return SCSNode;
-		}
-	}
-
-	return NULL;
-}
-
-
-/**
- * A static helper function used to retrieve a component's scene parent
- * 
- * @param  SceneComponentObject	The component you want the attached parent for
- * @return A pointer to the component's scene parent (NULL if there was not one)
- */
-static USceneComponent* GetAttachedParent(USceneComponent const* SceneComponentObject)
-{
-	USceneComponent* SceneParent = SceneComponentObject->AttachParent;
-	if (SceneParent == nullptr)
-	{
-		USCS_Node* const SCSNode = FindCorrespondingSCSNode(SceneComponentObject);
-		// if we didn't find a corresponding simple-construction-script node
-		if (SCSNode == nullptr) 
-		{
-			return nullptr;
-		}	
-
-		USimpleConstructionScript const* BlueprintSCS = GetSimpleConstructionScript(SceneComponentObject);
-		check(BlueprintSCS != nullptr); 
-
-		USCS_Node* const ParentSCSNode = BlueprintSCS->FindParentNode(SCSNode);
-		if (ParentSCSNode != nullptr)
-		{
-			SceneParent = Cast<USceneComponent>(ParentSCSNode->ComponentTemplate);
-		}
-	}
-
-	return SceneParent;
-}
-
-/**
- * A static helper function used meant as the default for determining if a 
- * component's Mobility should be overridden.
- * 
- * @param  CurrentMobility	The component's current Mobility setting
- * @param  NewMobility		The proposed new Mobility for the component in question
- * @return True if the two mobilities are not equal (false if they are equal)
- */
-static bool AreMobilitiesDifferent(EComponentMobility::Type CurrentMobility, EComponentMobility::Type NewMobility)
-{
-	return CurrentMobility != NewMobility;
-}
-DECLARE_DELEGATE_RetVal_OneParam(bool, FMobilityQueryDelegate, EComponentMobility::Type);
-
-/**
- * A static helper function that recursively alters the Mobility property for all 
- * sub-components (descending from the specified USceneComponent)
- * 
- * @param  SceneComponentObject		The component whose sub-components you want to alter
- * @param  NewMobilityType			The Mobility type you want to switch sub-components over to
- * @param  ShouldOverrideMobility	A delegate used to determine if a sub-component's Mobility should be overridden
- *									(if left unset it will default to the AreMobilitiesDifferent() function)
- * @return The number of decedents that had their mobility altered.
- */
-static int32 SetDecedentMobility(USceneComponent const* SceneComponentObject, EComponentMobility::Type NewMobilityType, FMobilityQueryDelegate ShouldOverrideMobility = FMobilityQueryDelegate())
-{
-	if (!ensure(SceneComponentObject != NULL))
-	{
-		return 0;
-	}
-
-	TArray<USceneComponent*> AttachedChildren = SceneComponentObject->AttachChildren;
-	// gather children for component templates
-	USCS_Node* SCSNode = FindCorrespondingSCSNode(SceneComponentObject);
-	if (SCSNode != NULL)
-	{
-		// gather children from the SCSNode
-		for (int32 ChildIndex = 0; ChildIndex < SCSNode->ChildNodes.Num(); ++ChildIndex)
-		{
-			USCS_Node* SCSChild = SCSNode->ChildNodes[ChildIndex];
-
-			USceneComponent* ChildSceneComponent = Cast<USceneComponent>(SCSChild->ComponentTemplate);
-			if (ChildSceneComponent != NULL)
-			{
-				AttachedChildren.Add(ChildSceneComponent);
-			}
-		}
-	}
-
-	if (!ShouldOverrideMobility.IsBound())
-	{
-		ShouldOverrideMobility = FMobilityQueryDelegate::CreateStatic(&AreMobilitiesDifferent, NewMobilityType);
-	}
-
-	int32 NumDecendentsChanged = 0;
-	// recursively alter the mobility for children and deeper decedents 
-	for (int32 ChildIndex = 0; ChildIndex < AttachedChildren.Num(); ++ChildIndex)
-	{
-		USceneComponent* ChildSceneComponent = AttachedChildren[ChildIndex];
-
-		if (ShouldOverrideMobility.Execute(ChildSceneComponent->Mobility))
-		{
-			// USceneComponents shouldn't be set Stationary 
-			if ((NewMobilityType == EComponentMobility::Stationary) && ChildSceneComponent->IsA(UStaticMeshComponent::StaticClass()))
-			{
-				// make it Movable (because it is acceptable for Stationary parents to have Movable children)
-				ChildSceneComponent->Mobility = EComponentMobility::Movable;
-			}
-			else
-			{
-				ChildSceneComponent->Mobility = NewMobilityType;
-			}
-			++NumDecendentsChanged;
-		}
-		NumDecendentsChanged += SetDecedentMobility(ChildSceneComponent, NewMobilityType, ShouldOverrideMobility);
-	}
-
-	return NumDecendentsChanged;
-}
-
-/**
- * A static helper function that alters the Mobility property for all ancestor
- * components (ancestors of the specified USceneComponent).
- * 
- * @param  SceneComponentObject		The component whose attached ancestors you want to alter
- * @param  NewMobilityType			The Mobility type you want to switch ancestor components over to
- * @param  ShouldOverrideMobility	A delegate used to determine if a ancestor's Mobility should be overridden
- *									(if left unset it will default to the AreMobilitiesDifferent() function)
- * @return The number of ancestors that had their mobility altered.
- */
-static int32 SetAncestorMobility(USceneComponent const* SceneComponentObject, EComponentMobility::Type NewMobilityType, FMobilityQueryDelegate ShouldOverrideMobility = FMobilityQueryDelegate())
-{
-	if (!ensure(SceneComponentObject != NULL))
-	{
-		return 0;
-	}
-
-	if (!ShouldOverrideMobility.IsBound())
-	{
-		ShouldOverrideMobility = FMobilityQueryDelegate::CreateStatic(&AreMobilitiesDifferent, NewMobilityType);
-	}
-
-	int32 MobilityAlteredCount = 0;
-	while(USceneComponent* AttachedParent = GetAttachedParent(SceneComponentObject))
-	{
-		if (ShouldOverrideMobility.Execute(AttachedParent->Mobility))
-		{
-			// USceneComponents shouldn't be set Stationary 
-			if ((NewMobilityType == EComponentMobility::Stationary) && AttachedParent->IsA(UStaticMeshComponent::StaticClass()))
-			{
-				// make it Static (because it is acceptable for Stationary children to have Static parents)
-				AttachedParent->Mobility = EComponentMobility::Static;
-			}
-			else
-			{
-				AttachedParent->Mobility = NewMobilityType;
-			}
-			++MobilityAlteredCount;
-		}
-		SceneComponentObject = AttachedParent;
-	}
-
-	return MobilityAlteredCount;
-}
 
 /**
  * Walks the scene hierarchy looking for inherited components (like ones from
@@ -240,7 +30,7 @@ static EComponentMobility::Type GetInheritedMobility(USceneComponent const* cons
 	// default to "static" since it doesn't restrict anything (in case we don't inherit any at all)
 	EComponentMobility::Type InheritedMobility = EComponentMobility::Static;
 
-	USCS_Node* ComponentNode = FindCorrespondingSCSNode(SceneComponent);
+	USCS_Node* ComponentNode = ComponentUtils::FindCorrespondingSCSNode(SceneComponent);
 	
 	if(ComponentNode == NULL)
 	{
@@ -410,9 +200,13 @@ void FSceneComponentDetails::CustomizeDetails( IDetailLayoutBuilder& DetailBuild
 		StationaryMobilityBitMask = (1u << EComponentMobility::Stationary),
 		MovableMobilityBitMask    = (1u << EComponentMobility::Movable),
 	};
+
+	TArray< TWeakObjectPtr<UObject> > SelectedSceneComponents;
+
 	// see if any of the selected objects have mobility restrictions
-	DetailBuilder.GetObjectsBeingCustomized( CachedSelectedSceneComponents );
-	for (TArray<TWeakObjectPtr<UObject>>::TConstIterator ObjectIt(CachedSelectedSceneComponents); ObjectIt; ++ObjectIt)
+	DetailBuilder.GetObjectsBeingCustomized( SelectedSceneComponents );
+
+	for (TArray<TWeakObjectPtr<UObject>>::TConstIterator ObjectIt(SelectedSceneComponents); ObjectIt; ++ObjectIt)
 	{
 		if (!ObjectIt->IsValid())
 		{
@@ -641,7 +435,7 @@ void FSceneComponentDetails::MakeTransformDetails( IDetailLayoutBuilder& DetailB
 			{
 				bShouldShowTransform = false;
 			}
-			else if (const USimpleConstructionScript* SCS = GetSimpleConstructionScript(SceneComponent))
+			else if (const USimpleConstructionScript* SCS = ComponentUtils::GetSimpleConstructionScript(SceneComponent))
 			{
 				const TArray<USCS_Node*>& RootNodes = SCS->GetRootNodes();
 				for(int32 RootNodeIndex = 0; bShouldShowTransform && RootNodeIndex < RootNodes.Num(); ++RootNodeIndex)
@@ -685,87 +479,4 @@ void FSceneComponentDetails::MakeTransformDetails( IDetailLayoutBuilder& DetailB
 	}
 
 }
-
-void FSceneComponentDetails::OnMobilityChanged(TSharedPtr<IPropertyHandle> MobilityProperty)
-{
-	if (!MobilityProperty.IsValid() || CachedSelectedSceneComponents.Num() == 0)
-	{
-		return;
-	}
-
-	// in case GetValue() fails (like with multiple values selected)
-	uint8 MobilityValue = Cast<const USceneComponent>(CachedSelectedSceneComponents[0].Get())->Mobility;
-	MobilityProperty->GetValue(MobilityValue);
-
-	// Attached parent components can't be more mobile than their children. This means that 
-	// certain mobility hierarchy structures are disallowed. So we have to walk the hierarchy 
-	// and alter parent/child components as a result of this property change.
-
-	// track how many other components we had to change
-	int32 NumMobilityChanges = 0; 
-
-	// Movable components can only have movable sub-components
-	if (MobilityValue == EComponentMobility::Movable) 
-	{
-		for (TArray< TWeakObjectPtr<UObject> >::TConstIterator ObjectIt(CachedSelectedSceneComponents); ObjectIt; ++ObjectIt)
-		{
-			NumMobilityChanges += SetDecedentMobility(Cast<const USceneComponent>(ObjectIt->Get()), EComponentMobility::Movable);
-		}
-	}
-	else if (MobilityValue == EComponentMobility::Stationary)
-	{
-		// a functor for checking if we should change a component's Mobility
-		struct FMobilityEqualityFunctor
-		{
-			bool operator()(EComponentMobility::Type CurrentMobility, EComponentMobility::Type CheckValue)
-			{
-				return CurrentMobility == CheckValue;
-			}
-		};
-		FMobilityEqualityFunctor EquivalenceFunctor;
-
-		// a delegate for checking if components are Static
-		FMobilityQueryDelegate IsStaticDelegate  = FMobilityQueryDelegate::CreateRaw(&EquivalenceFunctor, &FMobilityEqualityFunctor::operator(), EComponentMobility::Static);
-		// a delegate for checking if components are Movable
-		FMobilityQueryDelegate IsMovableDelegate = FMobilityQueryDelegate::CreateRaw(&EquivalenceFunctor, &FMobilityEqualityFunctor::operator(), EComponentMobility::Movable);
-
-		// Stationary components can't have Movable parents, and can't have Static children
-		for (TArray< TWeakObjectPtr<UObject> >::TConstIterator ObjectIt(CachedSelectedSceneComponents); ObjectIt; ++ObjectIt)
-		{
-			USceneComponent const* SelectedSceneComponent = Cast<const USceneComponent>(ObjectIt->Get());
-
-			// if any decedents are static, change them to stationary (or movable for static meshes)
-			NumMobilityChanges += SetDecedentMobility(SelectedSceneComponent, EComponentMobility::Stationary, IsStaticDelegate);
-
-			// if any ancestors are movable, change them to stationary (or static for static meshes)
-			NumMobilityChanges += SetAncestorMobility(SelectedSceneComponent, EComponentMobility::Stationary, IsMovableDelegate);
-		}
-	}
-	else // if MobilityValue == Static
-	{
-		// ensure we have the mobility we expected (in case someone adds a new one)
-		ensure(MobilityValue == EComponentMobility::Static); 
-
-		// Static components can only have Static parents
-		for (TArray< TWeakObjectPtr<UObject> >::TConstIterator ObjectIt(CachedSelectedSceneComponents); ObjectIt; ++ObjectIt)
-		{
-			NumMobilityChanges += SetAncestorMobility(Cast<const USceneComponent>(ObjectIt->Get()), EComponentMobility::Static);
-		}
-	}
-
-	// if we altered any components (other than the ones selected), then notify the user
-	if (NumMobilityChanges > 0)
-	{
-		FText NotificationText = LOCTEXT("MobilityAlteredSingularNotification", "Caused 1 component to also change Mobility");
-		if (NumMobilityChanges > 1)
-		{
-			NotificationText = FText::Format(LOCTEXT("MobilityAlteredPluralNotification", "Caused {0} other components to also change Mobility"), FText::AsNumber(NumMobilityChanges));
-		}		
-		FNotificationInfo Info(NotificationText);
-		Info.bFireAndForget = true;
-		Info.bUseThrobber   = true;
-		FSlateNotificationManager::Get().AddNotification(Info);
-	}
-}
-
 #undef LOCTEXT_NAMESPACE
diff --git a/Engine/Source/Editor/DetailCustomizations/Private/SceneComponentDetails.h b/Engine/Source/Editor/DetailCustomizations/Private/SceneComponentDetails.h
--- a/Engine/Source/Editor/DetailCustomizations/Private/SceneComponentDetails.h
+++ b/Engine/Source/Editor/DetailCustomizations/Private/SceneComponentDetails.h
@@ -21,28 +21,6 @@ private:
 	void OnMobilityChanged(ECheckBoxState InCheckedState, EComponentMobility::Type InMobility);
 
 	FText GetMobilityToolTip() const;
-	
-	/**
-	 * When a scene component's Mobility is altered, we need to make sure the scene hierarchy is
-	 * updated. Parents can't be more mobile than their children. This means that certain
-	 * mobility hierarchy structures are disallowed, like:
-	 *
-	 *   Movable
-	 *   |-Stationary   <-- NOT allowed
-	 *   Movable
-	 *   |-Static       <-- NOT allowed
-	 *   Stationary
-	 *   |-Static       <-- NOT allowed
-	 *
-	 * This method walks the hierarchy and alters parent/child component's Mobility as a result of
-	 * this property change.
-	 *
-	 * @param   MobilityPropertyHandle	A handle to the Mobility property that was altered
-	 */
-	void OnMobilityChanged(TSharedPtr<IPropertyHandle> MobilityPropertyHandle);
-
-	/** A collection of USceneComponents that were selected last time CustomizeDetails() was ran */
-	TArray< TWeakObjectPtr<UObject> > CachedSelectedSceneComponents;
 private:
 	TSharedPtr<IPropertyHandle> MobilityHandle;
 };
diff --git a/Engine/Source/Runtime/Engine/Classes/Components/SceneComponent.h b/Engine/Source/Runtime/Engine/Classes/Components/SceneComponent.h
--- a/Engine/Source/Runtime/Engine/Classes/Components/SceneComponent.h
+++ b/Engine/Source/Runtime/Engine/Classes/Components/SceneComponent.h
@@ -530,8 +530,13 @@ public:
 	virtual void PostInterpChange(UProperty* PropertyThatChanged) override;
 	virtual void BeginDestroy() override;
 	static void AddReferencedObjects(UObject* InThis, FReferenceCollector& Collector);
-	// End UObject Interface
 
+#if WITH_EDITOR
+	virtual void PostEditChangeProperty(FPropertyChangedEvent& PropertyChangedEvent) override;
+	virtual void PostEditChangeChainProperty(FPropertyChangedChainEvent& PropertyChangedEvent) override;
+#endif
+
+	// End UObject Interface
 protected:
 	/**
 	 * Internal helper, for use from MoveComponent().  Special codepath since the normal setters call MoveComponent.
diff --git a/Engine/Source/Runtime/Engine/Private/ComponentUtils.cpp b/Engine/Source/Runtime/Engine/Private/ComponentUtils.cpp
new file mode 100644
--- /dev/null
+++ b/Engine/Source/Runtime/Engine/Private/ComponentUtils.cpp
@@ -0,0 +1,87 @@
+// Copyright 1998-2015 Epic Games, Inc. All Rights Reserved.
+
+#include "EnginePrivate.h"
+#include "ComponentUtils.h"
+#include "Engine/SCS_Node.h"
+#include "Engine/SimpleConstructionScript.h"
+
+namespace ComponentUtils
+{
+
+	/**
+	 * A helper for retrieving the simple-construction-script that this component belongs in.
+	 *
+	 * @param  Component	The component you want the SCS for.
+	 * @return The component's blueprint SCS (NULL if one wasn't found).
+	 */
+	USimpleConstructionScript const* GetSimpleConstructionScript(USceneComponent const* Component)
+	{
+		USimpleConstructionScript const* BlueprintSCS = NULL;
+
+		check(Component);
+		UObject const* ComponentOuter = Component->GetOuter();
+
+		if(UBlueprint const* const OuterBlueprint = Cast<UBlueprint const>(ComponentOuter))
+		{
+			BlueprintSCS = OuterBlueprint->SimpleConstructionScript;
+		}
+		else if(UBlueprintGeneratedClass const* const GeneratedClass = Cast<UBlueprintGeneratedClass>(ComponentOuter))
+		{
+			BlueprintSCS = GeneratedClass->SimpleConstructionScript;
+		}
+
+		return BlueprintSCS;
+	}
+
+	USCS_Node* FindCorrespondingSCSNode(USceneComponent const* ComponentObj)
+	{
+		USimpleConstructionScript const* BlueprintSCS = GetSimpleConstructionScript(ComponentObj);
+		if(BlueprintSCS == NULL)
+		{
+			return NULL;
+		}
+
+		TArray<USCS_Node*> AllSCSNodes = BlueprintSCS->GetAllNodes();
+		for(int32 SCSNodeIndex = 0; SCSNodeIndex < AllSCSNodes.Num(); ++SCSNodeIndex)
+		{
+			USCS_Node* SCSNode = AllSCSNodes[SCSNodeIndex];
+			if(SCSNode->ComponentTemplate == ComponentObj)
+			{
+				return SCSNode;
+			}
+		}
+
+		return NULL;
+	}
+
+	/**
+	 * A static helper function used to retrieve a component's scene parent
+	 *
+	 * @param  SceneComponentObject	The component you want the attached parent for
+	 * @return A pointer to the component's scene parent (NULL if there was not one)
+	 */
+	USceneComponent* GetAttachedParent(USceneComponent const* SceneComponentObject)
+	{
+		USceneComponent* SceneParent = SceneComponentObject->AttachParent;
+		if(SceneParent == nullptr)
+		{
+			USCS_Node* const SCSNode = FindCorrespondingSCSNode(SceneComponentObject);
+			// if we didn't find a corresponding simple-construction-script node
+			if(SCSNode == nullptr)
+			{
+				return nullptr;
+			}
+
+			USimpleConstructionScript const* BlueprintSCS = GetSimpleConstructionScript(SceneComponentObject);
+			check(BlueprintSCS != nullptr);
+
+			USCS_Node* const ParentSCSNode = BlueprintSCS->FindParentNode(SCSNode);
+			if(ParentSCSNode != nullptr)
+			{
+				SceneParent = Cast<USceneComponent>(ParentSCSNode->ComponentTemplate);
+			}
+		}
+
+		return SceneParent;
+	}
+}
diff --git a/Engine/Source/Runtime/Engine/Private/Components/SceneComponent.cpp b/Engine/Source/Runtime/Engine/Private/Components/SceneComponent.cpp
--- a/Engine/Source/Runtime/Engine/Private/Components/SceneComponent.cpp
+++ b/Engine/Source/Runtime/Engine/Private/Components/SceneComponent.cpp
@@ -10,6 +10,10 @@
 #include "Net/UnrealNetwork.h"
 #include "GameFramework/PhysicsVolume.h"
 #include "ComponentReregisterContext.h"
+#include "ComponentUtils.h"
+#include "Engine/SCS_Node.h"
+#include "SNotificationList.h"
+#include "NotificationManager.h"
 
 #define LOCTEXT_NAMESPACE "SceneComponent"
 
@@ -51,6 +55,234 @@ void USceneComponent::AddReferencedObjects(UObject* InThis, FReferenceCollector&
 	Super::AddReferencedObjects(InThis, Collector);
 }
 
+#if WITH_EDITOR
+
+/**
+ * A static helper function used meant as the default for determining if a 
+ * component's Mobility should be overridden.
+ * 
+ * @param  CurrentMobility	The component's current Mobility setting
+ * @param  NewMobility		The proposed new Mobility for the component in question
+ * @return True if the two mobilities are not equal (false if they are equal)
+ */
+static bool AreMobilitiesDifferent(EComponentMobility::Type CurrentMobility, EComponentMobility::Type NewMobility)
+{
+	return CurrentMobility != NewMobility;
+}
+DECLARE_DELEGATE_RetVal_OneParam(bool, FMobilityQueryDelegate, EComponentMobility::Type);
+
+/**
+ * A static helper function that recursively alters the Mobility property for all 
+ * sub-components (descending from the specified USceneComponent)
+ * 
+ * @param  SceneComponentObject		The component whose sub-components you want to alter
+ * @param  NewMobilityType			The Mobility type you want to switch sub-components over to
+ * @param  ShouldOverrideMobility	A delegate used to determine if a sub-component's Mobility should be overridden
+ *									(if left unset it will default to the AreMobilitiesDifferent() function)
+ * @return The number of decedents that had their mobility altered.
+ */
+static int32 SetDecedentMobility(USceneComponent const* SceneComponentObject, EComponentMobility::Type NewMobilityType, FMobilityQueryDelegate ShouldOverrideMobility = FMobilityQueryDelegate())
+{
+	if (!ensure(SceneComponentObject != NULL))
+	{
+		return 0;
+	}
+
+	TArray<USceneComponent*> AttachedChildren = SceneComponentObject->AttachChildren;
+	// gather children for component templates
+	USCS_Node* SCSNode = ComponentUtils::FindCorrespondingSCSNode(SceneComponentObject);
+	if (SCSNode != NULL)
+	{
+		// gather children from the SCSNode
+		for (int32 ChildIndex = 0; ChildIndex < SCSNode->ChildNodes.Num(); ++ChildIndex)
+		{
+			USCS_Node* SCSChild = SCSNode->ChildNodes[ChildIndex];
+
+			USceneComponent* ChildSceneComponent = Cast<USceneComponent>(SCSChild->ComponentTemplate);
+			if (ChildSceneComponent != NULL)
+			{
+				AttachedChildren.Add(ChildSceneComponent);
+			}
+		}
+	}
+
+	if (!ShouldOverrideMobility.IsBound())
+	{
+		ShouldOverrideMobility = FMobilityQueryDelegate::CreateStatic(&AreMobilitiesDifferent, NewMobilityType);
+	}
+
+	int32 NumDecendentsChanged = 0;
+	// recursively alter the mobility for children and deeper decedents 
+	for (int32 ChildIndex = 0; ChildIndex < AttachedChildren.Num(); ++ChildIndex)
+	{
+		USceneComponent* ChildSceneComponent = AttachedChildren[ChildIndex];
+
+		if (ShouldOverrideMobility.Execute(ChildSceneComponent->Mobility))
+		{
+			// USceneComponents shouldn't be set Stationary 
+			if ((NewMobilityType == EComponentMobility::Stationary) && ChildSceneComponent->IsA(UStaticMeshComponent::StaticClass()))
+			{
+				// make it Movable (because it is acceptable for Stationary parents to have Movable children)
+				ChildSceneComponent->Mobility = EComponentMobility::Movable;
+			}
+			else
+			{
+				ChildSceneComponent->Mobility = NewMobilityType;
+			}
+			++NumDecendentsChanged;
+		}
+		NumDecendentsChanged += SetDecedentMobility(ChildSceneComponent, NewMobilityType, ShouldOverrideMobility);
+	}
+
+	return NumDecendentsChanged;
+}
+
+/**
+ * A static helper function that alters the Mobility property for all ancestor
+ * components (ancestors of the specified USceneComponent).
+ * 
+ * @param  SceneComponentObject		The component whose attached ancestors you want to alter
+ * @param  NewMobilityType			The Mobility type you want to switch ancestor components over to
+ * @param  ShouldOverrideMobility	A delegate used to determine if a ancestor's Mobility should be overridden
+ *									(if left unset it will default to the AreMobilitiesDifferent() function)
+ * @return The number of ancestors that had their mobility altered.
+ */
+static int32 SetAncestorMobility(USceneComponent const* SceneComponentObject, EComponentMobility::Type NewMobilityType, FMobilityQueryDelegate ShouldOverrideMobility = FMobilityQueryDelegate())
+{
+	if (!ensure(SceneComponentObject != NULL))
+	{
+		return 0;
+	}
+
+	if (!ShouldOverrideMobility.IsBound())
+	{
+		ShouldOverrideMobility = FMobilityQueryDelegate::CreateStatic(&AreMobilitiesDifferent, NewMobilityType);
+	}
+
+	int32 MobilityAlteredCount = 0;
+	while(USceneComponent* AttachedParent = ComponentUtils::GetAttachedParent(SceneComponentObject))
+	{
+		if (ShouldOverrideMobility.Execute(AttachedParent->Mobility))
+		{
+			// USceneComponents shouldn't be set Stationary 
+			if ((NewMobilityType == EComponentMobility::Stationary) && AttachedParent->IsA(UStaticMeshComponent::StaticClass()))
+			{
+				// make it Static (because it is acceptable for Stationary children to have Static parents)
+				AttachedParent->Mobility = EComponentMobility::Static;
+			}
+			else
+			{
+				AttachedParent->Mobility = NewMobilityType;
+			}
+			++MobilityAlteredCount;
+		}
+		SceneComponentObject = AttachedParent;
+	}
+
+	return MobilityAlteredCount;
+}
+
+/**
+* When a scene component's Mobility is altered, we need to make sure the scene hierarchy is
+* updated. Parents can't be more mobile than their children. This means that certain
+* mobility hierarchy structures are disallowed, like:
+*
+*   Movable
+*   |-Stationary   <-- NOT allowed
+*   Movable
+*   |-Static       <-- NOT allowed
+*   Stationary
+*   |-Static       <-- NOT allowed
+*
+* This method walks the hierarchy and alters parent/child component's Mobility as a result of
+* this property change.
+*/
+static void UpdateAttachedMobility(USceneComponent* ComponentThatChanged)
+{
+	// Attached parent components can't be more mobile than their children. This means that 
+	// certain mobility hierarchy structures are disallowed. So we have to walk the hierarchy 
+	// and alter parent/child components as a result of this property change.
+
+	// track how many other components we had to change
+	int32 NumMobilityChanges = 0;
+
+	// Movable components can only have movable sub-components
+	if(ComponentThatChanged->Mobility == EComponentMobility::Movable)
+	{
+		NumMobilityChanges += SetDecedentMobility(ComponentThatChanged, EComponentMobility::Movable);
+	}
+	else if(ComponentThatChanged->Mobility == EComponentMobility::Stationary)
+	{
+		// a functor for checking if we should change a component's Mobility
+		struct FMobilityEqualityFunctor
+		{
+			bool operator()(EComponentMobility::Type CurrentMobility, EComponentMobility::Type CheckValue)
+			{
+				return CurrentMobility == CheckValue;
+			}
+		};
+		FMobilityEqualityFunctor EquivalenceFunctor;
+
+		// a delegate for checking if components are Static
+		FMobilityQueryDelegate IsStaticDelegate  = FMobilityQueryDelegate::CreateRaw(&EquivalenceFunctor, &FMobilityEqualityFunctor::operator(), EComponentMobility::Static);
+		// a delegate for checking if components are Movable
+		FMobilityQueryDelegate IsMovableDelegate = FMobilityQueryDelegate::CreateRaw(&EquivalenceFunctor, &FMobilityEqualityFunctor::operator(), EComponentMobility::Movable);
+
+		// if any decedents are static, change them to stationary (or movable for static meshes)
+		NumMobilityChanges += SetDecedentMobility(ComponentThatChanged, EComponentMobility::Stationary, IsStaticDelegate);
+
+		// if any ancestors are movable, change them to stationary (or static for static meshes)
+		NumMobilityChanges += SetAncestorMobility(ComponentThatChanged, EComponentMobility::Stationary, IsMovableDelegate);
+	}
+	else // if MobilityValue == Static
+	{
+		// ensure we have the mobility we expected (in case someone adds a new one)
+		ensure(ComponentThatChanged->Mobility == EComponentMobility::Static);
+
+		NumMobilityChanges += SetAncestorMobility(ComponentThatChanged, EComponentMobility::Static);
+	}
+
+	// if we altered any components (other than the ones selected), then notify the user
+	if(NumMobilityChanges > 0)
+	{
+		FText NotificationText = LOCTEXT("MobilityAlteredSingularNotification", "Caused 1 component to also change Mobility");
+		if(NumMobilityChanges > 1)
+		{
+			NotificationText = FText::Format(LOCTEXT("MobilityAlteredPluralNotification", "Caused {0} other components to also change Mobility"), FText::AsNumber(NumMobilityChanges));
+		}
+		FNotificationInfo Info(NotificationText);
+		Info.bFireAndForget = true;
+		Info.bUseThrobber   = true;
+		FSlateNotificationManager::Get().AddNotification(Info);
+	}
+}
+
+void USceneComponent::PostEditChangeProperty(FPropertyChangedEvent& PropertyChangedEvent)
+{
+	// Note: This must be called before UActorComponent::PostEditChangeChainProperty is called because this component will be reset when UActorComponent reruns construction scripts 
+	static FName MobilityName("Mobility");
+	if(PropertyChangedEvent.Property && PropertyChangedEvent.Property->GetFName() == MobilityName)
+	{
+		UpdateAttachedMobility(this);
+	}
+
+	Super::PostEditChangeProperty(PropertyChangedEvent);
+}
+
+void USceneComponent::PostEditChangeChainProperty(FPropertyChangedChainEvent& PropertyChangedEvent)
+{
+	// Note: This must be called before UActorComponent::PostEditChangeChainProperty is called because this component will be reset when UActorComponent reruns construction scripts 
+	static FName MobilityName("Mobility");
+	if(PropertyChangedEvent.Property && PropertyChangedEvent.Property->GetFName() == MobilityName)
+	{
+		UpdateAttachedMobility(this);
+	}
+
+	Super::PostEditChangeChainProperty(PropertyChangedEvent);
+}
+
+#endif
+
 FTransform USceneComponent::CalcNewComponentToWorld(const FTransform& NewRelativeTransform, const USceneComponent* Parent) const
 {
 	Parent = Parent ? Parent : AttachParent;
diff --git a/Engine/Source/Runtime/Engine/Public/ComponentUtils.h b/Engine/Source/Runtime/Engine/Public/ComponentUtils.h
new file mode 100644
--- /dev/null
+++ b/Engine/Source/Runtime/Engine/Public/ComponentUtils.h
@@ -0,0 +1,36 @@
+// Copyright 1998-2015 Epic Games, Inc. All Rights Reserved.
+
+#pragma once
+
+class USimpleConstructionScript;
+class USCS_Node;
+class USceneComponent;
+
+namespace ComponentUtils
+{
+	/**
+	 * A helper for retrieving the simple-construction-script that this component belongs in.
+	 * 
+	 * @param  Component	The component you want the SCS for.
+	 * @return The component's blueprint SCS (NULL if one wasn't found).
+	 */
+	ENGINE_API USimpleConstructionScript const* GetSimpleConstructionScript(USceneComponent const* Component);
+
+	/**
+	 * A static helper function for retrieving the simple-construction-script node that
+	 * corresponds to the specified scene component template.
+	 * 
+	 * @param  ComponentObj	The component you want to find a USCS_Node for
+	 * @return A USCS_Node pointer corresponding to the specified component (NULL if we didn't find one)
+	 */
+	ENGINE_API USCS_Node* FindCorrespondingSCSNode(USceneComponent const* ComponentObj);
+
+	/**
+	 * A static helper function used to retrieve a component's scene parent
+	 * 
+	 * @param  SceneComponentObject	The component you want the attached parent for
+	 * @return A pointer to the component's scene parent (NULL if there was not one)
+	 */
+	ENGINE_API USceneComponent* GetAttachedParent(USceneComponent const* SceneComponentObject);
+
+};
\ No newline at end of file
