PVS-Studio has detected that this loop is malformed. Is it meant to be a for loop? an if statement?

diff --git a/Engine/Source/Editor/DetailCustomizations/Private/ComponentMaterialCategory.cpp b/Engine/Source/Editor/DetailCustomizations/Private/ComponentMaterialCategory.cpp
--- a/Engine/Source/Editor/DetailCustomizations/Private/ComponentMaterialCategory.cpp
+++ b/Engine/Source/Editor/DetailCustomizations/Private/ComponentMaterialCategory.cpp
@@ -82,7 +82,11 @@ public:
 					CurMaterial = Material;
 							
 					// We step only once in the iterator if we have a material.  Note: A null material is considered valid
-					return;
+					
+					// PVS-Studio has noticed that this is a fairly unorthodox return statement. Typically a break would 
+					// be the expected control sequence in a nested while loop. This may be correct, however, so for now
+					// we have disabled the warning.
+					return; //-V612
 				}
 				// Out of materials on this component, reset for the next component
 				CurMaterialIndex = 0;
