Made FSlateFontInfoStructCustomization use GET_MEMBER_NAME_CHECKED

diff --git a/Engine/Source/Editor/DetailCustomizations/Private/SlateFontInfoCustomization.cpp b/Engine/Source/Editor/DetailCustomizations/Private/SlateFontInfoCustomization.cpp
--- a/Engine/Source/Editor/DetailCustomizations/Private/SlateFontInfoCustomization.cpp
+++ b/Engine/Source/Editor/DetailCustomizations/Private/SlateFontInfoCustomization.cpp
@@ -14,15 +14,19 @@ TSharedRef<IPropertyTypeCustomization> FSlateFontInfoStructCustomization::MakeIn
 
 void FSlateFontInfoStructCustomization::CustomizeHeader(TSharedRef<IPropertyHandle> InStructPropertyHandle, FDetailWidgetRow& InHeaderRow, IPropertyTypeCustomizationUtils& InStructCustomizationUtils)
 {
+	static const FName FontObjectPropertyName = GET_MEMBER_NAME_CHECKED(FSlateFontInfo, FontObject);
+	static const FName TypefaceFontNamePropertyName = GET_MEMBER_NAME_CHECKED(FSlateFontInfo, TypefaceFontName);
+	static const FName SizePropertyName = GET_MEMBER_NAME_CHECKED(FSlateFontInfo, Size);
+
 	StructPropertyHandle = InStructPropertyHandle;
 
-	FontObjectProperty = InStructPropertyHandle->GetChildHandle(FName("FontObject"));
+	FontObjectProperty = InStructPropertyHandle->GetChildHandle(FontObjectPropertyName);
 	check(FontObjectProperty.IsValid());
 
-	TypefaceFontNameProperty = InStructPropertyHandle->GetChildHandle(FName("TypefaceFontName"));
+	TypefaceFontNameProperty = InStructPropertyHandle->GetChildHandle(TypefaceFontNamePropertyName);
 	check(TypefaceFontNameProperty.IsValid());
 
-	FontSizeProperty = InStructPropertyHandle->GetChildHandle(FName("Size"));
+	FontSizeProperty = InStructPropertyHandle->GetChildHandle(SizePropertyName);
 	check(FontSizeProperty.IsValid());
 
 	InHeaderRow
