Media: Properly handling empty URLs

diff --git a/Engine/Source/Editor/DetailCustomizations/Private/MediaPlayerCustomization.cpp b/Engine/Source/Editor/DetailCustomizations/Private/MediaPlayerCustomization.cpp
--- a/Engine/Source/Editor/DetailCustomizations/Private/MediaPlayerCustomization.cpp
+++ b/Engine/Source/Editor/DetailCustomizations/Private/MediaPlayerCustomization.cpp
@@ -388,7 +388,7 @@ FString FMediaPlayerCustomization::HandleUrlPickerFileTypeFilter() const
 
 void FMediaPlayerCustomization::HandleUrlPickerPathPicked( const FString& PickedPath )
 {
-	if (PickedPath.StartsWith(TEXT("./")) || PickedPath.Contains(TEXT("://")))
+	if (PickedPath.IsEmpty() || PickedPath.StartsWith(TEXT("./")) || PickedPath.Contains(TEXT("://")))
 	{
 		UrlProperty->SetValue(PickedPath);
 	}
