Fix bad filter string on "image" property for a slate brush

diff --git a/Engine/Source/Editor/DetailCustomizations/Private/SlateBrushCustomization.cpp b/Engine/Source/Editor/DetailCustomizations/Private/SlateBrushCustomization.cpp
--- a/Engine/Source/Editor/DetailCustomizations/Private/SlateBrushCustomization.cpp
+++ b/Engine/Source/Editor/DetailCustomizations/Private/SlateBrushCustomization.cpp
@@ -1197,7 +1197,7 @@ void FSlateBrushStructCustomization::CustomizeChildren( TSharedRef<IPropertyHand
 	TSharedPtr<IPropertyHandle> TintProperty = StructPropertyHandle->GetChildHandle( TEXT("TintColor") );
 	TSharedPtr<IPropertyHandle> ResourceObjectProperty = StructPropertyHandle->GetChildHandle( TEXT("ResourceObject") );
 	
-	StructBuilder.AddChildContent( NSLOCTEXT( "SlateBrushCustomization", "ResourceObjectFilterString", "Resource" ) )
+	StructBuilder.AddChildContent( NSLOCTEXT( "SlateBrushCustomization", "ResourceObjectFilterString", "Image" ) )
 	.NameContent()
 	[
 		ResourceObjectProperty->CreatePropertyNameWidget()
