Fix crash selecting objects in the level editor

diff --git a/Engine/Source/Editor/DetailCustomizations/Private/ComponentTransformDetails.cpp b/Engine/Source/Editor/DetailCustomizations/Private/ComponentTransformDetails.cpp
--- a/Engine/Source/Editor/DetailCustomizations/Private/ComponentTransformDetails.cpp
+++ b/Engine/Source/Editor/DetailCustomizations/Private/ComponentTransformDetails.cpp
@@ -579,11 +579,11 @@ void FComponentTransformDetails::CacheCommonLocationUnits()
 	{
 		LargestValue = CachedLocation.X.GetValue();
 	}
-	if (CachedLocation.X.IsSet() && CachedLocation.Y.GetValue() > LargestValue)
+	if (CachedLocation.Y.IsSet() && CachedLocation.Y.GetValue() > LargestValue)
 	{
 		LargestValue = CachedLocation.Y.GetValue();
 	}
-	if (CachedLocation.X.IsSet() && CachedLocation.Z.GetValue() > LargestValue)
+	if (CachedLocation.Z.IsSet() && CachedLocation.Z.GetValue() > LargestValue)
 	{
 		LargestValue = CachedLocation.Z.GetValue();
 	}
