PostProcess settings cleanup.  The FPostProcessSettings struct now has a customization to show categories for the post process settings inside the struct

diff --git a/Engine/Source/Editor/DetailCustomizations/Private/DetailCustomizations.cpp b/Engine/Source/Editor/DetailCustomizations/Private/DetailCustomizations.cpp
--- a/Engine/Source/Editor/DetailCustomizations/Private/DetailCustomizations.cpp
+++ b/Engine/Source/Editor/DetailCustomizations/Private/DetailCustomizations.cpp
@@ -100,6 +100,7 @@
 #include "AutoReimportDirectoryCustomization.h"
 #include "DistanceDatumStructCustomization.h"
 #include "HierarchicalSimplificationCustomizations.h"
+#include "PostProcessSettingsCustomization.h"
 
 IMPLEMENT_MODULE( FDetailCustomizationsModule, DetailCustomizations );
 
@@ -201,6 +202,7 @@ void FDetailCustomizationsModule::RegisterPropertyTypeCustomizations()
 	RegisterCustomPropertyTypeLayout("AutoReimportWildcard", FOnGetPropertyTypeCustomizationInstance::CreateStatic(&FAutoReimportWildcardCustomization::MakeInstance));
 	RegisterCustomPropertyTypeLayout("DistanceDatum", FOnGetPropertyTypeCustomizationInstance::CreateStatic(&FDistanceDatumStructCustomization::MakeInstance));
 	RegisterCustomPropertyTypeLayout("HierarchicalSimplification", FOnGetPropertyTypeCustomizationInstance::CreateStatic(&FHierarchicalSimplificationCustomizations::MakeInstance));
+	RegisterCustomPropertyTypeLayout("PostProcessSettings", FOnGetPropertyTypeCustomizationInstance::CreateStatic(&FPostProcessSettingsCustomization::MakeInstance));
 }
 
 
diff --git a/Engine/Source/Editor/DetailCustomizations/Private/PostProcessSettingsCustomization.cpp b/Engine/Source/Editor/DetailCustomizations/Private/PostProcessSettingsCustomization.cpp
new file mode 100644
--- /dev/null
+++ b/Engine/Source/Editor/DetailCustomizations/Private/PostProcessSettingsCustomization.cpp
@@ -0,0 +1,54 @@
+// Copyright 1998-2015 Epic Games, Inc. All Rights Reserved.
+
+#include "DetailCustomizationsPrivatePCH.h"
+#include "PostProcessSettingsCustomization.h"
+#include "ObjectEditorUtils.h"
+
+#define LOCTEXT_NAMESPACE "PostProcessSettingsCustomization"
+
+void FPostProcessSettingsCustomization::CustomizeChildren( TSharedRef<IPropertyHandle> StructPropertyHandle, class IDetailChildrenBuilder& StructBuilder, IPropertyTypeCustomizationUtils& StructCustomizationUtils )
+{
+	TMap<FString, IDetailGroup*> CategoryNameToGroupMap;
+
+	uint32 NumChildren = 0;
+	FPropertyAccess::Result Result = StructPropertyHandle->GetNumChildren(NumChildren);
+	if( Result == FPropertyAccess::Success && NumChildren > 0)
+	{
+		for( uint32 ChildIndex = 0; ChildIndex < NumChildren; ++ChildIndex )
+		{
+			TSharedPtr<IPropertyHandle> ChildHandle = StructPropertyHandle->GetChildHandle( ChildIndex );
+			if( ChildHandle.IsValid() && ChildHandle->GetProperty() )
+			{
+				UProperty* Property = ChildHandle->GetProperty();
+				FString Category = FObjectEditorUtils::GetCategory( Property );
+				IDetailGroup* Group = CategoryNameToGroupMap.FindRef( Category );
+
+				if( !Group )
+				{
+					Group = &StructBuilder.AddChildGroup( FName(*Category), FText::FromString( Category ) );
+					CategoryNameToGroupMap.Add( Category, Group );
+				}
+				
+				Group->AddPropertyRow( ChildHandle.ToSharedRef() );
+			}
+		}
+	}
+	
+}
+
+
+void FPostProcessSettingsCustomization::CustomizeHeader( TSharedRef<IPropertyHandle> StructPropertyHandle, class FDetailWidgetRow& HeaderRow, IPropertyTypeCustomizationUtils& StructCustomizationUtils )
+{
+	HeaderRow.NameContent()
+	[
+		StructPropertyHandle->CreatePropertyNameWidget()
+	];
+
+	HeaderRow.ValueContent()
+	[
+		StructPropertyHandle->CreatePropertyValueWidget()
+	];
+}
+
+
+#undef LOCTEXT_NAMESPACE
diff --git a/Engine/Source/Editor/DetailCustomizations/Private/PostProcessSettingsCustomization.h b/Engine/Source/Editor/DetailCustomizations/Private/PostProcessSettingsCustomization.h
new file mode 100644
--- /dev/null
+++ b/Engine/Source/Editor/DetailCustomizations/Private/PostProcessSettingsCustomization.h
@@ -0,0 +1,30 @@
+// Copyright 1998-2015 Epic Games, Inc. All Rights Reserved.
+
+#pragma once
+
+
+/**
+ * Implements a details view customization for the FTimespan structure.
+ */
+class FPostProcessSettingsCustomization
+	: public IPropertyTypeCustomization
+{
+public:
+
+	/**
+	 * Creates an instance of this class.
+	 *
+	 * @return The new instance.
+	 */
+	static TSharedRef<IPropertyTypeCustomization> MakeInstance( )
+	{
+		return MakeShareable(new FPostProcessSettingsCustomization());
+	}
+
+public:
+
+	// IPropertyTypeCustomization interface
+
+	virtual void CustomizeChildren( TSharedRef<class IPropertyHandle> StructPropertyHandle, class IDetailChildrenBuilder& StructBuilder, IPropertyTypeCustomizationUtils& StructCustomizationUtils ) override;
+	virtual void CustomizeHeader( TSharedRef<class IPropertyHandle> StructPropertyHandle, class FDetailWidgetRow& HeaderRow, IPropertyTypeCustomizationUtils& StructCustomizationUtils ) override;
+};
diff --git a/Engine/Source/Runtime/Engine/Classes/Camera/CameraComponent.h b/Engine/Source/Runtime/Engine/Classes/Camera/CameraComponent.h
--- a/Engine/Source/Runtime/Engine/Classes/Camera/CameraComponent.h
+++ b/Engine/Source/Runtime/Engine/Classes/Camera/CameraComponent.h
@@ -58,7 +58,7 @@ class UCameraComponent : public USceneComponent
 	float PostProcessBlendWeight;
 
 	/** Post process settings to use for this camera. Don't forget to check the properties you want to override */
-	UPROPERTY(Interp, BlueprintReadWrite, Category=CameraSettings, meta=(ShowOnlyInnerProperties))
+	UPROPERTY(Interp, BlueprintReadWrite, Category=CameraSettings)
 	struct FPostProcessSettings PostProcessSettings;
 
 	// UActorComponent interface
diff --git a/Engine/Source/Runtime/Engine/Classes/Engine/PostProcessVolume.h b/Engine/Source/Runtime/Engine/Classes/Engine/PostProcessVolume.h
--- a/Engine/Source/Runtime/Engine/Classes/Engine/PostProcessVolume.h
+++ b/Engine/Source/Runtime/Engine/Classes/Engine/PostProcessVolume.h
@@ -20,7 +20,7 @@ class APostProcessVolume : public AVolume, public IInterface_PostProcessVolume
 	GENERATED_UCLASS_BODY()
 
 	/** Post process settings to use for this volume. */
-	UPROPERTY(interp, Category=PostProcessVolume, meta=(ShowOnlyInnerProperties))
+	UPROPERTY(interp, Category=PostProcessVolume, meta=(DisplayName="Post Process Settings"))
 	struct FPostProcessSettings Settings;
 
 	/**
