Fix post process settings categories not having pretty print

diff --git a/Engine/Source/Editor/DetailCustomizations/Private/PostProcessSettingsCustomization.cpp b/Engine/Source/Editor/DetailCustomizations/Private/PostProcessSettingsCustomization.cpp
--- a/Engine/Source/Editor/DetailCustomizations/Private/PostProcessSettingsCustomization.cpp
+++ b/Engine/Source/Editor/DetailCustomizations/Private/PostProcessSettingsCustomization.cpp
@@ -25,7 +25,7 @@ void FPostProcessSettingsCustomization::CustomizeChildren( TSharedRef<IPropertyH
 
 				if( !Group )
 				{
-					Group = &StructBuilder.AddChildGroup( FName(*Category), FText::FromString( Category ) );
+					Group = &StructBuilder.AddChildGroup( FName(*Category), FText::FromString( FName::NameToDisplayString( Category, false ) ) );
 					CategoryNameToGroupMap.Add( Category, Group );
 				}
 				
