[UE-11839] Fix propagation of material swaps when editing component templates owned by the CDO.

change summary:
- relocated FMaterialIterator::SwapMaterial so it can be called on both archetype and instances
- modified the FComponentMaterialCategory::OnMaterialChanged() delegate to additionally propagate material property value changes to archetype instances if the selected component is a template
- updated the scoping of the FNavigationLockContext in OnMaterialChanged() to ensure that it remains valid around the pre/post edit change sequence

diff --git a/Engine/Source/Editor/DetailCustomizations/Private/ComponentMaterialCategory.cpp b/Engine/Source/Editor/DetailCustomizations/Private/ComponentMaterialCategory.cpp
--- a/Engine/Source/Editor/DetailCustomizations/Private/ComponentMaterialCategory.cpp
+++ b/Engine/Source/Editor/DetailCustomizations/Private/ComponentMaterialCategory.cpp
@@ -122,26 +122,6 @@ public:
 	 */
 	int32 GetMaterialIndex() const { return CurMaterialIndex; }
 
-	/**
-	 * Swaps the current material the iterator is stopped on with the new material
-	 *
-	 * @param NewMaterial	The material to swap in
-	 */
-	void SwapMaterial( UMaterialInterface* NewMaterial )
-	{
-		UPrimitiveComponent* PrimitiveComp = Cast<UPrimitiveComponent>( CurComponent );
-		UDecalComponent* DecalComponent = Cast<UDecalComponent>( CurComponent );
-
-		if( PrimitiveComp )
-		{
-			PrimitiveComp->SetMaterial( CurMaterialIndex, NewMaterial );
-		}
-		else if( DecalComponent )
-		{
-			DecalComponent->SetMaterial( CurMaterialIndex, NewMaterial );
-		}
-	}
-
 	/**
 	 * @return The current component using the current material
 	 */
@@ -243,6 +223,22 @@ void FComponentMaterialCategory::OnMaterialChanged( UMaterialInterface* NewMater
 	// Whether or not we made a transaction and need to end it
 	bool bMadeTransaction = false;
 
+	// Lambda to swap materials on a given component at the given slot index
+	auto SwapMaterialLambda = []( UActorComponent* InComponent, int32 InElementIndex, UMaterialInterface* InNewMaterial )
+	{
+		UPrimitiveComponent* PrimitiveComp = Cast<UPrimitiveComponent>( InComponent );
+		UDecalComponent* DecalComponent = Cast<UDecalComponent>( InComponent );
+
+		if( PrimitiveComp )
+		{
+			PrimitiveComp->SetMaterial( InElementIndex, InNewMaterial );
+		}
+		else if( DecalComponent )
+		{
+			DecalComponent->SetMaterial( InElementIndex, InNewMaterial );
+		}
+	};
+
 	// Scan the selected actors mesh components for the old material and swap it with the new material 
 	for( FMaterialIterator It( SelectedComponents ); It; ++It )
 	{
@@ -252,8 +248,6 @@ void FComponentMaterialCategory::OnMaterialChanged( UMaterialInterface* NewMater
 
 		if( CurrentComponent )
 		{
-			AActor* Actor = CurrentComponent->GetOwner();
-
 			// Component materials can be replaced if they are not created from a blueprint (not exposed to the user) and have material overrides on the component
 			bool bCanBeReplaced = 
 				( CurrentComponent->IsA( UMeshComponent::StaticClass() ) ||
@@ -294,9 +288,12 @@ void FComponentMaterialCategory::OnMaterialChanged( UMaterialInterface* NewMater
 					EditChangeObject = CastChecked<ULandscapeComponent>(CurrentComponent)->GetLandscapeProxy();
 				}
 
-				if( Actor )
+				// Add a navigation update lock only if the component world is valid
+				TSharedPtr<FNavigationLockContext> NavUpdateLock;
+				UWorld* World = CurrentComponent->GetWorld();
+				if( World )
 				{
-					FNavigationLockContext NavUpdateLock(Actor->GetWorld(), ENavigationLockReason::MaterialUpdate);
+					NavUpdateLock = MakeShareable( new FNavigationLockContext(World, ENavigationLockReason::MaterialUpdate) );
 				}
 
 				EditChangeObject->PreEditChange( MaterialProperty );
@@ -306,7 +303,7 @@ void FComponentMaterialCategory::OnMaterialChanged( UMaterialInterface* NewMater
 					NotifyHook->NotifyPreChange( MaterialProperty );
 				}
 
-				It.SwapMaterial( NewMaterial );
+				SwapMaterialLambda( CurrentComponent, It.GetMaterialIndex(), NewMaterial );
 
 				FPropertyChangedEvent PropertyChangedEvent( MaterialProperty );
 				EditChangeObject->PostEditChangeProperty( PropertyChangedEvent );
@@ -315,6 +312,35 @@ void FComponentMaterialCategory::OnMaterialChanged( UMaterialInterface* NewMater
 				{
 					NotifyHook->NotifyPostChange( PropertyChangedEvent, MaterialProperty );
 				}
+
+				TArray<UObject*> ArchetypeInstances;
+				if( CurrentComponent->IsTemplate() && !FApp::IsGame() )
+				{
+					// Propagate material change to instances of the edited component template
+					CurrentComponent->GetArchetypeInstances(ArchetypeInstances);
+					for( auto ArchetypeInstance : ArchetypeInstances )
+					{
+						CurrentComponent = CastChecked<UActorComponent>( ArchetypeInstance );
+						if( CurrentComponent->IsA<ULandscapeComponent>() )
+						{
+							ArchetypeInstance = CastChecked<ULandscapeComponent>(CurrentComponent)->GetLandscapeProxy();
+						}
+						
+						// Reset the navigation update lock if necessary
+						UWorld* PreviousWorld = World;
+						World = CurrentComponent->GetWorld();
+						if( PreviousWorld != World )
+						{
+							NavUpdateLock = MakeShareable( new FNavigationLockContext(World, ENavigationLockReason::MaterialUpdate) );
+						}
+
+						ArchetypeInstance->PreEditChange( MaterialProperty );
+
+						SwapMaterialLambda( CurrentComponent, It.GetMaterialIndex(), NewMaterial );
+
+						ArchetypeInstance->PostEditChangeProperty( PropertyChangedEvent );
+					}
+				}
 			}
 		}
 	}
