#HLOD

- Renamed DesiredBoundSize to DesiredBoundRadius - more clear and it implies sphere
- MergeSetting is now exposed if you decide not to simplify
- Shrink after all done so that we don't use that reserved mass memory

diff --git a/Engine/Source/Developer/MeshUtilities/Public/MeshUtilities.h b/Engine/Source/Developer/MeshUtilities/Public/MeshUtilities.h
--- a/Engine/Source/Developer/MeshUtilities/Public/MeshUtilities.h
+++ b/Engine/Source/Developer/MeshUtilities/Public/MeshUtilities.h
@@ -68,57 +68,6 @@ public:
 		) = 0;
 };
 
-/**
- * Mesh merging settings
- */
-struct FMeshMergingSettings
-{
-	/** Whether to generate lightmap UVs for a merged mesh*/
-	bool bGenerateLightMapUV;
-	
-	/** Target UV channel in a merged mesh for a lightmap */
-	int32 TargetLightMapUVChannel;
-
-	/** Target lightmap resolution */
-	int32 TargetLightMapResolution;
-		
-	/** Whether we should import vertex colors into merged mesh */
-	bool bImportVertexColors;
-	
-	/** Whether merged mesh should have pivot at world origin, or at first merged component otherwise */
-	bool bPivotPointAtZero;
-
-	/** Whether to merge source materials into one flat material */
-	bool bMergeMaterials;
-	/** Whether to export normal maps for material merging */
-	bool bExportNormalMap;
-	/** Whether to export metallic maps for material merging */
-	bool bExportMetallicMap;
-	/** Whether to export roughness maps for material merging */
-	bool bExportRoughnessMap;
-	/** Whether to export specular maps for material merging */
-	bool bExportSpecularMap;
-	/** Merged material texture atlas resolution */
-	int32 MergedMaterialAtlasResolution;
-		
-	/** Default settings. */
-	FMeshMergingSettings()
-		: bGenerateLightMapUV(false)
-		, TargetLightMapUVChannel(1)
-		, TargetLightMapResolution(256)
-		, bImportVertexColors(false)
-		, bPivotPointAtZero(false)
-		, bMergeMaterials(false)
-		, bExportNormalMap(true)
-		, bExportMetallicMap(false)
-		, bExportRoughnessMap(false)
-		, bExportSpecularMap(false)
-		, MergedMaterialAtlasResolution(1024)
-	{
-	}
-};
-
-
 
 /**
  * Mesh reduction module interface.
diff --git a/Engine/Source/Editor/DetailCustomizations/Private/HierarchicalSimplificationCustomizations.cpp b/Engine/Source/Editor/DetailCustomizations/Private/HierarchicalSimplificationCustomizations.cpp
--- a/Engine/Source/Editor/DetailCustomizations/Private/HierarchicalSimplificationCustomizations.cpp
+++ b/Engine/Source/Editor/DetailCustomizations/Private/HierarchicalSimplificationCustomizations.cpp
@@ -38,18 +38,27 @@ void FHierarchicalSimplificationCustomizations::CustomizeChildren( TSharedRef<IP
 		PropertyHandles.Add(PropertyName, ChildHandle);
 	}
 
-	TSharedPtr< IPropertyHandle > SimplifyMeshPropertyHandle = PropertyHandles.FindChecked(GET_MEMBER_NAME_CHECKED(FHierarchicalSimplification, bSimplifyMesh));
+	SimplifyMeshPropertyHandle = PropertyHandles.FindChecked(GET_MEMBER_NAME_CHECKED(FHierarchicalSimplification, bSimplifyMesh));
 
 	IDetailPropertyRow& SimplifyMeshRow = ChildBuilder.AddChildProperty(SimplifyMeshPropertyHandle.ToSharedRef());
 	SimplifyMeshRow.Visibility(TAttribute<EVisibility>(this, &FHierarchicalSimplificationCustomizations::IsSimplifyMeshVisible));
 
+	TSharedPtr< IPropertyHandle > ProxyMeshSettingPropertyHandle = PropertyHandles.FindChecked(GET_MEMBER_NAME_CHECKED(FHierarchicalSimplification, ProxySetting));
+	TSharedPtr< IPropertyHandle > MergeMeshSettingPropertyHandle = PropertyHandles.FindChecked(GET_MEMBER_NAME_CHECKED(FHierarchicalSimplification, MergeSetting));
+
 	for( auto Iter(PropertyHandles.CreateConstIterator()); Iter; ++Iter  )
 	{
-		if (Iter.Value() != SimplifyMeshPropertyHandle)
+		if (Iter.Value() != SimplifyMeshPropertyHandle && Iter.Value() != ProxyMeshSettingPropertyHandle && Iter.Value() != MergeMeshSettingPropertyHandle)
 		{
 			ChildBuilder.AddChildProperty(Iter.Value().ToSharedRef());
 		}
 	}
+
+	IDetailPropertyRow& ProxyMeshSettingRow = ChildBuilder.AddChildProperty(ProxyMeshSettingPropertyHandle.ToSharedRef());
+	ProxyMeshSettingRow.Visibility(TAttribute<EVisibility>(this, &FHierarchicalSimplificationCustomizations::IsProxyMeshSettingVisible));
+
+	IDetailPropertyRow& MergeMeshSettingRow = ChildBuilder.AddChildProperty(MergeMeshSettingPropertyHandle.ToSharedRef());
+	MergeMeshSettingRow.Visibility(TAttribute<EVisibility>(this, &FHierarchicalSimplificationCustomizations::IsMergeMeshSettingVisible));
 }
 
 EVisibility FHierarchicalSimplificationCustomizations::IsSimplifyMeshVisible() const
@@ -63,3 +72,27 @@ EVisibility FHierarchicalSimplificationCustomizations::IsSimplifyMeshVisible() c
 	return EVisibility::Hidden;
 }
 
+EVisibility FHierarchicalSimplificationCustomizations::IsProxyMeshSettingVisible() const
+{
+	bool bSimplifyMesh;
+
+	if (SimplifyMeshPropertyHandle->GetValue(bSimplifyMesh) == FPropertyAccess::Result::Success)
+	{
+		if (IsSimplifyMeshVisible() == EVisibility::Visible && bSimplifyMesh)
+		{
+			return EVisibility::Visible;
+		}
+	}
+		
+	return EVisibility::Hidden;
+}
+
+EVisibility FHierarchicalSimplificationCustomizations::IsMergeMeshSettingVisible() const
+{
+	if(IsProxyMeshSettingVisible() == EVisibility::Hidden)
+	{
+		return EVisibility::Visible;
+	}
+
+	return EVisibility::Hidden;
+}
\ No newline at end of file
diff --git a/Engine/Source/Editor/DetailCustomizations/Private/HierarchicalSimplificationCustomizations.h b/Engine/Source/Editor/DetailCustomizations/Private/HierarchicalSimplificationCustomizations.h
--- a/Engine/Source/Editor/DetailCustomizations/Private/HierarchicalSimplificationCustomizations.h
+++ b/Engine/Source/Editor/DetailCustomizations/Private/HierarchicalSimplificationCustomizations.h
@@ -13,5 +13,8 @@ public:
 
 protected:
 	EVisibility IsSimplifyMeshVisible() const;
-	TSharedPtr< IPropertyHandle > SimplifyMeshHandle;
+	EVisibility IsProxyMeshSettingVisible() const;
+	EVisibility IsMergeMeshSettingVisible() const;
+
+	TSharedPtr< IPropertyHandle > SimplifyMeshPropertyHandle;
 };
diff --git a/Engine/Source/Editor/UnrealEd/Private/HierarchicalLOD.cpp b/Engine/Source/Editor/UnrealEd/Private/HierarchicalLOD.cpp
--- a/Engine/Source/Editor/UnrealEd/Private/HierarchicalLOD.cpp
+++ b/Engine/Source/Editor/UnrealEd/Private/HierarchicalLOD.cpp
@@ -188,10 +188,10 @@ void FHierarchicalLODBuilder::BuildClusters(class ULevel* InLevel)
 			AWorldSettings* WorldSetting = InLevel->GetWorld()->GetWorldSettings();
 			// we use meter for bound. Otherwise it's very easy to get to overflow and have problem with filling ratio because
 			// bound is too huge
-			float DesiredBoundBoxHalfSize = WorldSetting->HierarchicalLODSetup[LODId].DesiredBoundSize * CM_TO_METER * 0.5f;
+			float DesiredBoundRadius = WorldSetting->HierarchicalLODSetup[LODId].DesiredBoundRadius * CM_TO_METER;
 			float DesiredFillingRatio = WorldSetting->HierarchicalLODSetup[LODId].DesiredFillingPercentage * 0.01f;
 			ensure(DesiredFillingRatio!=0.f);
-			float HighestCost = FMath::Pow(DesiredBoundBoxHalfSize, 3)/(DesiredFillingRatio);
+			float HighestCost = FMath::Pow(DesiredBoundRadius, 3)/(DesiredFillingRatio);
 			int32 MinNumActors = WorldSetting->HierarchicalLODSetup[LODId].MinNumberOfActorsToBuild;
 			check (MinNumActors > 0);
 			// test parameter I was playing with to cull adding to the array
@@ -227,6 +227,7 @@ void FHierarchicalLODBuilder::BuildClusters(class ULevel* InLevel)
 
 	// Clear Clusters. It is using stack mem, so it won't be good after this
 	Clusters.Empty();
+	Clusters.Shrink();
 }
 
 void FHierarchicalLODBuilder::FindMST() 
@@ -625,7 +626,9 @@ void FLODCluster::BuildActor(class ULevel* InLevel, const int32 LODIdx)
 		////////////////////////////////////////////////////////////////////////////////////
 		// create asset using Actors
 		const FHierarchicalSimplification& LODSetup = InLevel->GetWorld()->GetWorldSettings()->HierarchicalLODSetup[LODIdx];
-		const FMeshProxySettings& ProxySetting = LODSetup.Setting;
+		const FMeshProxySettings& ProxySetting = LODSetup.ProxySetting;
+		const FMeshMergingSettings& MergeSetting = LODSetup.MergeSetting;
+
 		// Where generated assets will be stored
 		UPackage* AssetsOuter = InLevel->GetOutermost(); // this asset is going to save with map, this means, I'll have to delete with it
 		if (AssetsOuter)
@@ -661,8 +664,6 @@ void FLODCluster::BuildActor(class ULevel* InLevel, const int32 LODIdx)
 				}
 				else
 				{
-					FMeshMergingSettings MergeSetting;
-					MergeSetting.bMergeMaterials = true;
 					MeshUtilities.MergeActors(Actors, MergeSetting, AssetsOuter, PackageName, LODIdx+1, OutAssets, OutProxyLocation, true );
 				}
 
diff --git a/Engine/Source/Runtime/Engine/Classes/Engine/EngineTypes.h b/Engine/Source/Runtime/Engine/Classes/Engine/EngineTypes.h
--- a/Engine/Source/Runtime/Engine/Classes/Engine/EngineTypes.h
+++ b/Engine/Source/Runtime/Engine/Classes/Engine/EngineTypes.h
@@ -2195,6 +2195,70 @@ struct FMeshProxySettings
 	}
 };
 
+/**
+ * Mesh merging settings
+ */
+USTRUCT()
+struct FMeshMergingSettings
+{
+	GENERATED_USTRUCT_BODY()
+
+	/** Whether to generate lightmap UVs for a merged mesh*/
+	UPROPERTY(EditAnywhere, Category=FMeshMergingSettings)
+	bool bGenerateLightMapUV;
+	
+	/** Target UV channel in a merged mesh for a lightmap */
+	UPROPERTY(EditAnywhere, Category=FMeshMergingSettings)
+	int32 TargetLightMapUVChannel;
+
+	/** Target lightmap resolution */
+UPROPERTY(EditAnywhere, Category=FMeshMergingSettings)
+	int32 TargetLightMapResolution;
+		
+	/** Whether we should import vertex colors into merged mesh */
+	UPROPERTY(EditAnywhere, Category=FMeshMergingSettings)
+	bool bImportVertexColors;
+	
+	/** Whether merged mesh should have pivot at world origin, or at first merged component otherwise */
+	UPROPERTY(EditAnywhere, Category=FMeshMergingSettings)
+	bool bPivotPointAtZero;
+
+	/** Whether to merge source materials into one flat material */
+	UPROPERTY(EditAnywhere, Category=MeshMerge)
+	bool bMergeMaterials;
+	/** Whether to export normal maps for material merging */
+	UPROPERTY(EditAnywhere, Category=MeshMerge)
+	bool bExportNormalMap;
+	/** Whether to export metallic maps for material merging */
+	UPROPERTY(EditAnywhere, Category=MeshMerge)
+	bool bExportMetallicMap;
+	/** Whether to export roughness maps for material merging */
+	UPROPERTY(EditAnywhere, Category=MeshMerge)
+	bool bExportRoughnessMap;
+	/** Whether to export specular maps for material merging */
+	UPROPERTY(EditAnywhere, Category=MeshMerge)
+	bool bExportSpecularMap;
+	/** Merged material texture atlas resolution */
+	UPROPERTY(EditAnywhere, Category=MeshMerge)
+	int32 MergedMaterialAtlasResolution;
+		
+	/** Default settings. */
+	FMeshMergingSettings()
+		: bGenerateLightMapUV(false)
+		, TargetLightMapUVChannel(1)
+		, TargetLightMapResolution(256)
+		, bImportVertexColors(false)
+		, bPivotPointAtZero(false)
+		, bMergeMaterials(false)
+		, bExportNormalMap(true)
+		, bExportMetallicMap(false)
+		, bExportRoughnessMap(false)
+		, bExportSpecularMap(false)
+		, MergedMaterialAtlasResolution(1024)
+	{
+	}
+};
+
 
 USTRUCT()
 struct ENGINE_API FDamageEvent
diff --git a/Engine/Source/Runtime/Engine/Classes/GameFramework/WorldSettings.h b/Engine/Source/Runtime/Engine/Classes/GameFramework/WorldSettings.h
--- a/Engine/Source/Runtime/Engine/Classes/GameFramework/WorldSettings.h
+++ b/Engine/Source/Runtime/Engine/Classes/GameFramework/WorldSettings.h
@@ -200,38 +200,44 @@ struct ENGINE_API FHierarchicalSimplification
 	/** If this is true, it will simplify mesh but it is slower. 
 	 * If false, it will just merge actors but not simplify using the lower LOD if exists. 
 	 * For example if you build LOD 1, it will use LOD 1 of the mesh to merge actors if exists.  
-	 * Also it will flattening material so that it can reduce drawcalls. 
+	 * If you merge material, it will reduce drawcalls. 
 	 */
 	UPROPERTY(Category=FHierarchicalSimplification, EditAnywhere)
 	bool bSimplifyMesh;
 
-	/** What are the draw distance for this LOD */
+	/** Draw Distance for this LOD actor to display. Once one of parent displays, it won't draw any of children. */
 	UPROPERTY(Category=FHierarchicalSimplification, EditAnywhere, AdvancedDisplay, meta=(UIMin=10.f, ClampMin=10.f))
 	float DrawDistance;
 
-	/** Desired Bounding Box Size for clustering - this is not guaranteed but used to calculate filling factor */
+	/** Desired Bounding Radius for clustering - this is not guaranteed but used to calculate filling factor for auto clustering */
 	UPROPERTY(EditAnywhere, Category=FHierarchicalSimplification, AdvancedDisplay, meta=(UIMin=10.f, ClampMin=10.f))
-	float DesiredBoundSize;
+	float DesiredBoundRadius;
 
-	/** Desired Filling Percentage for clustering - this is not guaranteed but used to calculate filling factor  */
+	/** Desired Filling Percentage for clustering - this is not guaranteed but used to calculate filling factor  for auto clustering */
 	UPROPERTY(EditAnywhere, Category=FHierarchicalSimplification, AdvancedDisplay, meta=(ClampMin = "0", ClampMax = "100", UIMin = "0", UIMax = "100"))
 	float DesiredFillingPercentage;
 
+	/** Min number of actors to build LODActor */
 	UPROPERTY(EditAnywhere, Category=FHierarchicalSimplification, AdvancedDisplay, meta=(ClampMin = "1", UIMin = "1"))
 	int32 MinNumberOfActorsToBuild;
 
 	/** Simplification Setting if bSimplifyMesh is true */
-	UPROPERTY(Category=FHierarchicalSimplification, EditAnywhere, AdvancedDisplay, meta=(editcondition = "bSimplifyMesh"))
-	FMeshProxySettings Setting;
+	UPROPERTY(Category=FHierarchicalSimplification, EditAnywhere, AdvancedDisplay)
+	FMeshProxySettings ProxySetting;
+
+	/** Merge Mesh Setting if bSimplifyMesh is false */
+	UPROPERTY(Category=FHierarchicalSimplification, EditAnywhere, AdvancedDisplay)
+	FMeshMergingSettings MergeSetting;
 
 	FHierarchicalSimplification()
 		: bSimplifyMesh(true)
-		, DrawDistance(2000)
-		, DesiredBoundSize(5000)
+		, DrawDistance(3000)
+		, DesiredBoundRadius(3000)
 		, DesiredFillingPercentage(50)
 		, MinNumberOfActorsToBuild(2)
 	{
-		Setting.ScreenSize = 200;
+		ProxySetting.ScreenSize = 200;
+		MergeSetting.bMergeMaterials = true;
 	}
 };
 
