UI improvements for constraint instance.
Break velocity linear drive into per component activation.

diff --git a/Engine/Source/Editor/DetailCustomizations/Private/PhysicsConstraintComponentDetails.cpp b/Engine/Source/Editor/DetailCustomizations/Private/PhysicsConstraintComponentDetails.cpp
--- a/Engine/Source/Editor/DetailCustomizations/Private/PhysicsConstraintComponentDetails.cpp
+++ b/Engine/Source/Editor/DetailCustomizations/Private/PhysicsConstraintComponentDetails.cpp
@@ -13,383 +13,497 @@ TSharedRef<IDetailCustomization> FPhysicsConstraintComponentDetails::MakeInstanc
 	return MakeShareable(new FPhysicsConstraintComponentDetails());
 }
 
-void FPhysicsConstraintComponentDetails::CustomizeDetails( IDetailLayoutBuilder& DetailBuilder )
+void FPhysicsConstraintComponentDetails::AddConstraintPresets(IDetailLayoutBuilder& DetailBuilder, APhysicsConstraintActor* OwningConstraintActor)
 {
-	TArray<TWeakObjectPtr<UObject>> Objects;
-	DetailBuilder.GetObjectsBeingCustomized(Objects);
-
-	TSharedPtr<IPropertyHandle> ConstraintInstance;
-	UPhysicsConstraintComponent* ConstraintComp = NULL;
-	APhysicsConstraintActor* OwningConstraintActor = NULL;
-
-	bool bInPhAT = false;
-
-	for (int32 i=0; i < Objects.Num(); ++i)
-	{
-		if (!Objects[i].IsValid()) { continue; }
-
-		if (Objects[i]->IsA(UPhysicsConstraintTemplate::StaticClass()))
-		{
-			ConstraintInstance = DetailBuilder.GetProperty(GET_MEMBER_NAME_CHECKED(UPhysicsConstraintTemplate, DefaultInstance));
-			bInPhAT = true;
-			break;
-		}
-		else if (Objects[i]->IsA(UPhysicsConstraintComponent::StaticClass()))
-		{
-			ConstraintInstance = DetailBuilder.GetProperty(GET_MEMBER_NAME_CHECKED(UPhysicsConstraintComponent, ConstraintInstance));
-			ConstraintComp = (UPhysicsConstraintComponent*)Objects[i].Get();
-			OwningConstraintActor = Cast<APhysicsConstraintActor>(ConstraintComp->GetOwner());
-			break;
-		}
-	}
-
 	if (OwningConstraintActor != NULL)
 	{
-		ConstraintComponent = TWeakObjectPtr<UObject>(ConstraintComp);
-
-		DetailBuilder.EditCategory( "Joint Presets" )
-			.AddCustomRow( FText::GetEmpty() )
+		DetailBuilder.EditCategory("Joint Presets")
+			.AddCustomRow(FText::GetEmpty())
 			[
 				SNew(SVerticalBox)
 				+ SVerticalBox::Slot()
-				.Padding( 0, 2.0f, 0, 0 )
+				.Padding(0, 2.0f, 0, 0)
 				.FillHeight(1.0f)
-				.VAlign( VAlign_Center )
+				.VAlign(VAlign_Center)
 				[
 					SNew(SHorizontalBox)
-					+SHorizontalBox::Slot()
+					+ SHorizontalBox::Slot()
 					.AutoWidth()
-					.Padding( 2.0f, 0.0f )
+					.Padding(2.0f, 0.0f)
 					.VAlign(VAlign_Center)
 					.HAlign(HAlign_Left)
 					[
 						SNew(SButton)
 						.VAlign(VAlign_Center)
-						.OnClicked( this, &FPhysicsConstraintComponentDetails::OnHingeClicked )
-						.Text( LOCTEXT("HingePreset", "Hinge") )
-						.ToolTipText( LOCTEXT("HingePresetTooltip", "Setup joint as hinge. A hinge joint allows motion only in one plane.") )
+						.OnClicked(this, &FPhysicsConstraintComponentDetails::OnHingeClicked)
+						.Text(LOCTEXT("HingePreset", "Hinge"))
+						.ToolTipText(LOCTEXT("HingePresetTooltip", "Setup joint as hinge. A hinge joint allows motion only in one plane."))
 					]
-					+SHorizontalBox::Slot()
+					+ SHorizontalBox::Slot()
 						.AutoWidth()
-						.Padding( 2.0f, 0.0f )
+						.Padding(2.0f, 0.0f)
 						.VAlign(VAlign_Center)
 						.HAlign(HAlign_Left)
 						[
 							SNew(SButton)
 							.VAlign(VAlign_Center)
-							.OnClicked( this, &FPhysicsConstraintComponentDetails::OnPrismaticClicked )
-							.Text( LOCTEXT("PrismaticPreset", "Prismatic") )
-							.ToolTipText( LOCTEXT("PrismaticPresetTooltip", "Setup joint as prismatic. A prismatic joint allows only linear sliding movement.") )
+							.OnClicked(this, &FPhysicsConstraintComponentDetails::OnPrismaticClicked)
+							.Text(LOCTEXT("PrismaticPreset", "Prismatic"))
+							.ToolTipText(LOCTEXT("PrismaticPresetTooltip", "Setup joint as prismatic. A prismatic joint allows only linear sliding movement."))
 						]
-					+SHorizontalBox::Slot()
+					+ SHorizontalBox::Slot()
 						.AutoWidth()
-						.Padding( 2.0f, 0.0f )
+						.Padding(2.0f, 0.0f)
 						.VAlign(VAlign_Center)
 						.HAlign(HAlign_Left)
 						[
 							SNew(SButton)
 							.VAlign(VAlign_Center)
-							.OnClicked( this, &FPhysicsConstraintComponentDetails::OnBallSocketClicked )
-							.Text( LOCTEXT("BSPreset", "Ball and Socket") )
-							.ToolTipText( LOCTEXT("BSPresetTooltip", "Setup joint as ball and socket. A Ball and Socket joint allows motion around an indefinite number of axes, which have one common center") )
+							.OnClicked(this, &FPhysicsConstraintComponentDetails::OnBallSocketClicked)
+							.Text(LOCTEXT("BSPreset", "Ball and Socket"))
+							.ToolTipText(LOCTEXT("BSPresetTooltip", "Setup joint as ball and socket. A Ball and Socket joint allows motion around an indefinite number of axes, which have one common center"))
 						]
 				]
 			];
 	}
+}
 
-	// Linear Limits
-	{
-		IDetailCategoryBuilder& LinearLimitCat = DetailBuilder.EditCategory("Linear Limits");
+void FPhysicsConstraintComponentDetails::AddLinearLimits(IDetailLayoutBuilder& DetailBuilder, TSharedPtr<IPropertyHandle> ConstraintInstance)
+{
+	IDetailCategoryBuilder& LinearLimitCat = DetailBuilder.EditCategory("Linear Limits");
 
-		LinearXMotionProperty = ConstraintInstance->GetChildHandle("LinearXMotion");
-		LinearYMotionProperty = ConstraintInstance->GetChildHandle("LinearYMotion");
-		LinearZMotionProperty = ConstraintInstance->GetChildHandle("LinearZMotion");
-		
-		TArray<TSharedPtr<FString>> LinearLimitOptionNames;
-		TArray<FText> LinearLimitOptionTooltips;
-		TArray<bool> LinearLimitOptionRestrictItems;
+	LinearXMotionProperty = ConstraintInstance->GetChildHandle("LinearXMotion");
+	LinearYMotionProperty = ConstraintInstance->GetChildHandle("LinearYMotion");
+	LinearZMotionProperty = ConstraintInstance->GetChildHandle("LinearZMotion");
 
-		const int32 ExpectedLinearLimitOptionCount = 3;
-		LinearXMotionProperty->GeneratePossibleValues(LinearLimitOptionNames, LinearLimitOptionTooltips, LinearLimitOptionRestrictItems);
-		checkf(LinearLimitOptionNames.Num() == ExpectedLinearLimitOptionCount && 
-			LinearLimitOptionTooltips.Num() == ExpectedLinearLimitOptionCount && 
-			LinearLimitOptionRestrictItems.Num() == ExpectedLinearLimitOptionCount,
-			TEXT("It seems the number of enum entries in ELinearConstraintMotion has changed. This must be handled here as well. "));
+	TArray<TSharedPtr<FString>> LinearLimitOptionNames;
+	TArray<FText> LinearLimitOptionTooltips;
+	TArray<bool> LinearLimitOptionRestrictItems;
 
-		
-		uint8 LinearLimitEnum[LCM_MAX] = { LCM_Free, LCM_Limited, LCM_Locked };
-		TSharedPtr<IPropertyHandle> LinearLimitProperties[] = { LinearXMotionProperty, LinearYMotionProperty, LinearZMotionProperty };
+	const int32 ExpectedLinearLimitOptionCount = 3;
+	LinearXMotionProperty->GeneratePossibleValues(LinearLimitOptionNames, LinearLimitOptionTooltips, LinearLimitOptionRestrictItems);
+	checkf(LinearLimitOptionNames.Num() == ExpectedLinearLimitOptionCount &&
+		LinearLimitOptionTooltips.Num() == ExpectedLinearLimitOptionCount &&
+		LinearLimitOptionRestrictItems.Num() == ExpectedLinearLimitOptionCount,
+		TEXT("It seems the number of enum entries in ELinearConstraintMotion has changed. This must be handled here as well. "));
 
 
-		for (int32 PropertyIdx=0; PropertyIdx < 3; ++PropertyIdx)
-		{
-			TSharedPtr<IPropertyHandle> CurProperty = LinearLimitProperties[PropertyIdx];
+	uint8 LinearLimitEnum[LCM_MAX] = { LCM_Free, LCM_Limited, LCM_Locked };
+	TSharedPtr<IPropertyHandle> LinearLimitProperties[] = { LinearXMotionProperty, LinearYMotionProperty, LinearZMotionProperty };
+
 
-			LinearLimitCat.AddProperty(CurProperty).CustomWidget()
-				.NameContent()
+	for (int32 PropertyIdx = 0; PropertyIdx < 3; ++PropertyIdx)
+	{
+		TSharedPtr<IPropertyHandle> CurProperty = LinearLimitProperties[PropertyIdx];
+
+		LinearLimitCat.AddProperty(CurProperty).CustomWidget()
+			.NameContent()
+			[
+				SNew(STextBlock)
+				.Font(DetailBuilder.GetDetailFont())
+				.Text(CurProperty->GetPropertyDisplayName())
+				.ToolTipText(CurProperty->GetToolTipText())
+			]
+		.ValueContent()
+			.MinDesiredWidth(125.0f * 3.0f)
+			.MaxDesiredWidth(125.0f * 3.0f)
+			[
+				SNew(SHorizontalBox)
+				+ SHorizontalBox::Slot()
+				.AutoWidth()
+				.HAlign(HAlign_Left)
 				[
-					SNew(STextBlock)
-					.Font(DetailBuilder.GetDetailFont())
-					.Text(CurProperty->GetPropertyDisplayName())
-					.ToolTipText(CurProperty->GetToolTipText())
+					SNew(SCheckBox)
+					.Style(FEditorStyle::Get(), "RadioButton")
+					.IsChecked(this, &FPhysicsConstraintComponentDetails::IsLimitRadioChecked, CurProperty, LinearLimitEnum[0])
+					.OnCheckStateChanged(this, &FPhysicsConstraintComponentDetails::OnLimitRadioChanged, CurProperty, LinearLimitEnum[0])
+					.ToolTipText(LinearLimitOptionTooltips[0])
+					[
+						SNew(STextBlock)
+						.Text(FText::FromString(*LinearLimitOptionNames[0].Get()))
+						.Font(IDetailLayoutBuilder::GetDetailFont())
+					]
 				]
-			.ValueContent()
-				[
-					SNew(SVerticalBox)
-					+SVerticalBox::Slot()
-					.AutoHeight()
+				+ SHorizontalBox::Slot()
+					.AutoWidth()
 					.HAlign(HAlign_Left)
+					.Padding(5, 0, 0, 0)
 					[
 						SNew(SCheckBox)
 						.Style(FEditorStyle::Get(), "RadioButton")
-						.IsChecked( this, &FPhysicsConstraintComponentDetails::IsLimitRadioChecked, CurProperty, LinearLimitEnum[0])
-						.OnCheckStateChanged( this, &FPhysicsConstraintComponentDetails::OnLimitRadioChanged, CurProperty, LinearLimitEnum[0] )
-						.ToolTipText(LinearLimitOptionTooltips[0])
+						.IsChecked(this, &FPhysicsConstraintComponentDetails::IsLimitRadioChecked, CurProperty, LinearLimitEnum[1])
+						.OnCheckStateChanged(this, &FPhysicsConstraintComponentDetails::OnLimitRadioChanged, CurProperty, LinearLimitEnum[1])
+						.ToolTipText(LinearLimitOptionTooltips[1])
 						[
-							SNew(STextBlock).Text(FText::FromString(*LinearLimitOptionNames[0].Get()))
-						]						
-					]
-					+SVerticalBox::Slot()
-						.AutoHeight()
-						.HAlign(HAlign_Left)
-						[
-							SNew(SCheckBox)
-							.Style(FEditorStyle::Get(), "RadioButton")
-							.IsChecked( this, &FPhysicsConstraintComponentDetails::IsLimitRadioChecked, CurProperty, LinearLimitEnum[1])
-							.OnCheckStateChanged( this, &FPhysicsConstraintComponentDetails::OnLimitRadioChanged, CurProperty, LinearLimitEnum[1] )
-							.ToolTipText(LinearLimitOptionTooltips[1])
-							[
-								SNew(STextBlock).Text(FText::FromString(*LinearLimitOptionNames[1].Get()))
-							]
+							SNew(STextBlock)
+							.Text(FText::FromString(*LinearLimitOptionNames[1].Get()))
+							.Font(IDetailLayoutBuilder::GetDetailFont())
 						]
-					+SVerticalBox::Slot()
-						.AutoHeight()
-						.HAlign(HAlign_Left)
+					]
+				+ SHorizontalBox::Slot()
+					.AutoWidth()
+					.HAlign(HAlign_Left)
+					.Padding(5, 0, 0, 0)
+					[
+						SNew(SCheckBox)
+						.Style(FEditorStyle::Get(), "RadioButton")
+						.IsChecked(this, &FPhysicsConstraintComponentDetails::IsLimitRadioChecked, CurProperty, LinearLimitEnum[2])
+						.OnCheckStateChanged(this, &FPhysicsConstraintComponentDetails::OnLimitRadioChanged, CurProperty, LinearLimitEnum[2])
+						.ToolTipText(LinearLimitOptionTooltips[2])
 						[
-							SNew(SCheckBox)
-							.Style(FEditorStyle::Get(), "RadioButton")
-							.IsChecked( this, &FPhysicsConstraintComponentDetails::IsLimitRadioChecked, CurProperty, LinearLimitEnum[2])
-							.OnCheckStateChanged( this, &FPhysicsConstraintComponentDetails::OnLimitRadioChanged, CurProperty, LinearLimitEnum[2] )
-							.ToolTipText(LinearLimitOptionTooltips[2])
-							[
-								SNew(STextBlock).Text(FText::FromString(*LinearLimitOptionNames[2].Get()))
-							]
+							SNew(STextBlock)
+							.Text(FText::FromString(*LinearLimitOptionNames[2].Get()))
+							.Font(IDetailLayoutBuilder::GetDetailFont())
 						]
-				];
-		}
-		
-
-		LinearLimitCat.AddProperty(ConstraintInstance->GetChildHandle(GET_MEMBER_NAME_CHECKED(FConstraintInstance, LinearLimitSize)).ToSharedRef())
-			.Visibility(TAttribute<EVisibility>::Create(TAttribute<EVisibility>::FGetter::CreateSP(this, &FPhysicsConstraintComponentDetails::IsPropertyVisible, EPropertyType::LinearLimit)));
-		LinearLimitCat.AddProperty(ConstraintInstance->GetChildHandle(GET_MEMBER_NAME_CHECKED(FConstraintInstance, bLinearLimitSoft)).ToSharedRef())
-			.Visibility(TAttribute<EVisibility>::Create(TAttribute<EVisibility>::FGetter::CreateSP(this, &FPhysicsConstraintComponentDetails::IsPropertyVisible, EPropertyType::LinearLimit)));
-		LinearLimitCat.AddProperty(ConstraintInstance->GetChildHandle(GET_MEMBER_NAME_CHECKED(FConstraintInstance, LinearLimitStiffness)).ToSharedRef())
-			.Visibility(TAttribute<EVisibility>::Create(TAttribute<EVisibility>::FGetter::CreateSP(this, &FPhysicsConstraintComponentDetails::IsPropertyVisible, EPropertyType::LinearLimit)));
-		LinearLimitCat.AddProperty(ConstraintInstance->GetChildHandle(GET_MEMBER_NAME_CHECKED(FConstraintInstance, LinearLimitDamping)).ToSharedRef())
-			.Visibility(TAttribute<EVisibility>::Create(TAttribute<EVisibility>::FGetter::CreateSP(this, &FPhysicsConstraintComponentDetails::IsPropertyVisible, EPropertyType::LinearLimit)));
-		LinearLimitCat.AddProperty(ConstraintInstance->GetChildHandle(GET_MEMBER_NAME_CHECKED(FConstraintInstance, bLinearBreakable)).ToSharedRef());
-		LinearLimitCat.AddProperty(ConstraintInstance->GetChildHandle(GET_MEMBER_NAME_CHECKED(FConstraintInstance, LinearBreakThreshold)).ToSharedRef());
+					]
+			];
 	}
 
-	// Angular Limits
-	{
-		IDetailCategoryBuilder& AngularLimitCat = DetailBuilder.EditCategory("Angular Limits");
 
-		AngularSwing1MotionProperty = ConstraintInstance->GetChildHandle(GET_MEMBER_NAME_CHECKED(FConstraintInstance, AngularSwing1Motion));
-		AngularSwing2MotionProperty = ConstraintInstance->GetChildHandle(GET_MEMBER_NAME_CHECKED(FConstraintInstance, AngularSwing2Motion));
-		AngularTwistMotionProperty = ConstraintInstance->GetChildHandle(GET_MEMBER_NAME_CHECKED(FConstraintInstance, AngularTwistMotion));
+	LinearLimitCat.AddProperty(ConstraintInstance->GetChildHandle(GET_MEMBER_NAME_CHECKED(FConstraintInstance, LinearLimitSize)).ToSharedRef()).IsEnabled(TAttribute<bool>::Create(TAttribute<bool>::FGetter::CreateSP(this, &FPhysicsConstraintComponentDetails::IsPropertyEnabled, EPropertyType::LinearLimit)));
+	LinearLimitCat.AddProperty(ConstraintInstance->GetChildHandle(GET_MEMBER_NAME_CHECKED(FConstraintInstance, bLinearLimitSoft)).ToSharedRef()).IsEnabled(TAttribute<bool>::Create(TAttribute<bool>::FGetter::CreateSP(this, &FPhysicsConstraintComponentDetails::IsPropertyEnabled, EPropertyType::LinearLimit)));
+	LinearLimitCat.AddProperty(ConstraintInstance->GetChildHandle(GET_MEMBER_NAME_CHECKED(FConstraintInstance, LinearLimitStiffness)).ToSharedRef()).IsEnabled(TAttribute<bool>::Create(TAttribute<bool>::FGetter::CreateSP(this, &FPhysicsConstraintComponentDetails::IsPropertyEnabled, EPropertyType::LinearLimit)));
+	LinearLimitCat.AddProperty(ConstraintInstance->GetChildHandle(GET_MEMBER_NAME_CHECKED(FConstraintInstance, LinearLimitDamping)).ToSharedRef()).IsEnabled(TAttribute<bool>::Create(TAttribute<bool>::FGetter::CreateSP(this, &FPhysicsConstraintComponentDetails::IsPropertyEnabled, EPropertyType::LinearLimit)));
+	LinearLimitCat.AddProperty(ConstraintInstance->GetChildHandle(GET_MEMBER_NAME_CHECKED(FConstraintInstance, bLinearBreakable)).ToSharedRef());
+	LinearLimitCat.AddProperty(ConstraintInstance->GetChildHandle(GET_MEMBER_NAME_CHECKED(FConstraintInstance, LinearBreakThreshold)).ToSharedRef());
+}
+
+void FPhysicsConstraintComponentDetails::AddAngularLimits(IDetailLayoutBuilder& DetailBuilder, TSharedPtr<IPropertyHandle> ConstraintInstance)
+{
+	IDetailCategoryBuilder& AngularLimitCat = DetailBuilder.EditCategory("Angular Limits");
 
-		TArray<TSharedPtr<FString>> AngularLimitOptionNames;
-		TArray<FText> AngularLimitOptionTooltips;
-		TArray<bool> AngularLimitOptionRestrictItems;
+	AngularSwing1MotionProperty = ConstraintInstance->GetChildHandle(GET_MEMBER_NAME_CHECKED(FConstraintInstance, AngularSwing1Motion));
+	AngularSwing2MotionProperty = ConstraintInstance->GetChildHandle(GET_MEMBER_NAME_CHECKED(FConstraintInstance, AngularSwing2Motion));
+	AngularTwistMotionProperty = ConstraintInstance->GetChildHandle(GET_MEMBER_NAME_CHECKED(FConstraintInstance, AngularTwistMotion));
 
-		const int32 ExpectedAngularLimitOptionCount = 3;
-		AngularSwing1MotionProperty->GeneratePossibleValues(AngularLimitOptionNames, AngularLimitOptionTooltips, AngularLimitOptionRestrictItems);
-		checkf(AngularLimitOptionNames.Num() == ExpectedAngularLimitOptionCount && 
-			AngularLimitOptionTooltips.Num() == ExpectedAngularLimitOptionCount && 
-			AngularLimitOptionRestrictItems.Num() == ExpectedAngularLimitOptionCount,
-			TEXT("It seems the number of enum entries in EAngularConstraintMotion has changed. This must be handled here as well. "));
+	TArray<TSharedPtr<FString>> AngularLimitOptionNames;
+	TArray<FText> AngularLimitOptionTooltips;
+	TArray<bool> AngularLimitOptionRestrictItems;
 
+	const int32 ExpectedAngularLimitOptionCount = 3;
+	AngularSwing1MotionProperty->GeneratePossibleValues(AngularLimitOptionNames, AngularLimitOptionTooltips, AngularLimitOptionRestrictItems);
+	checkf(AngularLimitOptionNames.Num() == ExpectedAngularLimitOptionCount &&
+		AngularLimitOptionTooltips.Num() == ExpectedAngularLimitOptionCount &&
+		AngularLimitOptionRestrictItems.Num() == ExpectedAngularLimitOptionCount,
+		TEXT("It seems the number of enum entries in EAngularConstraintMotion has changed. This must be handled here as well. "));
 
-		uint8 AngularLimitEnum[LCM_MAX] = { ACM_Free, LCM_Limited, LCM_Locked };
-		TSharedPtr<IPropertyHandle> AngularLimitProperties[] = { AngularSwing1MotionProperty, AngularTwistMotionProperty, AngularSwing2MotionProperty};
 
-		for (int32 PropertyIdx=0; PropertyIdx < 3; ++PropertyIdx)
-		{
-			TSharedPtr<IPropertyHandle> CurProperty = AngularLimitProperties[PropertyIdx];
+	uint8 AngularLimitEnum[LCM_MAX] = { ACM_Free, LCM_Limited, LCM_Locked };
+	TSharedPtr<IPropertyHandle> AngularLimitProperties[] = { AngularSwing1MotionProperty, AngularTwistMotionProperty, AngularSwing2MotionProperty };
+
+	for (int32 PropertyIdx = 0; PropertyIdx < 3; ++PropertyIdx)
+	{
+		TSharedPtr<IPropertyHandle> CurProperty = AngularLimitProperties[PropertyIdx];
 
-			AngularLimitCat.AddProperty(CurProperty).CustomWidget()
-				.NameContent()
+		AngularLimitCat.AddProperty(CurProperty).CustomWidget()
+			.NameContent()
+			[
+				SNew(STextBlock)
+				.Font(DetailBuilder.GetDetailFont())
+				.Text(CurProperty->GetPropertyDisplayName())
+				.ToolTipText(CurProperty->GetToolTipText())
+			]
+		.ValueContent()
+			.MinDesiredWidth(125.0f * 3.0f)
+			.MaxDesiredWidth(125.0f * 3.0f)
+			[
+				SNew(SHorizontalBox)
+				+ SHorizontalBox::Slot()
+				.AutoWidth()
+				.HAlign(HAlign_Left)
 				[
-					SNew(STextBlock)
-					.Font(DetailBuilder.GetDetailFont())
-					.Text(CurProperty->GetPropertyDisplayName())
-					.ToolTipText(CurProperty->GetToolTipText())
+					SNew(SCheckBox)
+					.Style(FEditorStyle::Get(), "RadioButton")
+					.IsChecked(this, &FPhysicsConstraintComponentDetails::IsLimitRadioChecked, CurProperty, AngularLimitEnum[0])
+					.OnCheckStateChanged(this, &FPhysicsConstraintComponentDetails::OnLimitRadioChanged, CurProperty, AngularLimitEnum[0])
+					.ToolTipText(AngularLimitOptionTooltips[0])
+					[
+						SNew(STextBlock)
+						.Text(FText::FromString(*AngularLimitOptionNames[0].Get()))
+						.Font(IDetailLayoutBuilder::GetDetailFont())
+					]
 				]
-			.ValueContent()
-				[
-					SNew(SVerticalBox)
-					+SVerticalBox::Slot()
-					.AutoHeight()
+				+ SHorizontalBox::Slot()
+					.AutoWidth()
 					.HAlign(HAlign_Left)
+					.Padding(5, 0, 0, 0)
 					[
 						SNew(SCheckBox)
 						.Style(FEditorStyle::Get(), "RadioButton")
-						.IsChecked( this, &FPhysicsConstraintComponentDetails::IsLimitRadioChecked, CurProperty, AngularLimitEnum[0])
-						.OnCheckStateChanged( this, &FPhysicsConstraintComponentDetails::OnLimitRadioChanged, CurProperty, AngularLimitEnum[0] )
-						.ToolTipText(AngularLimitOptionTooltips[0])
-						[
-							SNew(STextBlock).Text(FText::FromString(*AngularLimitOptionNames[0].Get()))
-						]						
-					]
-					+SVerticalBox::Slot()
-						.AutoHeight()
-						.HAlign(HAlign_Left)
+						.IsChecked(this, &FPhysicsConstraintComponentDetails::IsLimitRadioChecked, CurProperty, AngularLimitEnum[1])
+						.OnCheckStateChanged(this, &FPhysicsConstraintComponentDetails::OnLimitRadioChanged, CurProperty, AngularLimitEnum[1])
+						.ToolTipText(AngularLimitOptionTooltips[1])
 						[
-							SNew(SCheckBox)
-							.Style(FEditorStyle::Get(), "RadioButton")
-							.IsChecked( this, &FPhysicsConstraintComponentDetails::IsLimitRadioChecked, CurProperty, AngularLimitEnum[1])
-							.OnCheckStateChanged( this, &FPhysicsConstraintComponentDetails::OnLimitRadioChanged, CurProperty, AngularLimitEnum[1] )
-							.ToolTipText(AngularLimitOptionTooltips[1])
-							[
-								SNew(STextBlock).Text(FText::FromString(*AngularLimitOptionNames[1].Get()))
-							]
+							SNew(STextBlock)
+							.Text(FText::FromString(*AngularLimitOptionNames[1].Get()))
+							.Font(IDetailLayoutBuilder::GetDetailFont())
 						]
-					+SVerticalBox::Slot()
-						.AutoHeight()
-						.HAlign(HAlign_Left)
+					]
+				+ SHorizontalBox::Slot()
+					.AutoWidth()
+					.HAlign(HAlign_Left)
+					.Padding(5, 0, 0, 0)
+					[
+						SNew(SCheckBox)
+						.Style(FEditorStyle::Get(), "RadioButton")
+						.IsChecked(this, &FPhysicsConstraintComponentDetails::IsLimitRadioChecked, CurProperty, AngularLimitEnum[2])
+						.OnCheckStateChanged(this, &FPhysicsConstraintComponentDetails::OnLimitRadioChanged, CurProperty, AngularLimitEnum[2])
+						.ToolTipText(AngularLimitOptionTooltips[2])
 						[
-							SNew(SCheckBox)
-							.Style(FEditorStyle::Get(), "RadioButton")
-							.IsChecked( this, &FPhysicsConstraintComponentDetails::IsLimitRadioChecked, CurProperty, AngularLimitEnum[2])
-							.OnCheckStateChanged( this, &FPhysicsConstraintComponentDetails::OnLimitRadioChanged, CurProperty, AngularLimitEnum[2] )
-							.ToolTipText(AngularLimitOptionTooltips[2])
-							[
-								SNew(STextBlock).Text(FText::FromString(*AngularLimitOptionNames[2].Get()))
-							]
+							SNew(STextBlock)
+							.Text(FText::FromString(*AngularLimitOptionNames[2].Get()))
+							.Font(IDetailLayoutBuilder::GetDetailFont())
 						]
-				];
-		}
+					]
+			];
+	}
 
-		AngularLimitCat.AddProperty(ConstraintInstance->GetChildHandle(GET_MEMBER_NAME_CHECKED(FConstraintInstance, Swing1LimitAngle)).ToSharedRef())
-			.Visibility(TAttribute<EVisibility>::Create(TAttribute<EVisibility>::FGetter::CreateSP(this, &FPhysicsConstraintComponentDetails::IsPropertyVisible, EPropertyType::AngularSwing1Limit)));
-		AngularLimitCat.AddProperty(ConstraintInstance->GetChildHandle(GET_MEMBER_NAME_CHECKED(FConstraintInstance, TwistLimitAngle)).ToSharedRef())
-			.Visibility(TAttribute<EVisibility>::Create(TAttribute<EVisibility>::FGetter::CreateSP(this, &FPhysicsConstraintComponentDetails::IsPropertyVisible, EPropertyType::AngularTwistLimit)));
-		AngularLimitCat.AddProperty(ConstraintInstance->GetChildHandle(GET_MEMBER_NAME_CHECKED(FConstraintInstance, Swing2LimitAngle)).ToSharedRef())
-			.Visibility(TAttribute<EVisibility>::Create(TAttribute<EVisibility>::FGetter::CreateSP(this, &FPhysicsConstraintComponentDetails::IsPropertyVisible, EPropertyType::AngularSwing2Limit)));
-
-
-		AngularLimitCat.AddProperty(ConstraintInstance->GetChildHandle(GET_MEMBER_NAME_CHECKED(FConstraintInstance, bSwingLimitSoft)).ToSharedRef())
-			.Visibility(TAttribute<EVisibility>::Create(TAttribute<EVisibility>::FGetter::CreateSP(this, &FPhysicsConstraintComponentDetails::IsPropertyVisible, EPropertyType::AngularSwingLimit)));
-		AngularLimitCat.AddProperty(ConstraintInstance->GetChildHandle(GET_MEMBER_NAME_CHECKED(FConstraintInstance, SwingLimitStiffness)).ToSharedRef())
-			.Visibility(TAttribute<EVisibility>::Create(TAttribute<EVisibility>::FGetter::CreateSP(this, &FPhysicsConstraintComponentDetails::IsPropertyVisible, EPropertyType::AngularSwingLimit)));
-		AngularLimitCat.AddProperty(ConstraintInstance->GetChildHandle(GET_MEMBER_NAME_CHECKED(FConstraintInstance, SwingLimitDamping)).ToSharedRef())
-			.Visibility(TAttribute<EVisibility>::Create(TAttribute<EVisibility>::FGetter::CreateSP(this, &FPhysicsConstraintComponentDetails::IsPropertyVisible, EPropertyType::AngularSwingLimit)));
-
-		
-
-		AngularLimitCat.AddProperty(ConstraintInstance->GetChildHandle(GET_MEMBER_NAME_CHECKED(FConstraintInstance, bTwistLimitSoft)).ToSharedRef())
-			.Visibility(TAttribute<EVisibility>::Create(TAttribute<EVisibility>::FGetter::CreateSP(this, &FPhysicsConstraintComponentDetails::IsPropertyVisible, EPropertyType::AngularTwistLimit)));
-		AngularLimitCat.AddProperty(ConstraintInstance->GetChildHandle(GET_MEMBER_NAME_CHECKED(FConstraintInstance, TwistLimitStiffness)).ToSharedRef())
-			.Visibility(TAttribute<EVisibility>::Create(TAttribute<EVisibility>::FGetter::CreateSP(this, &FPhysicsConstraintComponentDetails::IsPropertyVisible, EPropertyType::AngularTwistLimit)));
-		AngularLimitCat.AddProperty(ConstraintInstance->GetChildHandle(GET_MEMBER_NAME_CHECKED(FConstraintInstance, TwistLimitDamping)).ToSharedRef())
-			.Visibility(TAttribute<EVisibility>::Create(TAttribute<EVisibility>::FGetter::CreateSP(this, &FPhysicsConstraintComponentDetails::IsPropertyVisible, EPropertyType::AngularTwistLimit)));
-		
-		if (bInPhAT == false)
-		{
-			AngularLimitCat.AddProperty(ConstraintInstance->GetChildHandle(GET_MEMBER_NAME_CHECKED(FConstraintInstance, AngularRotationOffset)).ToSharedRef())
-				.Visibility(TAttribute<EVisibility>::Create(TAttribute<EVisibility>::FGetter::CreateSP(this, &FPhysicsConstraintComponentDetails::IsPropertyVisible, EPropertyType::AngularAnyLimit)));
+	AngularLimitCat.AddProperty(ConstraintInstance->GetChildHandle(GET_MEMBER_NAME_CHECKED(FConstraintInstance, Swing1LimitAngle)).ToSharedRef()).IsEnabled(TAttribute<bool>::Create(TAttribute<bool>::FGetter::CreateSP(this, &FPhysicsConstraintComponentDetails::IsPropertyEnabled, EPropertyType::AngularSwing1Limit)));
+	AngularLimitCat.AddProperty(ConstraintInstance->GetChildHandle(GET_MEMBER_NAME_CHECKED(FConstraintInstance, TwistLimitAngle)).ToSharedRef()).IsEnabled(TAttribute<bool>::Create(TAttribute<bool>::FGetter::CreateSP(this, &FPhysicsConstraintComponentDetails::IsPropertyEnabled, EPropertyType::AngularTwistLimit)));
+	AngularLimitCat.AddProperty(ConstraintInstance->GetChildHandle(GET_MEMBER_NAME_CHECKED(FConstraintInstance, Swing2LimitAngle)).ToSharedRef()).IsEnabled(TAttribute<bool>::Create(TAttribute<bool>::FGetter::CreateSP(this, &FPhysicsConstraintComponentDetails::IsPropertyEnabled, EPropertyType::AngularSwing2Limit)));
 
-		}
-		else
-		{
-			AngularLimitCat.AddProperty(ConstraintInstance->GetChildHandle(GET_MEMBER_NAME_CHECKED(FConstraintInstance, AngularRotationOffset)).ToSharedRef())
-				.Visibility(EVisibility::Collapsed);
-		}
 
-		AngularLimitCat.AddProperty(ConstraintInstance->GetChildHandle(GET_MEMBER_NAME_CHECKED(FConstraintInstance, bAngularBreakable)).ToSharedRef());
-		AngularLimitCat.AddProperty(ConstraintInstance->GetChildHandle(GET_MEMBER_NAME_CHECKED(FConstraintInstance, AngularBreakThreshold)).ToSharedRef());
-	}
+	AngularLimitCat.AddProperty(ConstraintInstance->GetChildHandle(GET_MEMBER_NAME_CHECKED(FConstraintInstance, bSwingLimitSoft)).ToSharedRef())
+		.IsEnabled(TAttribute<bool>::Create(TAttribute<bool>::FGetter::CreateSP(this, &FPhysicsConstraintComponentDetails::IsPropertyEnabled, EPropertyType::AngularSwingLimit)));
+	AngularLimitCat.AddProperty(ConstraintInstance->GetChildHandle(GET_MEMBER_NAME_CHECKED(FConstraintInstance, SwingLimitStiffness)).ToSharedRef())
+		.IsEnabled(TAttribute<bool>::Create(TAttribute<bool>::FGetter::CreateSP(this, &FPhysicsConstraintComponentDetails::IsPropertyEnabled, EPropertyType::AngularSwingLimit)));
+	AngularLimitCat.AddProperty(ConstraintInstance->GetChildHandle(GET_MEMBER_NAME_CHECKED(FConstraintInstance, SwingLimitDamping)).ToSharedRef())
+		.IsEnabled(TAttribute<bool>::Create(TAttribute<bool>::FGetter::CreateSP(this, &FPhysicsConstraintComponentDetails::IsPropertyEnabled, EPropertyType::AngularSwingLimit)));
+
+
+
+	AngularLimitCat.AddProperty(ConstraintInstance->GetChildHandle(GET_MEMBER_NAME_CHECKED(FConstraintInstance, bTwistLimitSoft)).ToSharedRef())
+		.IsEnabled(TAttribute<bool>::Create(TAttribute<bool>::FGetter::CreateSP(this, &FPhysicsConstraintComponentDetails::IsPropertyEnabled, EPropertyType::AngularTwistLimit)));
+	AngularLimitCat.AddProperty(ConstraintInstance->GetChildHandle(GET_MEMBER_NAME_CHECKED(FConstraintInstance, TwistLimitStiffness)).ToSharedRef())
+		.IsEnabled(TAttribute<bool>::Create(TAttribute<bool>::FGetter::CreateSP(this, &FPhysicsConstraintComponentDetails::IsPropertyEnabled, EPropertyType::AngularTwistLimit)));
+	AngularLimitCat.AddProperty(ConstraintInstance->GetChildHandle(GET_MEMBER_NAME_CHECKED(FConstraintInstance, TwistLimitDamping)).ToSharedRef())
+		.IsEnabled(TAttribute<bool>::Create(TAttribute<bool>::FGetter::CreateSP(this, &FPhysicsConstraintComponentDetails::IsPropertyEnabled, EPropertyType::AngularTwistLimit)));
+
+	if (bInPhat == false)
+	{
+		AngularLimitCat.AddProperty(ConstraintInstance->GetChildHandle(GET_MEMBER_NAME_CHECKED(FConstraintInstance, AngularRotationOffset)).ToSharedRef())
+			.IsEnabled(TAttribute<bool>::Create(TAttribute<bool>::FGetter::CreateSP(this, &FPhysicsConstraintComponentDetails::IsPropertyEnabled, EPropertyType::AngularAnyLimit)));
 
-	// Linear Drive
+	}
+	else
 	{
-		IDetailCategoryBuilder& LinearMotorCat = DetailBuilder.EditCategory("LinearMotor");
-
-		LinearPositionDriveProperty = ConstraintInstance->GetChildHandle(GET_MEMBER_NAME_CHECKED(FConstraintInstance, bLinearPositionDrive));
-		LinearVelocityDriveProperty = ConstraintInstance->GetChildHandle(GET_MEMBER_NAME_CHECKED(FConstraintInstance, bLinearVelocityDrive));
-
-		IDetailGroup& PositionGroup = LinearMotorCat.AddGroup("Linear Position Drive", FText::GetEmpty());
-		PositionGroup.HeaderProperty(LinearPositionDriveProperty.ToSharedRef());
-		PositionGroup.AddPropertyRow(ConstraintInstance->GetChildHandle(GET_MEMBER_NAME_CHECKED(FConstraintInstance, bLinearXPositionDrive)).ToSharedRef())
-			.Visibility(TAttribute<EVisibility>::Create(TAttribute<EVisibility>::FGetter::CreateSP(this, &FPhysicsConstraintComponentDetails::IsPropertyVisible, EPropertyType::LinearPositionDrive)));
-		PositionGroup.AddPropertyRow(ConstraintInstance->GetChildHandle(GET_MEMBER_NAME_CHECKED(FConstraintInstance, bLinearYPositionDrive)).ToSharedRef())
-			.Visibility(TAttribute<EVisibility>::Create(TAttribute<EVisibility>::FGetter::CreateSP(this, &FPhysicsConstraintComponentDetails::IsPropertyVisible, EPropertyType::LinearPositionDrive)));
-		PositionGroup.AddPropertyRow(ConstraintInstance->GetChildHandle(GET_MEMBER_NAME_CHECKED(FConstraintInstance, bLinearZPositionDrive)).ToSharedRef())
-			.Visibility(TAttribute<EVisibility>::Create(TAttribute<EVisibility>::FGetter::CreateSP(this, &FPhysicsConstraintComponentDetails::IsPropertyVisible, EPropertyType::LinearPositionDrive)));
-		PositionGroup.AddPropertyRow(ConstraintInstance->GetChildHandle(GET_MEMBER_NAME_CHECKED(FConstraintInstance, LinearPositionTarget)).ToSharedRef())
-			.Visibility(TAttribute<EVisibility>::Create(TAttribute<EVisibility>::FGetter::CreateSP(this, &FPhysicsConstraintComponentDetails::IsPropertyVisible, EPropertyType::LinearPositionDrive)));
-
-		IDetailGroup& VelocityGroup = LinearMotorCat.AddGroup("Linear Velocity Drive", FText::GetEmpty());
-		VelocityGroup.HeaderProperty(LinearVelocityDriveProperty.ToSharedRef());
-		VelocityGroup.AddPropertyRow(ConstraintInstance->GetChildHandle(GET_MEMBER_NAME_CHECKED(FConstraintInstance, LinearVelocityTarget)).ToSharedRef())
-			.Visibility(TAttribute<EVisibility>::Create(TAttribute<EVisibility>::FGetter::CreateSP(this, &FPhysicsConstraintComponentDetails::IsPropertyVisible, EPropertyType::LinearVelocityDrive)));
-
-		LinearMotorCat.AddProperty(ConstraintInstance->GetChildHandle(GET_MEMBER_NAME_CHECKED(FConstraintInstance, LinearDriveSpring)).ToSharedRef())
-			.Visibility(TAttribute<EVisibility>::Create(TAttribute<EVisibility>::FGetter::CreateSP(this, &FPhysicsConstraintComponentDetails::IsPropertyVisible, EPropertyType::LinearAnyDrive)));
-		LinearMotorCat.AddProperty(ConstraintInstance->GetChildHandle(GET_MEMBER_NAME_CHECKED(FConstraintInstance, LinearDriveDamping)).ToSharedRef())
-			.Visibility(TAttribute<EVisibility>::Create(TAttribute<EVisibility>::FGetter::CreateSP(this, &FPhysicsConstraintComponentDetails::IsPropertyVisible, EPropertyType::LinearAnyDrive)));
-		LinearMotorCat.AddProperty(ConstraintInstance->GetChildHandle(GET_MEMBER_NAME_CHECKED(FConstraintInstance, LinearDriveForceLimit)).ToSharedRef())
-			.Visibility(TAttribute<EVisibility>::Create(TAttribute<EVisibility>::FGetter::CreateSP(this, &FPhysicsConstraintComponentDetails::IsPropertyVisible, EPropertyType::LinearAnyDrive)));
+		AngularLimitCat.AddProperty(ConstraintInstance->GetChildHandle(GET_MEMBER_NAME_CHECKED(FConstraintInstance, AngularRotationOffset)).ToSharedRef())
+			.Visibility(EVisibility::Collapsed);
 	}
 
-	// Angular Drive
+	AngularLimitCat.AddProperty(ConstraintInstance->GetChildHandle(GET_MEMBER_NAME_CHECKED(FConstraintInstance, bAngularBreakable)).ToSharedRef());
+	AngularLimitCat.AddProperty(ConstraintInstance->GetChildHandle(GET_MEMBER_NAME_CHECKED(FConstraintInstance, AngularBreakThreshold)).ToSharedRef());
+}
+
+TSharedRef<SWidget> FPhysicsConstraintComponentDetails::JoinPropertyWidgets(TSharedPtr<IPropertyHandle> TargetProperty, FName TargetChildName, TSharedPtr<IPropertyHandle> ConstraintInstance, FName CheckPropertyName, TSharedPtr<IPropertyHandle>& StoreCheckProperty, EPropertyType::Type DriveType)
+{
+	StoreCheckProperty = ConstraintInstance->GetChildHandle(CheckPropertyName);
+	StoreCheckProperty->MarkHiddenByCustomization();
+	TSharedRef<SWidget> TargetWidget = TargetProperty->GetChildHandle(TargetChildName)->CreatePropertyValueWidget();
+	TargetWidget->SetEnabled(TAttribute<bool>::Create(TAttribute<bool>::FGetter::CreateSP(this, &FPhysicsConstraintComponentDetails::IsPropertyEnabled, DriveType)));
+
+	return SNew(SHorizontalBox)
+		+ SHorizontalBox::Slot()
+		.AutoWidth()
+		.Padding(0, 0, 5, 0)
+		[
+			StoreCheckProperty->CreatePropertyValueWidget()
+		]
+	+ SHorizontalBox::Slot()
+		[
+			TargetWidget
+		];
+}
+
+
+void FPhysicsConstraintComponentDetails::AddLinearDrive(IDetailLayoutBuilder& DetailBuilder, TSharedPtr<IPropertyHandle> ConstraintInstance)
+{
+	IDetailCategoryBuilder& LinearMotorCat = DetailBuilder.EditCategory("LinearMotor");
+	LinearPositionDriveProperty = ConstraintInstance->GetChildHandle(GET_MEMBER_NAME_CHECKED(FConstraintInstance, bLinearPositionDrive));
+
+	IDetailGroup& PositionGroup = LinearMotorCat.AddGroup("Linear Position Drive", FText::GetEmpty());
+	PositionGroup.HeaderProperty(LinearPositionDriveProperty.ToSharedRef());
+
+	TSharedRef<IPropertyHandle> LinearPositionTargetProperty = ConstraintInstance->GetChildHandle(GET_MEMBER_NAME_CHECKED(FConstraintInstance, LinearPositionTarget)).ToSharedRef();
+	
+
+	TSharedRef<SWidget> LinearPositionXWidget = JoinPropertyWidgets(LinearPositionTargetProperty, FName("X"), ConstraintInstance, FName("bLinearXPositionDrive"), LinearXPositionDriveProperty, EPropertyType::LinearXPositionDrive);
+	TSharedRef<SWidget> LinearPositionYWidget = JoinPropertyWidgets(LinearPositionTargetProperty, FName("Y"), ConstraintInstance, FName("bLinearYPositionDrive"), LinearYPositionDriveProperty, EPropertyType::LinearYPositionDrive);
+	TSharedRef<SWidget> LinearPositionZWidget = JoinPropertyWidgets(LinearPositionTargetProperty, FName("Z"), ConstraintInstance, FName("bLinearZPositionDrive"), LinearZPositionDriveProperty, EPropertyType::LinearZPositionDrive);
+
+
+	FDetailWidgetRow& LinearPositionTargetWidget = PositionGroup.AddPropertyRow(LinearPositionTargetProperty).CustomWidget(true);
+	LinearPositionTargetWidget.NameContent()
+	[
+			LinearPositionTargetProperty->CreatePropertyNameWidget()
+	];
+
+	LinearPositionTargetWidget.ValueContent()
+	.MinDesiredWidth(125 * 3 + 18 * 3)
+	.MaxDesiredWidth(125 * 3 + 18 * 3)
+	[
+		SNew(SHorizontalBox)
+		+ SHorizontalBox::Slot()
+		[
+			LinearPositionXWidget
+		]
+		+ SHorizontalBox::Slot()
+		.Padding(5, 0, 0, 0)
+		[
+			LinearPositionYWidget
+		]
+
+		+ SHorizontalBox::Slot()
+		.Padding(5, 0, 0, 0)
+		[
+			LinearPositionZWidget
+		]
+	];
+
+	PositionGroup.AddPropertyRow(ConstraintInstance->GetChildHandle(GET_MEMBER_NAME_CHECKED(FConstraintInstance, LinearDriveSpring)).ToSharedRef())
+		.IsEnabled(TAttribute<bool>::Create(TAttribute<bool>::FGetter::CreateSP(this, &FPhysicsConstraintComponentDetails::IsPropertyEnabled, EPropertyType::LinearPositionDrive)));
+
+	//velocity
+	LinearVelocityDriveProperty = ConstraintInstance->GetChildHandle(GET_MEMBER_NAME_CHECKED(FConstraintInstance, bLinearVelocityDrive));
+	TSharedRef<IPropertyHandle> LinearVelocityTargetProperty = ConstraintInstance->GetChildHandle(GET_MEMBER_NAME_CHECKED(FConstraintInstance, LinearVelocityTarget)).ToSharedRef();
+
+	IDetailGroup& VelocityGroup = LinearMotorCat.AddGroup("Linear Velocity Drive", FText::GetEmpty());
+	VelocityGroup.HeaderProperty(LinearVelocityDriveProperty.ToSharedRef());
+
+	TSharedRef<SWidget> LinearVelocityXWidget = JoinPropertyWidgets(LinearVelocityTargetProperty, FName("X"), ConstraintInstance, FName("bLinearXVelocityDrive"), LinearXVelocityDriveProperty, EPropertyType::LinearXVelocityDrive);
+	TSharedRef<SWidget> LinearVelocityYWidget = JoinPropertyWidgets(LinearVelocityTargetProperty, FName("Y"), ConstraintInstance, FName("bLinearYVelocityDrive"), LinearYVelocityDriveProperty, EPropertyType::LinearYVelocityDrive);
+	TSharedRef<SWidget> LinearVelocityZWidget = JoinPropertyWidgets(LinearVelocityTargetProperty, FName("Z"), ConstraintInstance, FName("bLinearZVelocityDrive"), LinearZVelocityDriveProperty, EPropertyType::LinearZVelocityDrive);
+
+
+	FDetailWidgetRow& LinearVelocityTargetWidget = VelocityGroup.AddPropertyRow(LinearVelocityTargetProperty).CustomWidget(true);
+	LinearVelocityTargetWidget.NameContent()
+	[
+		LinearVelocityTargetProperty->CreatePropertyNameWidget()
+	];
+
+	LinearVelocityTargetWidget.ValueContent()
+	.MinDesiredWidth(125 * 3 + 18 * 3)
+	.MaxDesiredWidth(125 * 3 + 18 * 3)
+	[
+		SNew(SHorizontalBox)
+		+ SHorizontalBox::Slot()
+		[
+			LinearVelocityXWidget
+		]
+		+ SHorizontalBox::Slot()
+		.Padding(5, 0, 0, 0)
+		[
+			LinearVelocityYWidget
+		]
+
+		+ SHorizontalBox::Slot()
+		.Padding(5, 0, 0, 0)
+		[
+			LinearVelocityZWidget
+		]
+	];
+
+	VelocityGroup.AddPropertyRow(ConstraintInstance->GetChildHandle(GET_MEMBER_NAME_CHECKED(FConstraintInstance, LinearDriveDamping)).ToSharedRef())
+		.IsEnabled(TAttribute<bool>::Create(TAttribute<bool>::FGetter::CreateSP(this, &FPhysicsConstraintComponentDetails::IsPropertyEnabled, EPropertyType::LinearVelocityDrive)));
+
+	LinearMotorCat.AddProperty(ConstraintInstance->GetChildHandle(GET_MEMBER_NAME_CHECKED(FConstraintInstance, LinearDriveForceLimit)).ToSharedRef())
+		.IsEnabled(TAttribute<bool>::Create(TAttribute<bool>::FGetter::CreateSP(this, &FPhysicsConstraintComponentDetails::IsPropertyEnabled, EPropertyType::LinearAnyDrive)));
+}
+
+void FPhysicsConstraintComponentDetails::AddAngularDrive(IDetailLayoutBuilder& DetailBuilder, TSharedPtr<IPropertyHandle> ConstraintInstance)
+{
+	AngularPositionDriveProperty = ConstraintInstance->GetChildHandle(GET_MEMBER_NAME_CHECKED(FConstraintInstance, bAngularOrientationDrive));
+	AngularVelocityDriveProperty = ConstraintInstance->GetChildHandle(GET_MEMBER_NAME_CHECKED(FConstraintInstance, bAngularVelocityDrive));
+
+	IDetailCategoryBuilder& AngularMotorCat = DetailBuilder.EditCategory("AngularMotor");
+
+	IDetailGroup& PositionGroup = AngularMotorCat.AddGroup("Angular Orientation Drive", FText::GetEmpty());
+	PositionGroup.HeaderProperty(AngularPositionDriveProperty.ToSharedRef());
+
+	PositionGroup.AddPropertyRow(ConstraintInstance->GetChildHandle(GET_MEMBER_NAME_CHECKED(FConstraintInstance, AngularOrientationTarget)).ToSharedRef())
+		.IsEnabled(TAttribute<bool>::Create(TAttribute<bool>::FGetter::CreateSP(this, &FPhysicsConstraintComponentDetails::IsPropertyEnabled, EPropertyType::AngularPositionDrive)));
+
+	PositionGroup.AddPropertyRow(ConstraintInstance->GetChildHandle(GET_MEMBER_NAME_CHECKED(FConstraintInstance, AngularDriveSpring)).ToSharedRef())
+		.IsEnabled(TAttribute<bool>::Create(TAttribute<bool>::FGetter::CreateSP(this, &FPhysicsConstraintComponentDetails::IsPropertyEnabled, EPropertyType::AngularAnyDrive)));
+
+	IDetailGroup& VelocityGroup = AngularMotorCat.AddGroup("Angular Velocity Drive", FText::GetEmpty());
+	VelocityGroup.HeaderProperty(AngularVelocityDriveProperty.ToSharedRef());
+	VelocityGroup.AddPropertyRow(ConstraintInstance->GetChildHandle(GET_MEMBER_NAME_CHECKED(FConstraintInstance, AngularVelocityTarget)).ToSharedRef())
+		.IsEnabled(TAttribute<bool>::Create(TAttribute<bool>::FGetter::CreateSP(this, &FPhysicsConstraintComponentDetails::IsPropertyEnabled, EPropertyType::AngularVelocityDrive)));
+	VelocityGroup.AddPropertyRow(ConstraintInstance->GetChildHandle(GET_MEMBER_NAME_CHECKED(FConstraintInstance, AngularDriveDamping)).ToSharedRef())
+		.IsEnabled(TAttribute<bool>::Create(TAttribute<bool>::FGetter::CreateSP(this, &FPhysicsConstraintComponentDetails::IsPropertyEnabled, EPropertyType::AngularAnyDrive)));
+	AngularMotorCat.AddProperty(ConstraintInstance->GetChildHandle(GET_MEMBER_NAME_CHECKED(FConstraintInstance, AngularDriveForceLimit)).ToSharedRef())
+		.IsEnabled(TAttribute<bool>::Create(TAttribute<bool>::FGetter::CreateSP(this, &FPhysicsConstraintComponentDetails::IsPropertyEnabled, EPropertyType::AngularAnyDrive)));
+	AngularMotorCat.AddProperty(ConstraintInstance->GetChildHandle(GET_MEMBER_NAME_CHECKED(FConstraintInstance, AngularDriveMode)).ToSharedRef())
+		.IsEnabled(TAttribute<bool>::Create(TAttribute<bool>::FGetter::CreateSP(this, &FPhysicsConstraintComponentDetails::IsPropertyEnabled, EPropertyType::AngularAnyDrive)));
+}
+
+void FPhysicsConstraintComponentDetails::CustomizeDetails( IDetailLayoutBuilder& DetailBuilder )
+{
+	TArray<TWeakObjectPtr<UObject>> Objects;
+	DetailBuilder.GetObjectsBeingCustomized(Objects);
+
+	TSharedPtr<IPropertyHandle> ConstraintInstance;
+	UPhysicsConstraintComponent* ConstraintComp = NULL;
+	APhysicsConstraintActor* OwningConstraintActor = NULL;
+
+	bInPhat = false;
+
+	for (int32 i=0; i < Objects.Num(); ++i)
 	{
-		AngularPositionDriveProperty = ConstraintInstance->GetChildHandle(GET_MEMBER_NAME_CHECKED(FConstraintInstance, bAngularOrientationDrive));
-		AngularVelocityDriveProperty = ConstraintInstance->GetChildHandle(GET_MEMBER_NAME_CHECKED(FConstraintInstance, bAngularVelocityDrive));
-
-		IDetailCategoryBuilder& AngularMotorCat = DetailBuilder.EditCategory("AngularMotor");
-
-		IDetailGroup& PositionGroup = AngularMotorCat.AddGroup("Angular Orientation Drive", FText::GetEmpty());
-		PositionGroup.HeaderProperty(AngularPositionDriveProperty.ToSharedRef());
-		
-		PositionGroup.AddPropertyRow(ConstraintInstance->GetChildHandle(GET_MEMBER_NAME_CHECKED(FConstraintInstance, AngularOrientationTarget)).ToSharedRef())
-			.Visibility(TAttribute<EVisibility>::Create(TAttribute<EVisibility>::FGetter::CreateSP(this, &FPhysicsConstraintComponentDetails::IsPropertyVisible, EPropertyType::AngularPositionDrive)));
-
-		IDetailGroup& VelocityGroup = AngularMotorCat.AddGroup("Angular Velocity Drive", FText::GetEmpty());
-		VelocityGroup.HeaderProperty(AngularVelocityDriveProperty.ToSharedRef());
-		VelocityGroup.AddPropertyRow(ConstraintInstance->GetChildHandle(GET_MEMBER_NAME_CHECKED(FConstraintInstance, AngularVelocityTarget)).ToSharedRef())
-			.Visibility(TAttribute<EVisibility>::Create(TAttribute<EVisibility>::FGetter::CreateSP(this, &FPhysicsConstraintComponentDetails::IsPropertyVisible, EPropertyType::AngularVelocityDrive)));
-		VelocityGroup.AddPropertyRow(ConstraintInstance->GetChildHandle(GET_MEMBER_NAME_CHECKED(FConstraintInstance, AngularDriveForceLimit)).ToSharedRef())
-			.Visibility(TAttribute<EVisibility>::Create(TAttribute<EVisibility>::FGetter::CreateSP(this, &FPhysicsConstraintComponentDetails::IsPropertyVisible, EPropertyType::AngularVelocityDrive)));
-
-		AngularMotorCat.AddProperty(ConstraintInstance->GetChildHandle(GET_MEMBER_NAME_CHECKED(FConstraintInstance, AngularDriveSpring)).ToSharedRef())
-			.Visibility(TAttribute<EVisibility>::Create(TAttribute<EVisibility>::FGetter::CreateSP(this, &FPhysicsConstraintComponentDetails::IsPropertyVisible, EPropertyType::AngularAnyDrive)));
-		AngularMotorCat.AddProperty(ConstraintInstance->GetChildHandle(GET_MEMBER_NAME_CHECKED(FConstraintInstance, AngularDriveDamping)).ToSharedRef())
-			.Visibility(TAttribute<EVisibility>::Create(TAttribute<EVisibility>::FGetter::CreateSP(this, &FPhysicsConstraintComponentDetails::IsPropertyVisible, EPropertyType::AngularAnyDrive)));
+		if (!Objects[i].IsValid()) { continue; }
+
+		if (Objects[i]->IsA(UPhysicsConstraintTemplate::StaticClass()))
+		{
+			ConstraintInstance = DetailBuilder.GetProperty(GET_MEMBER_NAME_CHECKED(UPhysicsConstraintTemplate, DefaultInstance));
+			bInPhat = true;
+			break;
+		}
+		else if (Objects[i]->IsA(UPhysicsConstraintComponent::StaticClass()))
+		{
+			ConstraintInstance = DetailBuilder.GetProperty(GET_MEMBER_NAME_CHECKED(UPhysicsConstraintComponent, ConstraintInstance));
+			ConstraintComp = (UPhysicsConstraintComponent*)Objects[i].Get();
+			OwningConstraintActor = Cast<APhysicsConstraintActor>(ConstraintComp->GetOwner());
+			break;
+		}
 	}
+
+	AddConstraintPresets(DetailBuilder, OwningConstraintActor);
+	AddLinearLimits(DetailBuilder, ConstraintInstance);
+	AddAngularLimits(DetailBuilder, ConstraintInstance);
+	AddLinearDrive(DetailBuilder, ConstraintInstance);
+	AddAngularDrive(DetailBuilder, ConstraintInstance);
 }
 
-EVisibility FPhysicsConstraintComponentDetails::IsPropertyVisible( EPropertyType::Type Type ) const
+bool FPhysicsConstraintComponentDetails::IsPropertyEnabled( EPropertyType::Type Type ) const
 {
 	bool bIsVisible = false;
 	switch (Type)
 	{
-		case EPropertyType::LinearPositionDrive:	bIsVisible = GetBoolProperty(LinearPositionDriveProperty); break;
-		case EPropertyType::LinearVelocityDrive:	bIsVisible = GetBoolProperty(LinearVelocityDriveProperty); break;
-		case EPropertyType::LinearAnyDrive:			bIsVisible = GetBoolProperty(LinearPositionDriveProperty) || GetBoolProperty(LinearVelocityDriveProperty); break;
-		case EPropertyType::AngularPositionDrive:	bIsVisible = GetBoolProperty(AngularPositionDriveProperty); break;
-		case EPropertyType::AngularVelocityDrive:	bIsVisible = GetBoolProperty(AngularVelocityDriveProperty); break;
-		case EPropertyType::AngularAnyDrive:		bIsVisible = GetBoolProperty(AngularPositionDriveProperty) || GetBoolProperty(AngularVelocityDriveProperty); break;
-		case EPropertyType::LinearLimit:			bIsVisible = IsLinearMotionLimited(); break;
-		case EPropertyType::AngularSwing1Limit:		bIsVisible = IsAngularPropertyLimited(AngularSwing1MotionProperty); break;
-		case EPropertyType::AngularSwing2Limit:		bIsVisible = IsAngularPropertyLimited(AngularSwing2MotionProperty); break;
-		case EPropertyType::AngularSwingLimit:		bIsVisible = IsAngularPropertyLimited(AngularSwing1MotionProperty) || IsAngularPropertyLimited(AngularSwing2MotionProperty); break;
-		case EPropertyType::AngularTwistLimit:		bIsVisible = IsAngularPropertyLimited(AngularTwistMotionProperty); break;
-		case EPropertyType::AngularAnyLimit:		bIsVisible = IsAngularPropertyLimited(AngularSwing1MotionProperty) || IsAngularPropertyLimited(AngularSwing2MotionProperty) || 
-																 IsAngularPropertyLimited(AngularTwistMotionProperty); break;
+		case EPropertyType::LinearXPositionDrive:	return GetBoolProperty(LinearXPositionDriveProperty);
+		case EPropertyType::LinearYPositionDrive:	return GetBoolProperty(LinearYPositionDriveProperty);
+		case EPropertyType::LinearZPositionDrive:	return GetBoolProperty(LinearZPositionDriveProperty);
+
+		case EPropertyType::LinearXVelocityDrive:	return GetBoolProperty(LinearXVelocityDriveProperty);
+		case EPropertyType::LinearYVelocityDrive:	return GetBoolProperty(LinearYVelocityDriveProperty);
+		case EPropertyType::LinearZVelocityDrive:	return GetBoolProperty(LinearZVelocityDriveProperty);
+
+		case EPropertyType::LinearPositionDrive:	return GetBoolProperty(LinearPositionDriveProperty);
+		case EPropertyType::LinearVelocityDrive:	return GetBoolProperty(LinearVelocityDriveProperty);
+		case EPropertyType::LinearAnyDrive:			return GetBoolProperty(LinearPositionDriveProperty) || GetBoolProperty(LinearVelocityDriveProperty);
+		case EPropertyType::AngularPositionDrive:	return GetBoolProperty(AngularPositionDriveProperty);
+		case EPropertyType::AngularVelocityDrive:	return GetBoolProperty(AngularVelocityDriveProperty);
+		case EPropertyType::AngularAnyDrive:		return GetBoolProperty(AngularPositionDriveProperty) || GetBoolProperty(AngularVelocityDriveProperty);
+		case EPropertyType::LinearLimit:			return IsLinearMotionLimited();
+		case EPropertyType::AngularSwing1Limit:		return IsAngularPropertyLimited(AngularSwing1MotionProperty);
+		case EPropertyType::AngularSwing2Limit:		return IsAngularPropertyLimited(AngularSwing2MotionProperty);
+		case EPropertyType::AngularSwingLimit:		return IsAngularPropertyLimited(AngularSwing1MotionProperty) || IsAngularPropertyLimited(AngularSwing2MotionProperty);
+		case EPropertyType::AngularTwistLimit:		return IsAngularPropertyLimited(AngularTwistMotionProperty);
+		case EPropertyType::AngularAnyLimit:		return IsAngularPropertyLimited(AngularSwing1MotionProperty) || IsAngularPropertyLimited(AngularSwing2MotionProperty) || 
+																 IsAngularPropertyLimited(AngularTwistMotionProperty);
 	}	
 
-	return bIsVisible ? EVisibility::Visible : EVisibility::Hidden;
+	return bIsVisible;
 }
 
 bool FPhysicsConstraintComponentDetails::GetBoolProperty(TSharedPtr<IPropertyHandle> Prop) const
diff --git a/Engine/Source/Editor/DetailCustomizations/Private/PhysicsConstraintComponentDetails.h b/Engine/Source/Editor/DetailCustomizations/Private/PhysicsConstraintComponentDetails.h
--- a/Engine/Source/Editor/DetailCustomizations/Private/PhysicsConstraintComponentDetails.h
+++ b/Engine/Source/Editor/DetailCustomizations/Private/PhysicsConstraintComponentDetails.h
@@ -19,7 +19,13 @@ private:
 	{
 		enum Type
 		{
+			LinearXPositionDrive,
+			LinearYPositionDrive,
+			LinearZPositionDrive,
 			LinearPositionDrive,
+			LinearXVelocityDrive,
+			LinearYVelocityDrive,
+			LinearZVelocityDrive,
 			LinearVelocityDrive,
 			LinearAnyDrive,
 			LinearLimit,
@@ -35,7 +41,7 @@ private:
 		};
 	};
 
-	EVisibility IsPropertyVisible(EPropertyType::Type Type) const;
+	bool IsPropertyEnabled(EPropertyType::Type Type) const;
 	
 	bool IsLinearMotionLimited() const;
 	
@@ -49,9 +55,24 @@ private:
 	ECheckBoxState IsLimitRadioChecked(TSharedPtr<IPropertyHandle> Property, uint8 Value) const ;
 	void OnLimitRadioChanged(ECheckBoxState CheckType, TSharedPtr<IPropertyHandle> Property, uint8 Value);
 
+	void AddConstraintPresets(IDetailLayoutBuilder& DetailBuilder, class APhysicsConstraintActor* OwningConstraintActor);
+	void AddLinearLimits(IDetailLayoutBuilder& DetailBuilder, TSharedPtr<IPropertyHandle> ConstraintInstance);
+	void AddAngularLimits(IDetailLayoutBuilder& DetailBuilder, TSharedPtr<IPropertyHandle> ConstraintInstance);
+	void AddLinearDrive(IDetailLayoutBuilder& DetailBuilder, TSharedPtr<IPropertyHandle> ConstraintInstance);
+	void AddAngularDrive(IDetailLayoutBuilder& DetailBuilder, TSharedPtr<IPropertyHandle> ConstraintInstance);
+
+	TSharedRef<SWidget> JoinPropertyWidgets(TSharedPtr<IPropertyHandle> TargetProperty, FName TargetChildName, TSharedPtr<IPropertyHandle> ConstraintInstance, FName CheckPropertyName, TSharedPtr<IPropertyHandle>& StoreCheckProperty, EPropertyType::Type DriveType);
+private:
 
 	TSharedPtr<IPropertyHandle> LinearPositionDriveProperty;
+	TSharedPtr<IPropertyHandle> LinearXPositionDriveProperty;
+	TSharedPtr<IPropertyHandle> LinearYPositionDriveProperty;
+	TSharedPtr<IPropertyHandle> LinearZPositionDriveProperty;
+
 	TSharedPtr<IPropertyHandle> LinearVelocityDriveProperty;
+	TSharedPtr<IPropertyHandle> LinearXVelocityDriveProperty;
+	TSharedPtr<IPropertyHandle> LinearYVelocityDriveProperty;
+	TSharedPtr<IPropertyHandle> LinearZVelocityDriveProperty;
 	TSharedPtr<IPropertyHandle> AngularPositionDriveProperty;
 	TSharedPtr<IPropertyHandle> AngularVelocityDriveProperty;
 
@@ -64,4 +85,6 @@ private:
 	TSharedPtr<IPropertyHandle> AngularTwistMotionProperty;
 
 	TWeakObjectPtr<UObject> ConstraintComponent;
+	
+	bool bInPhat;
 };
diff --git a/Engine/Source/Runtime/Core/Public/UObject/ObjectVersion.h b/Engine/Source/Runtime/Core/Public/UObject/ObjectVersion.h
--- a/Engine/Source/Runtime/Core/Public/UObject/ObjectVersion.h
+++ b/Engine/Source/Runtime/Core/Public/UObject/ObjectVersion.h
@@ -504,6 +504,8 @@ enum EUnrealEngineObjectUE4Version
 	VER_UE4_MATERIAL_INSTANCE_BASE_PROPERTY_OVERRIDES_DITHERED_LOD_TRANSITION,
 	// Serialize ES2 textures separately rather than overwriting the properties used on other platforms
 	VER_UE4_SERIALIZE_LANDSCAPE_ES2_TEXTURES,
+	// Constraint motor velocity is broken into per-component
+	VER_UE4_CONSTRAINT_INSTANCE_MOTOR_FLAGS,
 
 	// -----<new versions can be added before this line>-------------------------------------------------
 	// - this needs to be the last line (see note below)
diff --git a/Engine/Source/Runtime/Engine/Classes/PhysicsEngine/ConstraintInstance.h b/Engine/Source/Runtime/Engine/Classes/PhysicsEngine/ConstraintInstance.h
--- a/Engine/Source/Runtime/Engine/Classes/PhysicsEngine/ConstraintInstance.h
+++ b/Engine/Source/Runtime/Engine/Classes/PhysicsEngine/ConstraintInstance.h
@@ -251,53 +251,59 @@ struct ENGINE_API FConstraintInstance
 	UPROPERTY(EditAnywhere, AdvancedDisplay, Category=Angular, meta=(editcondition = "bAngularBreakable", ClampMin = "0.0"))
 	float AngularBreakThreshold;
 
-	/** Enables/Disables linear drive along the x axis. */
-	UPROPERTY(EditAnywhere, BlueprintReadOnly, Category=LinearMotor)
+private:
+	/** Enables/Disables linear position drive along the x axis. */
+	UPROPERTY(EditAnywhere, Category=LinearMotor)
 	uint32 bLinearXPositionDrive:1;
 
-	UPROPERTY()
-	uint32 bLinearXVelocityDrive_DEPRECATED:1;
+	/** Enables/Disables linear velocity drive along the x axis. */
+	UPROPERTY(EditAnywhere, Category=LinearMotor)
+	uint32 bLinearXVelocityDrive:1;
 
-	/** Enables/Disables linear drive along the y axis. */
-	UPROPERTY(EditAnywhere, BlueprintReadOnly, Category=LinearMotor)
+	/** Enables/Disables linear position drive along the y axis. */
+	UPROPERTY(EditAnywhere, Category=LinearMotor)
 	uint32 bLinearYPositionDrive:1;
 
-	UPROPERTY()
-	uint32 bLinearYVelocityDrive_DEPRECATED:1;
+	/** Enables/Disables linear velocity drive along the y axis. */
+	UPROPERTY(EditAnywhere, Category=LinearMotor)
+	uint32 bLinearYVelocityDrive:1;
 
-	/** Enables/Disables linear drive along the z axis. */
-	UPROPERTY(EditAnywhere, BlueprintReadOnly, Category=LinearMotor)
+	/** Enables/Disables linear position drive along the z axis. */
+	UPROPERTY(EditAnywhere, Category=LinearMotor)
 	uint32 bLinearZPositionDrive:1;
 
-	UPROPERTY()
-	uint32 bLinearZVelocityDrive_DEPRECATED:1;
+	/** Enables/Disables linear velocity drive along the z axis. */
+	UPROPERTY(EditAnywhere, Category=LinearMotor)
+	uint32 bLinearZVelocityDrive:1;
 
-	/** Enables/Disables the linear position drive. */
+public:
+
+	/** Enables/Disables linear position drive. */
 	UPROPERTY(EditAnywhere, BlueprintReadOnly, Category=LinearMotor)
 	uint32 bLinearPositionDrive:1;
 	
-	/** Enables/Disables the linear velocity drive. */
-	UPROPERTY(EditAnywhere, BlueprintReadOnly, Category=LinearMotor)
+	/** Enables/Disables linear velocity drive. */
+	UPROPERTY(EditAnywhere, BlueprintReadOnly, Category = LinearMotor)
 	uint32 bLinearVelocityDrive:1;
 
 	/** Target position the linear drive. Only the components that are enabled are used. */
-	UPROPERTY(EditAnywhere, BlueprintReadOnly, Category=LinearMotor)
+	UPROPERTY(EditAnywhere, BlueprintReadOnly, Category=LinearMotor, meta=(editcondition="bLinearPositionDrive"))
 	FVector LinearPositionTarget;
 
 	/** Target velocity the linear drive. */
-	UPROPERTY(EditAnywhere, BlueprintReadOnly, Category=LinearMotor)
+	UPROPERTY(EditAnywhere, BlueprintReadOnly, Category = LinearMotor, meta = (editcondition = "bLinearVelocityDrive"))
 	FVector LinearVelocityTarget;
 
 	/** Spring to apply to the for linear drive. */
-	UPROPERTY(EditAnywhere, BlueprintReadOnly, Category = LinearMotor, meta = (DisplayName = "Linear Position Strength", editcondition = "bLinearPositionDrive"))
+	UPROPERTY(EditAnywhere, BlueprintReadOnly, Category = LinearMotor, meta = (DisplayName = "Linear Position Strength"))
 	float LinearDriveSpring;
 
 	/** Damping to apply to the for linear drive. */
-	UPROPERTY(EditAnywhere, BlueprintReadOnly, Category = LinearMotor, meta = (DisplayName = "Linear Velocity Strength", editcondition = "bLinearVelocityDrive"))
+	UPROPERTY(EditAnywhere, BlueprintReadOnly, Category = LinearMotor, meta = (DisplayName = "Linear Velocity Strength"))
 	float LinearDriveDamping;
 
 	/** Limit to the force the linear drive can apply. */
-	UPROPERTY(EditAnywhere, BlueprintReadOnly, Category=LinearMotor)
+	UPROPERTY(EditAnywhere, BlueprintReadOnly, Category=LinearMotor, meta = (DisplayName = "Max Linear Force"))
 	float LinearDriveForceLimit;
 
 	UPROPERTY()
@@ -343,11 +349,11 @@ struct ENGINE_API FConstraintInstance
 	float AngularDriveSpring;
 
 	/** Damping value to apply to the for angular drive. */
-	UPROPERTY(EditAnywhere, BlueprintReadOnly, Category = AngularMotor, meta = (DisplayName = "Angular Velocity Strength"))
+	UPROPERTY(EditAnywhere, BlueprintReadOnly, Category = AngularMotor, meta = (DisplayName = "Angular Velocity Strength", editcondition = "bAngularVelocityDrive"))
 	float AngularDriveDamping;
 
 	/** Limit to the force the angular drive can apply. */
-	UPROPERTY(EditAnywhere, BlueprintReadOnly, Category=AngularMotor)
+	UPROPERTY(EditAnywhere, BlueprintReadOnly, Category = AngularMotor, meta = (DisplayName = "Max Angular Force"))
 	float AngularDriveForceLimit;
 
 	float AverageMass;
diff --git a/Engine/Source/Runtime/Engine/Private/PhysicsEngine/ConstraintInstance.cpp b/Engine/Source/Runtime/Engine/Private/PhysicsEngine/ConstraintInstance.cpp
--- a/Engine/Source/Runtime/Engine/Private/PhysicsEngine/ConstraintInstance.cpp
+++ b/Engine/Source/Runtime/Engine/Private/PhysicsEngine/ConstraintInstance.cpp
@@ -813,6 +813,9 @@ void FConstraintInstance::SetLinearVelocityDrive(bool bEnableXDrive, bool bEnabl
 	}
 #endif
 
+	bLinearXVelocityDrive = bEnableXDrive;
+	bLinearYVelocityDrive = bEnableYDrive;
+	bLinearZVelocityDrive = bEnableZDrive;
 	bLinearVelocityDrive = bEnableXDrive || bEnableYDrive || bEnableZDrive;
 }
 
@@ -919,18 +922,18 @@ void FConstraintInstance::SetLinearDriveParams(float InSpring, float InDamping,
 		{
 			// X-Axis linear drive
 			const float DriveSpringX = bLinearPositionDrive && bLinearXPositionDrive ? InSpring : 0.0f;
-			const float DriveDampingX = (bLinearVelocityDrive && LinearVelocityTarget.X != 0.f) ? InDamping : 0.0f;
+			const float DriveDampingX = bLinearVelocityDrive && bLinearXVelocityDrive ? InDamping : 0.0f;
 			const float LinearForceLimit = InForceLimit > 0.f ? InForceLimit : PX_MAX_F32;
 			Joint->setDrive(PxD6Drive::eX, PxD6JointDrive(DriveSpringX, DriveDampingX, LinearForceLimit, bIsAccelerationDrive));
 
 			// Y-Axis linear drive
 			const float DriveSpringY = bLinearPositionDrive && bLinearYPositionDrive ? InSpring : 0.0f;
-			const float DriveDampingY = (bLinearVelocityDrive && LinearVelocityTarget.Y != 0.f) ? InDamping : 0.0f;
+			const float DriveDampingY = bLinearVelocityDrive && bLinearYVelocityDrive ? InDamping : 0.0f;
 			Joint->setDrive(PxD6Drive::eY, PxD6JointDrive(DriveSpringY, DriveDampingY, LinearForceLimit, bIsAccelerationDrive));
 
 			// Z-Axis linear drive
 			const float DriveSpringZ = bLinearPositionDrive && bLinearZPositionDrive ? InSpring : 0.0f;
-			const float DriveDampingZ = (bLinearVelocityDrive && LinearVelocityTarget.Z != 0.f) ? InDamping : 0.0f;
+			const float DriveDampingZ = bLinearVelocityDrive && bLinearZVelocityDrive ? InDamping : 0.0f;
 			Joint->setDrive(PxD6Drive::eZ, PxD6JointDrive(DriveSpringZ, DriveDampingZ, LinearForceLimit, bIsAccelerationDrive));
 		});
 	}
@@ -1259,6 +1262,13 @@ void FConstraintInstance::PostSerialize(const FArchive& Ar)
 	{
 		AngularVelocityTarget *= 1.f / (2.f * PI);	//we want to use revolutions per second - old system was using radians directly
 	}
+
+	if (Ar.IsLoading() && Ar.UE4Ver() < VER_UE4_CONSTRAINT_INSTANCE_MOTOR_FLAGS)
+	{
+		bLinearXVelocityDrive = LinearVelocityTarget.X != 0.f;
+		bLinearYVelocityDrive = LinearVelocityTarget.Y != 0.f;
+		bLinearZVelocityDrive = LinearVelocityTarget.Z != 0.f;
+	}
 }
 
 void FConstraintInstance::OnConstraintBroken()
