Comment cleanup in Config Editor

diff --git a/Engine/Source/Editor/DetailCustomizations/Private/ConfigEditorPropertyDetails.cpp b/Engine/Source/Editor/DetailCustomizations/Private/ConfigEditorPropertyDetails.cpp
--- a/Engine/Source/Editor/DetailCustomizations/Private/ConfigEditorPropertyDetails.cpp
+++ b/Engine/Source/Editor/DetailCustomizations/Private/ConfigEditorPropertyDetails.cpp
@@ -63,7 +63,6 @@ void FConfigPropertyHelperDetails::OnPropertyValueChanged(UObject* Object, FProp
 }
 
 
-// TPMB - This didnt work... why?
 void FConfigPropertyHelperDetails::OnConfigFileListChanged()
 {
 	PropertyTable->RequestRefresh();
