Changing the project settings value no longer changes all of the values in the Config Editor for that property.

UE-14444

diff --git a/Engine/Source/Editor/DetailCustomizations/Private/ConfigEditorPropertyDetails.cpp b/Engine/Source/Editor/DetailCustomizations/Private/ConfigEditorPropertyDetails.cpp
--- a/Engine/Source/Editor/DetailCustomizations/Private/ConfigEditorPropertyDetails.cpp
+++ b/Engine/Source/Editor/DetailCustomizations/Private/ConfigEditorPropertyDetails.cpp
@@ -26,19 +26,32 @@ void FConfigPropertyHelperDetails::CustomizeDetails(IDetailLayoutBuilder& Detail
 
 	UObject* PropValue;
 	PropertyHandle->GetValue(PropValue);
+	OriginalProperty = CastChecked<UProperty>(PropValue);
+
+	// Create a runtime UClass with the provided property as the only member. We will use this in the details view for the config hierarchy.
+	ConfigEditorPropertyViewClass = NewObject<UClass>(GetTransientPackage(), TEXT("TempConfigEditorUClass"), RF_Public|RF_Standalone);
 
 	// Keep a record of the UProperty we are looking to update
-	Property = CastChecked<UProperty>(PropValue);
+	ConfigEditorCopyOfEditProperty = DuplicateObject<UProperty>(OriginalProperty, ConfigEditorPropertyViewClass, *(PropValue->GetName()));
+	ConfigEditorPropertyViewClass->ClassConfigName = OriginalProperty->GetOwnerClass()->ClassConfigName;
+	ConfigEditorPropertyViewClass->SetSuperStruct(UObject::StaticClass());
+	ConfigEditorPropertyViewClass->ClassFlags |= (CLASS_DefaultConfig | CLASS_Config);
+	ConfigEditorPropertyViewClass->AddCppProperty(ConfigEditorCopyOfEditProperty);
+	ConfigEditorPropertyViewClass->Bind();
+	ConfigEditorPropertyViewClass->StaticLink(true);
+	ConfigEditorPropertyViewClass->AddToRoot();
+	
+	// Cache the CDO for the object
+	ConfigEditorPropertyViewCDO = ConfigEditorPropertyViewClass->GetDefaultObject(true);
+	ConfigEditorPropertyViewCDO->AddToRoot();
 
 	// Get access to all of the config files where this property is configurable.
 	ConfigFilesHandle = DetailBuilder.GetProperty("ConfigFilePropertyObjects");
-	FSimpleDelegate OnConfigFileListChangedDelegate = FSimpleDelegate::CreateSP(this, &FConfigPropertyHelperDetails::OnConfigFileListChanged);
-	ConfigFilesHandle->SetOnPropertyValueChanged(OnConfigFileListChangedDelegate);
 	DetailBuilder.HideProperty(ConfigFilesHandle);
 
 	// Add the properties to a property table so we can edit these.
-	IDetailCategoryBuilder& A = DetailBuilder.EditCategory("ConfigHierarchy");
-	A.AddCustomRow(LOCTEXT("ConfigHierarchy", "ConfigHierarchy"))
+	IDetailCategoryBuilder& ConfigHierarchyCategory = DetailBuilder.EditCategory("ConfigHierarchy");
+	ConfigHierarchyCategory.AddCustomRow(LOCTEXT("ConfigHierarchy", "ConfigHierarchy"))
 	[
 		// Create a property table with the values.
 		ConstructPropertyTable(DetailBuilder)
@@ -51,45 +64,75 @@ void FConfigPropertyHelperDetails::CustomizeDetails(IDetailLayoutBuilder& Detail
 
 void FConfigPropertyHelperDetails::OnPropertyValueChanged(UObject* Object, FPropertyChangedEvent& PropertyChangedEvent)
 {
-	UClass* OwnerClass = Property->GetOwnerClass();
+	UClass* OwnerClass = ConfigEditorCopyOfEditProperty->GetOwnerClass();
 	if (Object->IsA(OwnerClass))
 	{
 		const FString* FileName = ConfigFileAndPropertySourcePairings.FindKey(Object);
 		if (FileName != nullptr)
 		{
-			Object->UpdateSinglePropertyInConfigFile(Property, *FileName);
-		}
-	}
-}
+			const FString& ConfigIniName = *FileName;
 
+			// We should set this up to work with the UObject Config system, its difficult as the outer object isnt of the same type
+			// create a sandbox FConfigCache
+			FConfigCacheIni Config(EConfigCacheType::Temporary);
 
-void FConfigPropertyHelperDetails::OnConfigFileListChanged()
-{
-	PropertyTable->RequestRefresh();
+			// add an empty file to the config so it doesn't read in the original file (see FConfigCacheIni.Find())
+			FConfigFile& NewFile = Config.Add(ConfigIniName, FConfigFile());
+
+			// save the object properties to this file
+			OriginalProperty->GetOwnerClass()->GetDefaultObject()->SaveConfig(CPF_Config, *ConfigIniName, &Config);
+
+			// Take the saved section for this object and have the config system process and write out the one property we care about.
+			ensureMsgf(Config.Num() == 1, TEXT("UObject::UpdateDefaultConfig() caused more files than expected in the Sandbox config cache!"));
+
+			TArray<FString> Keys;
+			NewFile.GetKeys(Keys);
+
+			const FString SectionName = Keys[0];
+			const FString PropertyName = ConfigEditorCopyOfEditProperty->GetName();
+			FString	Value;
+			ConfigEditorCopyOfEditProperty->ExportText_InContainer(0, Value, Object, Object, Object, 0);
+			NewFile.SetString(*SectionName, *PropertyName, *Value);
+
+			NewFile.UpdateSinglePropertyInSection(*ConfigIniName, *PropertyName, *SectionName);
+
+			// reload the file, so that it refresh the cache internally.
+			FString FinalIniFileName;
+			GConfig->LoadGlobalIniFile(FinalIniFileName, *OriginalProperty->GetOwnerClass()->ClassConfigName.ToString(), NULL, true);
+
+			// Update the CDO, as this change might have had an impact on it's value.
+			OriginalProperty->GetOwnerClass()->GetDefaultObject()->ReloadConfig();
+		}
+	}
 }
 
 
-void FConfigPropertyHelperDetails::AddEditablePropertyForConfig(IDetailLayoutBuilder& DetailBuilder, const UPropertyConfigFileDisplayRow* ConfigFileProperty)
+void FConfigPropertyHelperDetails::AddEditablePropertyForConfig(IDetailLayoutBuilder& DetailBuilder, const UPropertyConfigFileDisplayRow* ConfigFilePropertyRowObj)
 {
-	AssociatedConfigFileAndObjectPairings.Add(ConfigFileProperty->ConfigFileName, (UObject*)ConfigFileProperty);
+	AssociatedConfigFileAndObjectPairings.Add(ConfigFilePropertyRowObj->ConfigFileName, (UObject*)ConfigFilePropertyRowObj);
 
 	// Add the properties to a property table so we can edit these.
-	IDetailCategoryBuilder& A = DetailBuilder.EditCategory("TestCat");
+	IDetailCategoryBuilder& TempCategory = DetailBuilder.EditCategory("TempCategory");
 
-	// Create an instance of the owner object so we can dictate values for our property on a per object basis.
-	UObject* CDO = Property->GetOwnerClass()->GetDefaultObject<UObject>();
-	UObject* ConfigEntryObject = StaticDuplicateObject(CDO, GetTransientPackage(), *(ConfigFileProperty->ConfigFileName + TEXT("_cdoDupe")));
+	UObject* ConfigEntryObject = StaticDuplicateObject(ConfigEditorPropertyViewCDO, GetTransientPackage(), *(ConfigFilePropertyRowObj->ConfigFileName + TEXT("_cdoDupe")));
 	ConfigEntryObject->AddToRoot();
-	ConfigEntryObject->LoadConfig(Property->GetOwnerClass(), *ConfigFileProperty->ConfigFileName, UE4::LCPF_None, Property);
 
-	// Add a reference for future saving.
-	ConfigFileAndPropertySourcePairings.Add(ConfigFileProperty->ConfigFileName, (UObject*)ConfigEntryObject);
+	FString ExistingConfigEntryValue;
+	static FString SectionName = OriginalProperty->GetOwnerClass()->GetPathName();
+	static FString PropertyName = ConfigEditorCopyOfEditProperty->GetName();
+	if (GConfig->GetString(*SectionName, *PropertyName, ExistingConfigEntryValue, ConfigFilePropertyRowObj->ConfigFileName))
+	{
+		ConfigEditorCopyOfEditProperty->ImportText(*ExistingConfigEntryValue, ConfigEditorCopyOfEditProperty->ContainerPtrToValuePtr<uint8>(ConfigEntryObject), 0, nullptr);
+	}
+
+	// Cache a reference for future usage.
+	ConfigFileAndPropertySourcePairings.Add(ConfigFilePropertyRowObj->ConfigFileName, (UObject*)ConfigEntryObject);
 
 	// We need to add a property row for each config file entry. 
 	// This allows us to have an editable widget for each config file.
 	TArray<UObject*> ConfigPropertyDisplayObjects;
 	ConfigPropertyDisplayObjects.Add(ConfigEntryObject);
-	if (IDetailPropertyRow* ExternalRow = A.AddExternalProperty(ConfigPropertyDisplayObjects, Property->GetFName()))
+	if (IDetailPropertyRow* ExternalRow = TempCategory.AddExternalProperty(ConfigPropertyDisplayObjects, ConfigEditorCopyOfEditProperty->GetFName()))
 	{
 		TSharedPtr<SWidget> NameWidget;
 		TSharedPtr<SWidget> ValueWidget;
@@ -98,7 +141,7 @@ void FConfigPropertyHelperDetails::AddEditablePropertyForConfig(IDetailLayoutBui
 		// Register the Value widget and config file pairing with the config editor.
 		// The config editor needs this to determine what a cell presenter shows.
 		IConfigEditorModule& ConfigEditor = FModuleManager::Get().LoadModuleChecked<IConfigEditorModule>("ConfigEditor");
-		ConfigEditor.AddExternalPropertyValueWidgetAndConfigPairing(ConfigFileProperty->ConfigFileName, ValueWidget);
+		ConfigEditor.AddExternalPropertyValueWidgetAndConfigPairing(ConfigFilePropertyRowObj->ConfigFileName, ValueWidget);
 		
 		// now hide the property so it is not added to the property display view
 		ExternalRow->Visibility(EVisibility::Hidden);
@@ -119,7 +162,7 @@ TSharedRef<SWidget> FConfigPropertyHelperDetails::ConstructPropertyTable(IDetail
 
 	TArray< TSharedRef<class IPropertyTableCustomColumn>> CustomColumns;
 	TSharedRef<FConfigPropertyCustomColumn> EditPropertyColumn = MakeShareable(new FConfigPropertyCustomColumn());
-	EditPropertyColumn->EditProperty = Property; // FindFieldChecked<UProperty>(UPropertyConfigFileDisplay::StaticClass(), TEXT("EditProperty"));
+	EditPropertyColumn->EditProperty = ConfigEditorCopyOfEditProperty; // FindFieldChecked<UProperty>(UPropertyConfigFileDisplay::StaticClass(), TEXT("EditProperty"));
 	CustomColumns.Add(EditPropertyColumn);
 
 	//TSharedRef<FConfigPropertyConfigFileStateCustomColumn> ConfigFileStateColumn = MakeShareable(new FConfigPropertyConfigFileStateCustomColumn());
diff --git a/Engine/Source/Editor/DetailCustomizations/Private/ConfigEditorPropertyDetails.h b/Engine/Source/Editor/DetailCustomizations/Private/ConfigEditorPropertyDetails.h
--- a/Engine/Source/Editor/DetailCustomizations/Private/ConfigEditorPropertyDetails.h
+++ b/Engine/Source/Editor/DetailCustomizations/Private/ConfigEditorPropertyDetails.h
@@ -62,7 +62,6 @@ private:
 private:
 	void AddEditablePropertyForConfig(IDetailLayoutBuilder& DetailBuilder, const class UPropertyConfigFileDisplayRow* ConfigFileProperty);
 
-	void OnConfigFileListChanged();
 
 	//////////////////////////////////////////////////////////////////////////
 	/// Misc properties...
@@ -71,13 +70,23 @@ private:
 	// Keep track of the property handle for the config files.
 	TSharedPtr<IPropertyHandle> ConfigFilesHandle;
 
-	// The property we wish to edit on a per config basis.
-	UProperty* Property;
+	// A copy of the edit property we use with our helper class to update values on a per-config basis.
+	UProperty* ConfigEditorCopyOfEditProperty;
+
+	// The original property from the Project settings, that we have chosen to edit.
+	UProperty* OriginalProperty;
+
+	// A runtime class generated with the Original property as a member.
+	// This allows us to edit a property on a per config basis.
+	UClass* ConfigEditorPropertyViewClass;
+
+	// The 'Class Default Object' of the runtime class we generate.
+	// We duplicate this for each config file instance.
+	UObject* ConfigEditorPropertyViewCDO;
 
 	// Coupling of Config files and their editable objects.
 	TMap<FString, UObject*> AssociatedConfigFileAndObjectPairings;
 
+	// Mapping of Config Files and ConfigEditorPropertyViewClass objects.
 	TMap<FString, UObject*> ConfigFileAndPropertySourcePairings;
-
-	TMap<FString, TSharedPtr<SWidget>> ExternalPropertyValueWidgetAndConfigPairings;
 };
