Fix for Config Editor Values were not being cached in GConfig.

UE-15482

diff --git a/Engine/Source/Editor/DetailCustomizations/Private/ConfigEditorPropertyDetails.cpp b/Engine/Source/Editor/DetailCustomizations/Private/ConfigEditorPropertyDetails.cpp
--- a/Engine/Source/Editor/DetailCustomizations/Private/ConfigEditorPropertyDetails.cpp
+++ b/Engine/Source/Editor/DetailCustomizations/Private/ConfigEditorPropertyDetails.cpp
@@ -93,6 +93,7 @@ void FConfigPropertyHelperDetails::OnPropertyValueChanged(UObject* Object, FProp
 			FString	Value;
 			ConfigEditorCopyOfEditProperty->ExportText_InContainer(0, Value, Object, Object, Object, 0);
 			NewFile.SetString(*SectionName, *PropertyName, *Value);
+			GConfig->SetString(*SectionName, *PropertyName, *Value, ConfigIniName);
 
 			NewFile.UpdateSinglePropertyInSection(*ConfigIniName, *PropertyName, *SectionName);
 
