Config Heirarchy Cell selection fix. Fixes UE-15478

diff --git a/Engine/Source/Editor/DetailCustomizations/Private/ConfigEditorPropertyDetails.cpp b/Engine/Source/Editor/DetailCustomizations/Private/ConfigEditorPropertyDetails.cpp
--- a/Engine/Source/Editor/DetailCustomizations/Private/ConfigEditorPropertyDetails.cpp
+++ b/Engine/Source/Editor/DetailCustomizations/Private/ConfigEditorPropertyDetails.cpp
@@ -156,6 +156,7 @@ TSharedRef<SWidget> FConfigPropertyHelperDetails::ConstructPropertyTable(IDetail
 
 	PropertyTable = PropertyEditorModule.CreatePropertyTable();
 	PropertyTable->SetSelectionMode(ESelectionMode::None);
+	PropertyTable->SetSelectionUnit(EPropertyTableSelectionUnit::None);
 	PropertyTable->SetIsUserAllowedToChangeRoot(false);
 	PropertyTable->SetShowObjectName(false);
 
