Fix for crash when opening retargeting options in rig asset editor.
#UE-15800

diff --git a/Engine/Source/Editor/DetailCustomizations/Private/RigDetails.cpp b/Engine/Source/Editor/DetailCustomizations/Private/RigDetails.cpp
--- a/Engine/Source/Editor/DetailCustomizations/Private/RigDetails.cpp
+++ b/Engine/Source/Editor/DetailCustomizations/Private/RigDetails.cpp
@@ -271,7 +271,11 @@ void FRigDetails::GenerateTransformBaseArrayElementWidget(TSharedRef<IPropertyHa
 	// we add them on demand since the array can expand or reduce
 	if(!ParentSpaceOptionList.IsValidIndex(ArrayIndex))
 	{
-		ParentSpaceOptionList.AddZeroed(ArrayIndex-ParentSpaceOptionList.Num()+1);
+		int32 NumToAdd = ArrayIndex-ParentSpaceOptionList.Num()+1;
+		for(int32 Idx = 0 ; Idx < NumToAdd ; ++Idx)
+		{
+			ParentSpaceOptionList.Add(new TArray<TSharedPtr<FString>>);
+		}
 	}
 
 	// ParentSpaceComboBoxees creates 2 per element - one for translation, and one for rotation
diff --git a/Engine/Source/Editor/DetailCustomizations/Private/RigDetails.h b/Engine/Source/Editor/DetailCustomizations/Private/RigDetails.h
--- a/Engine/Source/Editor/DetailCustomizations/Private/RigDetails.h
+++ b/Engine/Source/Editor/DetailCustomizations/Private/RigDetails.h
@@ -24,7 +24,7 @@ private:
 
 	// array custom boxes - these will stay around as long as this window is up
 	TArray<TSharedPtr<SEditableTextBox>>	DisplayNameTextBoxes;
-	TArray<TArray<TSharedPtr<FString>>>		ParentSpaceOptionList;
+	TIndirectArray<TArray<TSharedPtr<FString>>>		ParentSpaceOptionList;
 	TArray<TSharedPtr<SComboBox< TSharedPtr<FString> >>>	ParentSpaceComboBoxes;
 
 	/** we only support one item */
