Disable WinXP for Rocket builds. We dont provide generated source or binarys for WinXP as part of the binary release.

diff --git a/Engine/Source/Editor/DetailCustomizations/Private/WindowsTargetSettingsDetails.cpp b/Engine/Source/Editor/DetailCustomizations/Private/WindowsTargetSettingsDetails.cpp
--- a/Engine/Source/Editor/DetailCustomizations/Private/WindowsTargetSettingsDetails.cpp
+++ b/Engine/Source/Editor/DetailCustomizations/Private/WindowsTargetSettingsDetails.cpp
@@ -8,6 +8,8 @@
 #include "GameProjectGenerationModule.h"
 #include "ISourceControlModule.h"
 
+#define LOCTEXT_NAMESPACE "WindowsTargetSettingsDetails"
+
 namespace WindowsTargetSettingsDetailsConstants
 {
 	/** The filename for the game splash screen */
@@ -15,9 +17,10 @@ namespace WindowsTargetSettingsDetailsConstants
 
 	/** The filename for the editor splash screen */
 	const FString EditorSplashFileName(TEXT("Splash/EdSplash.bmp"));
-}
 
-#define LOCTEXT_NAMESPACE "WindowsTargetSettingsDetails"
+	/** ToolTip used when an option is not available to binary users. */
+	const FText DisabledTip = LOCTEXT("GitHubSourceRequiredToolTip", "This requires GitHub source.");
+}
 
 FText GetFriendlyNameFromRHIName(const FString& InRHIName)
 {
@@ -116,6 +119,15 @@ void FWindowsTargetSettingsDetails::CustomizeDetails( IDetailLayoutBuilder& Deta
 	TargetShaderFormatsDetails = MakeShareable(new FTargetShaderFormatsPropertyDetails(&DetailBuilder));
 	TargetShaderFormatsDetails->CreateTargetShaderFormatsPropertyView();
 
+	TSharedRef<IPropertyHandle> MinOSProperty = DetailBuilder.GetProperty("MinimumOSVersion");
+	IDetailCategoryBuilder& OSInfoCategory = DetailBuilder.EditCategory(TEXT("OS Info"));
+	
+	// Setup edit condition and tool tip of Min OS property. Determined by whether the engine is installed or not.
+	bool bIsMinOSSelectionAvailable = FApp::IsEngineInstalled() == false;
+	IDetailPropertyRow& MinOSRow = OSInfoCategory.AddProperty(MinOSProperty);
+	MinOSRow.IsEnabled(bIsMinOSSelectionAvailable);
+	MinOSRow.ToolTip(bIsMinOSSelectionAvailable ? MinOSProperty->GetToolTipText() : WindowsTargetSettingsDetailsConstants::DisabledTip);
+
 	// Next add the splash image customization
 	const FText EditorSplashDesc(LOCTEXT("EditorSplashLabel", "Editor Splash"));
 	IDetailCategoryBuilder& SplashCategoryBuilder = DetailBuilder.EditCategory(TEXT("Splash"));
