[UE-15985] The "Reset Emitter" button is now functional when ParticleSystem component nodes are selected in the Blueprint editor.

change summary:
- modified FParticleSystemComponentDetails::OnResetEmitter() to redirect the reset to the preview Actor instance if the selected object is an archetype

diff --git a/Engine/Source/Editor/DetailCustomizations/Private/ParticleSystemComponentDetails.cpp b/Engine/Source/Editor/DetailCustomizations/Private/ParticleSystemComponentDetails.cpp
--- a/Engine/Source/Editor/DetailCustomizations/Private/ParticleSystemComponentDetails.cpp
+++ b/Engine/Source/Editor/DetailCustomizations/Private/ParticleSystemComponentDetails.cpp
@@ -5,6 +5,9 @@
 #include "Particles/Emitter.h"
 #include "Particles/ParticleSystemComponent.h"
 
+// @todo Can remove this if we switch to instance-based editing in the Blueprint editor.
+#include "ActorEditorUtils.h"
+
 #define LOCTEXT_NAMESPACE "ParticleSystemComponentDetails"
 
 TSharedRef<IDetailCustomization> FParticleSystemComponentDetails::MakeInstance()
@@ -94,6 +97,24 @@ FReply FParticleSystemComponentDetails::OnResetEmitter()
 				}
 			}
 
+			// If the object selected to the details view is a template, then we need to redirect the reset to the preview instance (e.g. in the Blueprint editor).
+			// @todo Can remove this if we switch to instance-based editing in the Blueprint editor.
+			if (PSC && PSC->IsTemplate())
+			{
+				TArray<UObject*> Instances;
+				PSC->GetArchetypeInstances(Instances);
+				UObject** ElementPtr = Instances.FindByPredicate([](const UObject* Element)
+				{
+					const AActor* Owner = Cast<AActor>(Element->GetOuter());
+					return Owner != nullptr && FActorEditorUtils::IsAPreviewOrInactiveActor(Owner);
+				});
+
+				if (ElementPtr)
+				{
+					PSC = Cast<UParticleSystemComponent>(*ElementPtr);
+				}
+			}
+
 			if (PSC)
 			{
 				PSC->ResetParticles();
