[INTEGRATE] Change 2566527 by [REDACTED_PERSON][REDACTED_HANDLE] on 2015/05/27 09:39:51

	Add experimental/early access warning banner to details panel when an experimental component is selected

diff --git a/Engine/Documentation/Source/Shared/LevelEditor/LevelEditor.INT.udn b/Engine/Documentation/Source/Shared/LevelEditor/LevelEditor.INT.udn
--- a/Engine/Documentation/Source/Shared/LevelEditor/LevelEditor.INT.udn
+++ b/Engine/Documentation/Source/Shared/LevelEditor/LevelEditor.INT.udn
@@ -124,3 +124,32 @@ While not considered production-ready, they are a step beyond 'experimental' and
 * There is probably no documentation or example content yet.
 * They will be promoted out of 'Early Access' when they are production ready.
 [/EXCERPT:ActorUsesEarlyAccessClass]
+
+[EXCERPT:ComponentUsesExperimentalClass]
+###Here be dragons!
+
+![](HereBeDragons.png)
+
+This Component uses one or more 'experimental' classes.
+
+These are totally unsupported and undocumented prototypes:
+* They may be dramatically changed or abandoned in a future build.
+* Data or code compatibility is not guaranteed.
+* The code is in an early state and may cause editor instability or crashes.
+* They may be promoted to 'Early Access' or 'Fully Supported' in the future.
+
+Use these at your own risk!
+[/EXCERPT:ComponentUsesExperimentalClass]
+
+[EXCERPT:ComponentUsesEarlyAccessClass]
+###Early Access
+
+This Component uses one or more 'early access' classes.
+
+While not considered production-ready, they are a step beyond 'experimental' and are being provided as a preview of things to come:
+* We will try to provide forward-compatibility for content you create.
+* The classes may change significantly in the future.
+* The code is in an early state and may not meet the desired polish / quality bar.
+* There is probably no documentation or example content yet.
+* They will be promoted out of 'Early Access' when they are production ready.
+[/EXCERPT:ComponentUsesEarlyAccessClass]
\ No newline at end of file
diff --git a/Engine/Source/Editor/DetailCustomizations/Private/ActorComponentDetails.cpp b/Engine/Source/Editor/DetailCustomizations/Private/ActorComponentDetails.cpp
--- a/Engine/Source/Editor/DetailCustomizations/Private/ActorComponentDetails.cpp
+++ b/Engine/Source/Editor/DetailCustomizations/Private/ActorComponentDetails.cpp
@@ -2,6 +2,10 @@
 
 #include "DetailCustomizationsPrivatePCH.h"
 #include "ActorComponentDetails.h"
+#include "ObjectEditorUtils.h"
+#include "IDocumentation.h"
+
+#define LOCTEXT_NAMESPACE "ActorComponentDetails"
 
 TSharedRef<IDetailCustomization> FActorComponentDetails::MakeInstance()
 {
@@ -10,6 +14,8 @@ TSharedRef<IDetailCustomization> FActorComponentDetails::MakeInstance()
 
 void FActorComponentDetails::CustomizeDetails( IDetailLayoutBuilder& DetailBuilder )
 {
+	AddExperimentalWarningCategory(DetailBuilder);
+
 	TSharedPtr<IPropertyHandle> PrimaryTickProperty = DetailBuilder.GetProperty(GET_MEMBER_NAME_CHECKED(UActorComponent, PrimaryComponentTick));
 
 	// Defaults only show tick properties
@@ -24,3 +30,64 @@ void FActorComponentDetails::CustomizeDetails( IDetailLayoutBuilder& DetailBuild
 
 	PrimaryTickProperty->MarkHiddenByCustomization();
 }
+
+void FActorComponentDetails::AddExperimentalWarningCategory( IDetailLayoutBuilder& DetailBuilder )
+{
+	const TArray< TWeakObjectPtr<UObject> >& SelectedObjects = DetailBuilder.GetDetailsView().GetSelectedObjects();
+
+	bool bIsExperimental = false;
+	bool bIsEarlyAccess = false;
+	for (const TWeakObjectPtr<UObject>& SelectedObjectPtr : SelectedObjects)
+	{
+		if (UObject* SelectedObject = SelectedObjectPtr.Get())
+		{
+			bool bObjectClassIsExperimental, bObjectClassIsEarlyAccess;
+			FObjectEditorUtils::GetClassDevelopmentStatus(SelectedObject->GetClass(), bObjectClassIsExperimental, bObjectClassIsEarlyAccess);
+			bIsExperimental |= bObjectClassIsExperimental;
+			bIsEarlyAccess |= bObjectClassIsEarlyAccess;
+		}
+	}
+
+
+	if (bIsExperimental || bIsEarlyAccess)
+	{
+		const FName CategoryName(TEXT("Warning"));
+		const FText CategoryDisplayName = LOCTEXT("WarningCategoryDisplayName", "Warning");
+		const FText WarningText = bIsExperimental ? LOCTEXT("ExperimentalClassWarning", "Uses experimental class") : LOCTEXT("EarlyAccessClassWarning", "Uses early access class");
+		const FText SearchString = WarningText;
+		const FText Tooltip = bIsExperimental ? LOCTEXT("ExperimentalClassTooltip", "Here be dragons!  Uses one or more unsupported 'experimental' classes") : LOCTEXT("EarlyAccessClassTooltip", "Uses one or more 'early access' classes");
+		const FString ExcerptName = bIsExperimental ? TEXT("ComponentUsesExperimentalClass") : TEXT("ComponentUsesEarlyAccessClass");
+		const FSlateBrush* WarningIcon = FEditorStyle::GetBrush(bIsExperimental ? "PropertyEditor.ExperimentalClass" : "PropertyEditor.EarlyAccessClass");
+
+		IDetailCategoryBuilder& WarningCategory = DetailBuilder.EditCategory(CategoryName, CategoryDisplayName, ECategoryPriority::Transform);
+
+		FDetailWidgetRow& WarningRow = WarningCategory.AddCustomRow(SearchString)
+			.WholeRowContent()
+			[
+				SNew(SHorizontalBox)
+				.ToolTip(IDocumentation::Get()->CreateToolTip(Tooltip, nullptr, TEXT("Shared/LevelEditor"), ExcerptName))
+				.Visibility(EVisibility::Visible)
+
+				+ SHorizontalBox::Slot()
+				.VAlign(VAlign_Center)
+				.AutoWidth()
+				.Padding(4.0f, 0.0f, 0.0f, 0.0f)
+				[
+					SNew(SImage)
+					.Image(WarningIcon)
+				]
+
+				+SHorizontalBox::Slot()
+				.VAlign(VAlign_Center)
+				.AutoWidth()
+				.Padding(4.0f, 0.0f, 0.0f, 0.0f)
+				[
+					SNew(STextBlock)
+					.Text(WarningText)
+					.Font(IDetailLayoutBuilder::GetDetailFont())
+				]
+			];
+	}
+}
+
+#undef LOCTEXT_NAMESPACE
diff --git a/Engine/Source/Editor/DetailCustomizations/Private/ActorComponentDetails.h b/Engine/Source/Editor/DetailCustomizations/Private/ActorComponentDetails.h
--- a/Engine/Source/Editor/DetailCustomizations/Private/ActorComponentDetails.h
+++ b/Engine/Source/Editor/DetailCustomizations/Private/ActorComponentDetails.h
@@ -10,4 +10,6 @@ public:
 
 	/** IDetailCustomization interface */
 	virtual void CustomizeDetails( IDetailLayoutBuilder& DetailBuilder ) override;
+
+	void AddExperimentalWarningCategory( IDetailLayoutBuilder& DetailBuilder );
 };
diff --git a/Engine/Source/Editor/Kismet/Private/BlueprintDetailsCustomization.cpp b/Engine/Source/Editor/Kismet/Private/BlueprintDetailsCustomization.cpp
--- a/Engine/Source/Editor/Kismet/Private/BlueprintDetailsCustomization.cpp
+++ b/Engine/Source/Editor/Kismet/Private/BlueprintDetailsCustomization.cpp
@@ -4446,6 +4446,8 @@ void FBlueprintComponentDetails::CustomizeDetails(IDetailLayoutBuilder& DetailLa
 
 	TArray<FSCSEditorTreeNodePtrType> Nodes = Editor->GetSelectedNodes();
 
+	AddExperimentalWarningCategory(DetailLayout, Nodes);
+
 	if (!Nodes.Num())
 	{
 		CachedNodePtr = nullptr;
@@ -4851,6 +4853,63 @@ void FBlueprintComponentDetails::OnSocketSelection( FName SocketName )
 	}
 }
 
+void FBlueprintComponentDetails::AddExperimentalWarningCategory( IDetailLayoutBuilder& DetailBuilder, const TArray<FSCSEditorTreeNodePtrType>& Nodes )
+{
+	bool bIsExperimental = false;
+	bool bIsEarlyAccess = false;
+	for (const FSCSEditorTreeNodePtrType& Node : Nodes)
+	{
+		
+		if (UActorComponent* Component = Node->GetComponentTemplate())
+		{
+			bool bObjectClassIsExperimental, bObjectClassIsEarlyAccess;
+			FObjectEditorUtils::GetClassDevelopmentStatus(Component->GetClass(), bObjectClassIsExperimental, bObjectClassIsEarlyAccess);
+			bIsExperimental |= bObjectClassIsExperimental;
+			bIsEarlyAccess |= bObjectClassIsEarlyAccess;
+		}
+	}
+	
+	if (bIsExperimental || bIsEarlyAccess)
+	{
+		const FName CategoryName(TEXT("Warning"));
+		const FText CategoryDisplayName = LOCTEXT("WarningCategoryDisplayName", "Warning");
+		const FText WarningText = bIsExperimental ? LOCTEXT("ExperimentalClassWarning", "Uses experimental class") : LOCTEXT("EarlyAccessClassWarning", "Uses early access class");
+		const FText SearchString = WarningText;
+		const FText Tooltip = bIsExperimental ? LOCTEXT("ExperimentalClassTooltip", "Here be dragons!  Uses one or more unsupported 'experimental' classes") : LOCTEXT("EarlyAccessClassTooltip", "Uses one or more 'early access' classes");
+		const FString ExcerptName = bIsExperimental ? TEXT("ComponentUsesExperimentalClass") : TEXT("ComponentUsesEarlyAccessClass");
+		const FSlateBrush* WarningIcon = FEditorStyle::GetBrush(bIsExperimental ? "PropertyEditor.ExperimentalClass" : "PropertyEditor.EarlyAccessClass");
+
+		IDetailCategoryBuilder& WarningCategory = DetailBuilder.EditCategory(CategoryName, CategoryDisplayName, ECategoryPriority::Variable);
+
+		FDetailWidgetRow& WarningRow = WarningCategory.AddCustomRow(SearchString)
+			.WholeRowContent()
+			[
+				SNew(SHorizontalBox)
+				.ToolTip(IDocumentation::Get()->CreateToolTip(Tooltip, nullptr, TEXT("Shared/LevelEditor"), ExcerptName))
+				.Visibility(EVisibility::Visible)
+
+				+ SHorizontalBox::Slot()
+				.VAlign(VAlign_Center)
+				.AutoWidth()
+				.Padding(4.0f, 0.0f, 0.0f, 0.0f)
+				[
+					SNew(SImage)
+					.Image(WarningIcon)
+				]
+
+				+SHorizontalBox::Slot()
+				.VAlign(VAlign_Center)
+				.AutoWidth()
+				.Padding(4.0f, 0.0f, 0.0f, 0.0f)
+				[
+					SNew(STextBlock)
+					.Text(WarningText)
+					.Font(IDetailLayoutBuilder::GetDetailFont())
+				]
+			];
+	}
+}
+
 BEGIN_SLATE_FUNCTION_BUILD_OPTIMIZATION
 void FBlueprintGraphNodeDetails::CustomizeDetails( IDetailLayoutBuilder& DetailLayout )
 {
@@ -5234,4 +5293,5 @@ TSharedRef<SWidget> FBlueprintDocumentationDetails::GenerateExcerptList()
 		];
 }
 
+
 #undef LOCTEXT_NAMESPACE
diff --git a/Engine/Source/Editor/Kismet/Private/BlueprintDetailsCustomization.h b/Engine/Source/Editor/Kismet/Private/BlueprintDetailsCustomization.h
--- a/Engine/Source/Editor/Kismet/Private/BlueprintDetailsCustomization.h
+++ b/Engine/Source/Editor/Kismet/Private/BlueprintDetailsCustomization.h
@@ -738,6 +738,8 @@ protected:
 	void OnSocketSelection( FName SocketName );
 
 	void PopulateVariableCategories();
+	
+	void AddExperimentalWarningCategory( IDetailLayoutBuilder& DetailBuilder, const TArray<FSCSEditorTreeNodePtrType>& Nodes );
 
 private:
 	/** Weak reference to the Blueprint editor */
