Fix copy/paste error in comment

diff --git a/Engine/Source/Editor/DetailCustomizations/Private/PostProcessSettingsCustomization.h b/Engine/Source/Editor/DetailCustomizations/Private/PostProcessSettingsCustomization.h
--- a/Engine/Source/Editor/DetailCustomizations/Private/PostProcessSettingsCustomization.h
+++ b/Engine/Source/Editor/DetailCustomizations/Private/PostProcessSettingsCustomization.h
@@ -4,7 +4,7 @@
 
 
 /**
- * Implements a details view customization for the FTimespan structure.
+ * Implements a details view customization for the FPostProcessSettings structure.
  */
 class FPostProcessSettingsCustomization
 	: public IPropertyTypeCustomization
