Improved Blendables, exposed weight on each blendable entry, better blueprint access, Outer=package/object
base for a change to replace all PostProcessSettings with Blendables. See new documentation on blendables.

diff --git a/Engine/Plugins/Blendables/LightPropagationVolume/Source/LightPropagationVolumeEditor/Private/LightPropagationVolumeEditorModule.cpp b/Engine/Plugins/Blendables/LightPropagationVolume/Source/LightPropagationVolumeEditor/Private/LightPropagationVolumeEditorModule.cpp
--- a/Engine/Plugins/Blendables/LightPropagationVolume/Source/LightPropagationVolumeEditor/Private/LightPropagationVolumeEditorModule.cpp
+++ b/Engine/Plugins/Blendables/LightPropagationVolume/Source/LightPropagationVolumeEditor/Private/LightPropagationVolumeEditorModule.cpp
@@ -37,7 +37,11 @@ class FLightPropagationVolumeModule : public IModuleInterface
 		// Register asset types
 		IAssetTools& AssetTools = FModuleManager::LoadModuleChecked<FAssetToolsModule>("AssetTools").Get();
 
- 		EAssetTypeCategories::Type AssetCategoryBit = AssetTools.RegisterAdvancedAssetCategory(FName(TEXT("Blendables")), LOCTEXT("BlendablesAssetCategory", "Blendables"));
+ 		EAssetTypeCategories::Type AssetCategoryBit = AssetTools.FindAdvancedAssetCategory(FName(TEXT("Blendables")));
+		
+		// this category should have been registered already
+		check(AssetCategoryBit != EAssetTypeCategories::Misc);
+
 		RegisterAssetTypeAction(AssetTools, MakeShareable(new FAssetTypeActions_LightPropagationVolumeBlendable(AssetCategoryBit)));
 	}
 
diff --git a/Engine/Plugins/Blendables/LightPropagationVolume/Source/LightPropagationVolumeRuntime/Public/LightPropagationVolumeBlendable.h b/Engine/Plugins/Blendables/LightPropagationVolume/Source/LightPropagationVolumeRuntime/Public/LightPropagationVolumeBlendable.h
--- a/Engine/Plugins/Blendables/LightPropagationVolume/Source/LightPropagationVolumeRuntime/Public/LightPropagationVolumeBlendable.h
+++ b/Engine/Plugins/Blendables/LightPropagationVolume/Source/LightPropagationVolumeRuntime/Public/LightPropagationVolumeBlendable.h
@@ -4,7 +4,6 @@
 
 #include "CoreUObject.h"
 #include "Engine/BlendableInterface.h"
-#include "Engine/DataAsset.h"
 #include "LightPropagationVolumeBlendable.generated.h"
 
 USTRUCT(BlueprintType)
@@ -102,8 +101,7 @@ struct FLightPropagationVolumeSettings
 	UPROPERTY(interp, BlueprintReadWrite, Category=LightPropagationVolume, AdvancedDisplay, meta=(editcondition = "bOverride_LPVSpecularOcclusionIntensity", UIMin = "0", UIMax = "4", DisplayName = "Specular occlusion intensity") )
 	float LPVSpecularOcclusionIntensity;
 
-	// good start values for a new volume, by default no value is overriding
-
+	// good start values for a new volume
 	FLightPropagationVolumeSettings()
 	{
 		// to set all bOverride_.. by default to false
@@ -144,9 +142,9 @@ struct FLightPropagationVolumeSettings
 	}
 };
 
-UCLASS(MinimalAPI, Blueprintable)
-//class ULightPropagationVolumeBlendable : public UDataAsset, public IBlendableInterface
-class ULightPropagationVolumeBlendable : public UDataAsset, public IBlendableInterface
+// BlueprintType to make the object spawnable in blueprint
+UCLASS(MinimalAPI, Blueprintable, BlueprintType)
+class ULightPropagationVolumeBlendable : public UObject, public IBlendableInterface
 {
 	GENERATED_UCLASS_BODY()
 
diff --git a/Engine/Source/Developer/AssetTools/Private/AssetTools.cpp b/Engine/Source/Developer/AssetTools/Private/AssetTools.cpp
--- a/Engine/Source/Developer/AssetTools/Private/AssetTools.cpp
+++ b/Engine/Source/Developer/AssetTools/Private/AssetTools.cpp
@@ -34,6 +34,8 @@ FAssetTools::FAssetTools()
 	AllocatedCategoryBits.Add(TEXT("_BuiltIn_6"), FAdvancedAssetCategory(EAssetTypeCategories::Misc, LOCTEXT("MiscellaneousAssetCategory", "Miscellaneous")));
 	AllocatedCategoryBits.Add(TEXT("_BuiltIn_7"), FAdvancedAssetCategory(EAssetTypeCategories::Gameplay, LOCTEXT("GameplayAssetCategory", "Gameplay")));
 
+	EAssetTypeCategories::Type BlendablesCategoryBit = RegisterAdvancedAssetCategory(FName(TEXT("Blendables")), LOCTEXT("BlendablesAssetCategory", "Blendables"));
+
 	// Register the built-in asset type actions
 	RegisterAssetTypeActions( MakeShareable(new FAssetTypeActions_AnimationAsset) );
 	RegisterAssetTypeActions( MakeShareable(new FAssetTypeActions_AnimBlueprint) );
@@ -66,9 +68,9 @@ FAssetTools::FAssetTools()
 	RegisterAssetTypeActions( MakeShareable(new FAssetTypeActions_InterpData) );
 	RegisterAssetTypeActions( MakeShareable(new FAssetTypeActions_LandscapeLayer) );
 	RegisterAssetTypeActions( MakeShareable(new FAssetTypeActions_LandscapeGrassType));
-	RegisterAssetTypeActions( MakeShareable(new FAssetTypeActions_Material));
+	RegisterAssetTypeActions( MakeShareable(new FAssetTypeActions_Material(BlendablesCategoryBit)));
 	RegisterAssetTypeActions( MakeShareable(new FAssetTypeActions_MaterialFunction) );
-	RegisterAssetTypeActions( MakeShareable(new FAssetTypeActions_MaterialInstanceConstant) );
+	RegisterAssetTypeActions( MakeShareable(new FAssetTypeActions_MaterialInstanceConstant(BlendablesCategoryBit)) );
 	RegisterAssetTypeActions( MakeShareable(new FAssetTypeActions_MaterialInterface) );
 	RegisterAssetTypeActions( MakeShareable(new FAssetTypeActions_MaterialParameterCollection) );
 	RegisterAssetTypeActions( MakeShareable(new FAssetTypeActions_MorphTarget) );
diff --git a/Engine/Source/Developer/AssetTools/Private/AssetTypeActions/AssetTypeActions_Material.h b/Engine/Source/Developer/AssetTools/Private/AssetTypeActions/AssetTypeActions_Material.h
--- a/Engine/Source/Developer/AssetTools/Private/AssetTypeActions/AssetTypeActions_Material.h
+++ b/Engine/Source/Developer/AssetTools/Private/AssetTypeActions/AssetTypeActions_Material.h
@@ -5,11 +5,20 @@
 class FAssetTypeActions_Material : public FAssetTypeActions_MaterialInterface
 {
 public:
+	FAssetTypeActions_Material(EAssetTypeCategories::Type InAssetCategoryBit)
+		: AssetCategoryBit(InAssetCategoryBit)
+	{
+	}
+
 	// IAssetTypeActions Implementation
 	virtual FText GetName() const override { return NSLOCTEXT("AssetTypeActions", "AssetTypeActions_Material", "Material"); }
 	virtual FColor GetTypeColor() const override { return FColor(64,192,64); }
 	virtual UClass* GetSupportedClass() const override { return UMaterial::StaticClass(); }
 	virtual void OpenAssetEditor( const TArray<UObject*>& InObjects, TSharedPtr<class IToolkitHost> EditWithinLevelEditor = TSharedPtr<IToolkitHost>() ) override;
 	virtual bool CanFilter() override { return true; }
-	virtual uint32 GetCategories() override { return FAssetTypeActions_MaterialInterface::GetCategories() | EAssetTypeCategories::Basic; }
+	virtual uint32 GetCategories() override { return FAssetTypeActions_MaterialInterface::GetCategories() | EAssetTypeCategories::Basic | AssetCategoryBit; }
+
+private:
+
+	EAssetTypeCategories::Type AssetCategoryBit;
 };
\ No newline at end of file
diff --git a/Engine/Source/Developer/AssetTools/Private/AssetTypeActions/AssetTypeActions_MaterialInstanceConstant.h b/Engine/Source/Developer/AssetTools/Private/AssetTypeActions/AssetTypeActions_MaterialInstanceConstant.h
--- a/Engine/Source/Developer/AssetTools/Private/AssetTypeActions/AssetTypeActions_MaterialInstanceConstant.h
+++ b/Engine/Source/Developer/AssetTools/Private/AssetTypeActions/AssetTypeActions_MaterialInstanceConstant.h
@@ -7,6 +7,11 @@
 class FAssetTypeActions_MaterialInstanceConstant : public FAssetTypeActions_MaterialInterface
 {
 public:
+	FAssetTypeActions_MaterialInstanceConstant(EAssetTypeCategories::Type InAssetCategoryBit)
+		: AssetCategoryBit(InAssetCategoryBit)
+	{
+	}
+
 	// IAssetTypeActions Implementation
 	virtual FText GetName() const override { return NSLOCTEXT("AssetTypeActions", "AssetTypeActions_MaterialInstanceConstant", "Material Instance"); }
 	virtual FColor GetTypeColor() const override { return FColor(0,128,0); }
@@ -14,8 +19,11 @@ public:
 	virtual void GetActions( const TArray<UObject*>& InObjects, FMenuBuilder& MenuBuilder ) override;
 	virtual void OpenAssetEditor( const TArray<UObject*>& InObjects, TSharedPtr<class IToolkitHost> EditWithinLevelEditor = TSharedPtr<IToolkitHost>() ) override;
 	virtual bool CanFilter() override { return true; }
+	virtual uint32 GetCategories() override { return FAssetTypeActions_MaterialInterface::GetCategories() | AssetCategoryBit; }
 
 private:
 	/** Handler for when FindParent is selected */
 	void ExecuteFindParent(TArray<TWeakObjectPtr<UMaterialInstanceConstant>> Objects);
+
+	EAssetTypeCategories::Type AssetCategoryBit;
 };
\ No newline at end of file
diff --git a/Engine/Source/Editor/DetailCustomizations/DetailCustomizations.Build.cs b/Engine/Source/Editor/DetailCustomizations/DetailCustomizations.Build.cs
--- a/Engine/Source/Editor/DetailCustomizations/DetailCustomizations.Build.cs
+++ b/Engine/Source/Editor/DetailCustomizations/DetailCustomizations.Build.cs
@@ -41,13 +41,14 @@ public class DetailCustomizations : ModuleRules
                 "MediaAssets",
                 "HardwareTargeting",
 				"SharedSettingsWidgets",
+                "LightPropagationVolumeRuntime",
                 "AIModule", 
                 "MeshUtilities",
 				"ConfigEditor",
 			}
 		);
 
-		PrivateIncludePathModuleNames.AddRange(
+        PrivateIncludePathModuleNames.AddRange(
 			new string[] {
                 "Engine",
                 "Media",
diff --git a/Engine/Source/Editor/DetailCustomizations/Private/DetailCustomizations.cpp b/Engine/Source/Editor/DetailCustomizations/Private/DetailCustomizations.cpp
--- a/Engine/Source/Editor/DetailCustomizations/Private/DetailCustomizations.cpp
+++ b/Engine/Source/Editor/DetailCustomizations/Private/DetailCustomizations.cpp
@@ -205,6 +205,7 @@ void FDetailCustomizationsModule::RegisterPropertyTypeCustomizations()
 	RegisterCustomPropertyTypeLayout("DistanceDatum", FOnGetPropertyTypeCustomizationInstance::CreateStatic(&FDistanceDatumStructCustomization::MakeInstance));
 	RegisterCustomPropertyTypeLayout("HierarchicalSimplification", FOnGetPropertyTypeCustomizationInstance::CreateStatic(&FHierarchicalSimplificationCustomizations::MakeInstance));
 	RegisterCustomPropertyTypeLayout("PostProcessSettings", FOnGetPropertyTypeCustomizationInstance::CreateStatic(&FPostProcessSettingsCustomization::MakeInstance));
+	RegisterCustomPropertyTypeLayout("WeightedBlendable", FOnGetPropertyTypeCustomizationInstance::CreateStatic(&FWeightedBlendableCustomization::MakeInstance));
 }
 
 
diff --git a/Engine/Source/Editor/DetailCustomizations/Private/PostProcessSettingsCustomization.cpp b/Engine/Source/Editor/DetailCustomizations/Private/PostProcessSettingsCustomization.cpp
--- a/Engine/Source/Editor/DetailCustomizations/Private/PostProcessSettingsCustomization.cpp
+++ b/Engine/Source/Editor/DetailCustomizations/Private/PostProcessSettingsCustomization.cpp
@@ -3,15 +3,25 @@
 #include "DetailCustomizationsPrivatePCH.h"
 #include "PostProcessSettingsCustomization.h"
 #include "ObjectEditorUtils.h"
+#include "LightPropagationVolumeBlendable.h"
+#include "SWidgetSwitcher.h"
+#include "Classes/Engine/BlendableInterface.h"
 
 #define LOCTEXT_NAMESPACE "PostProcessSettingsCustomization"
 
 void FPostProcessSettingsCustomization::CustomizeChildren( TSharedRef<IPropertyHandle> StructPropertyHandle, class IDetailChildrenBuilder& StructBuilder, IPropertyTypeCustomizationUtils& StructCustomizationUtils )
 {
-	TMap<FString, IDetailGroup*> CategoryNameToGroupMap;
+	TMap<FName, IDetailGroup*> CategoryNameToGroupMap;
 
 	uint32 NumChildren = 0;
 	FPropertyAccess::Result Result = StructPropertyHandle->GetNumChildren(NumChildren);
+
+	UProperty* Prop = StructPropertyHandle->GetProperty();
+	UStructProperty* StructProp = Cast<UStructProperty>(Prop);
+
+	// a category with this name should be one level higher, should be "PostProcessSettings"
+	FName ClassName = StructProp->Struct->GetFName();
+
 	if( Result == FPropertyAccess::Success && NumChildren > 0)
 	{
 		for( uint32 ChildIndex = 0; ChildIndex < NumChildren; ++ChildIndex )
@@ -20,23 +30,32 @@ void FPostProcessSettingsCustomization::CustomizeChildren( TSharedRef<IPropertyH
 			if( ChildHandle.IsValid() && ChildHandle->GetProperty() )
 			{
 				UProperty* Property = ChildHandle->GetProperty();
-				FString Category = FObjectEditorUtils::GetCategory( Property );
-				IDetailGroup* Group = CategoryNameToGroupMap.FindRef( Category );
+				FName Category = FObjectEditorUtils::GetCategoryFName( Property );
 
-				if( !Group )
+				if(Category == ClassName)
 				{
-					Group = &StructBuilder.AddChildGroup( FName(*Category), FText::FromString( FName::NameToDisplayString( Category, false ) ) );
-					CategoryNameToGroupMap.Add( Category, Group );
+					// Some elements should be outside of the categories, they need to specify the struct nama as category
+					StructBuilder.AddChildProperty( ChildHandle.ToSharedRef() );
 				}
+				else
+				{
+					IDetailGroup* Group = CategoryNameToGroupMap.FindRef( Category );
+
+					if( !Group )
+					{
+						Group = &StructBuilder.AddChildGroup( Category, FText::FromString( FName::NameToDisplayString( Category.ToString(), false ) ) );
+						CategoryNameToGroupMap.Add( Category, Group );
+					}
 				
-				Group->AddPropertyRow( ChildHandle.ToSharedRef() );
+					Group->AddPropertyRow( ChildHandle.ToSharedRef() );
+				}
 			}
 		}
 	}
-	
 }
 
 
+
 void FPostProcessSettingsCustomization::CustomizeHeader( TSharedRef<IPropertyHandle> StructPropertyHandle, class FDetailWidgetRow& HeaderRow, IPropertyTypeCustomizationUtils& StructCustomizationUtils )
 {
 	HeaderRow.NameContent()
@@ -50,5 +69,224 @@ void FPostProcessSettingsCustomization::CustomizeHeader( TSharedRef<IPropertyHan
 	];
 }
 
+void FWeightedBlendableCustomization::AddDirectAsset(TSharedRef<IPropertyHandle> StructPropertyHandle, UPackage* Package, TSharedPtr<IPropertyHandle> Weight, TSharedPtr<IPropertyHandle> Value)
+{
+	Weight->SetValue(1.0f);
+
+	{
+		TArray<UObject*> Objects;
+		StructPropertyHandle->GetOuterObjects(Objects);
+
+		TArray<FString> Values;
+
+		for(TArray<UObject*>::TConstIterator It = Objects.CreateConstIterator(); It; It++)
+		{
+			UObject* Obj = *It;
+			const UObject* NewObj = NewObject<ULightPropagationVolumeBlendable>(Obj);
+
+			FString Str = NewObj->GetPathName();
+
+			Values.Add(Str);
+		}
+
+		Value->SetPerObjectValues(Values);
+	}
+}
+
+void FWeightedBlendableCustomization::AddIndirectAsset(TSharedPtr<IPropertyHandle> Weight)
+{
+	Weight->SetValue(1.0f);
+}
+
+EVisibility FWeightedBlendableCustomization::IsWeightVisible(TSharedPtr<IPropertyHandle> Weight) const
+{
+	float WeightValue = 1.0f;
+	
+	Weight->GetValue(WeightValue);
+
+	return (WeightValue >= 0) ? EVisibility::Visible : EVisibility::Hidden;
+}
+
+FText FWeightedBlendableCustomization::GetDirectAssetName(TSharedPtr<IPropertyHandle> Value) const
+{
+	UObject* RefObject = 0;
+	
+	Value->GetValue(RefObject);
+
+	check(RefObject);
+
+	return FText::FromString(RefObject->GetFullName());
+}
+
+FReply FWeightedBlendableCustomization::JumpToDirectAsset(TSharedPtr<IPropertyHandle> Value)
+{
+	UObject* RefObject = 0;
+	
+	Value->GetValue(RefObject);
+
+	FAssetEditorManager::Get().OpenEditorForAsset(RefObject);
+
+	return FReply::Handled();
+}
+
+TSharedRef<SWidget> FWeightedBlendableCustomization::GenerateContentWidget(TSharedRef<IPropertyHandle> StructPropertyHandle, UPackage* Package, TSharedPtr<IPropertyHandle> Weight, TSharedPtr<IPropertyHandle> Value)
+{
+	FMenuBuilder MenuBuilder(true, NULL);
+	{
+/*
+		for(TObjectIterator<UClass> It; It; ++It)
+		{
+			if(It->IsChildOf(UFactory::StaticClass()))
+			if(It->IsChildOf(UBlendableInterface::StaticClass()))
+*/
+		// todo: iterate through all UFactories that create UClasses with UBlendableInterface
+		{
+			FUIAction Direct2(FExecuteAction::CreateSP(this, &FWeightedBlendableCustomization::AddDirectAsset, StructPropertyHandle, Package, Weight, Value));
+			MenuBuilder.AddMenuEntry(LOCTEXT("Blendable_DirectAsset2", "Light Propagation Volume Blendable"), 
+				LOCTEXT("Blendable_DirectAsset2h", "references Light Propagation Volume Blendable (owned by the containing object)"), FSlateIcon(), Direct2);
+		}
+
+		MenuBuilder.AddMenuSeparator();
+
+		FUIAction Indirect(FExecuteAction::CreateSP(this, &FWeightedBlendableCustomization::AddIndirectAsset, Weight));
+		MenuBuilder.AddMenuEntry(LOCTEXT("Blendable_IndirectAsset", "Asset reference"), 
+			LOCTEXT("Blendable_IndirectAsseth", "reference a Blendable asset (owned by a content package), e.g. material with Post Process domain"), FSlateIcon(), Indirect);
+	}
+	
+
+	TSharedRef<SWidgetSwitcher> Switcher = SNew(SWidgetSwitcher)
+			.WidgetIndex(this, &FWeightedBlendableCustomization::ComputeSwitcherIndex, StructPropertyHandle, Package, Weight, Value);
+
+	Switcher->AddSlot()
+		[
+			SNew(SComboButton)
+			.ButtonContent()
+			[
+				SNew(STextBlock)
+				.Text(LOCTEXT("Blendable_ChooseElement", "Choose"))
+			]
+			.ContentPadding(FMargin(6.0, 2.0))
+			.MenuContent()
+			[
+				MenuBuilder.MakeWidget()
+			]
+		];
+
+	Switcher->AddSlot()
+		[
+			SNew(SButton)
+			.ContentPadding(FMargin(0,0))
+			.Text(this, &FWeightedBlendableCustomization::GetDirectAssetName, Value)
+			.OnClicked(this, &FWeightedBlendableCustomization::JumpToDirectAsset, Value)
+		];
+
+	Switcher->AddSlot()
+		[
+			SNew(SObjectPropertyEntryBox)
+			.PropertyHandle(Value)
+		];
+
+	return Switcher;
+}
+
+int32 FWeightedBlendableCustomization::ComputeSwitcherIndex(TSharedRef<IPropertyHandle> StructPropertyHandle, UPackage* Package, TSharedPtr<IPropertyHandle> Weight, TSharedPtr<IPropertyHandle> Value) const
+{
+	float WeightValue = 1.0f;
+	UObject* RefObject = 0;
+	
+	Weight->GetValue(WeightValue);
+	Value->GetValue(RefObject);
+
+	if(RefObject)
+	{
+		UPackage* PropPackage = RefObject->GetOutermost();
+
+		return (PropPackage == Package) ? 1 : 2;
+	}
+	else
+	{
+		return (WeightValue < 0.0f) ? 0 : 2;
+	}
+}
+
+void FWeightedBlendableCustomization::CustomizeChildren( TSharedRef<class IPropertyHandle> StructPropertyHandle, class IDetailChildrenBuilder& StructBuilder, IPropertyTypeCustomizationUtils& StructCustomizationUtils )
+{
+	// we don't have children but this is a pure virtual so we need to override
+}
+
+void FWeightedBlendableCustomization::CustomizeHeader( TSharedRef<IPropertyHandle> StructPropertyHandle, class FDetailWidgetRow& HeaderRow, IPropertyTypeCustomizationUtils& StructCustomizationUtils )
+{
+	TSharedPtr<IPropertyHandle> SharedWeightProp;
+	{
+		TSharedPtr<IPropertyHandle> ChildHandle = StructPropertyHandle->GetChildHandle(FName(TEXT("Weight")));
+		if (ChildHandle.IsValid() && ChildHandle->GetProperty())
+		{
+			SharedWeightProp = ChildHandle;
+		}
+	}
+		
+	TSharedPtr<IPropertyHandle> SharedValueProp;
+	{
+		TSharedPtr<IPropertyHandle> ChildHandle = StructPropertyHandle->GetChildHandle(FName(TEXT("Object")));
+		if (ChildHandle.IsValid() && ChildHandle->GetProperty())
+		{
+			SharedValueProp = ChildHandle;
+		}
+	}
+
+	float WeightValue = 1.0f;
+	UObject* RefObject = 0;
+	
+	SharedWeightProp->GetValue(WeightValue);
+	SharedValueProp->GetValue(RefObject);
+
+	UPackage* StructPackage = 0;
+	{
+		const TSharedPtr<IPropertyHandle> ParentHandle = StructPropertyHandle->GetParentHandle();
+		TArray<UObject*> Objects;
+		StructPropertyHandle->GetOuterObjects(Objects);
+
+		for (TArray<UObject*>::TConstIterator It = Objects.CreateConstIterator(); It; It++)
+		{
+			UObject* ref = *It;
+
+			if(StructPackage)
+			{
+				// does this mean we have to deal with multiple levels? The code here is not ready for that
+				check(StructPackage == ref->GetOutermost());
+			}
+
+			StructPackage = ref->GetOutermost();
+		}
+	}
+
+	HeaderRow.NameContent()
+	[
+		SNew(SHorizontalBox)
+		.Visibility(this, &FWeightedBlendableCustomization::IsWeightVisible, SharedWeightProp)
+		+SHorizontalBox::Slot()
+		[
+			SNew(SBox)
+			.MinDesiredWidth(44.0f)
+			.MaxDesiredWidth(44.0f)		
+			[
+				SharedWeightProp->CreatePropertyValueWidget()
+			]
+		]
+	];
+
+	HeaderRow.ValueContent()
+	[
+		SNew(SHorizontalBox)
+		+ SHorizontalBox::Slot()
+		[
+			GenerateContentWidget(StructPropertyHandle, StructPackage, SharedWeightProp, SharedValueProp)
+		]
+	];
+}
+
+
+
+
 
 #undef LOCTEXT_NAMESPACE
diff --git a/Engine/Source/Editor/DetailCustomizations/Private/PostProcessSettingsCustomization.h b/Engine/Source/Editor/DetailCustomizations/Private/PostProcessSettingsCustomization.h
--- a/Engine/Source/Editor/DetailCustomizations/Private/PostProcessSettingsCustomization.h
+++ b/Engine/Source/Editor/DetailCustomizations/Private/PostProcessSettingsCustomization.h
@@ -28,3 +28,49 @@ public:
 	virtual void CustomizeChildren( TSharedRef<class IPropertyHandle> StructPropertyHandle, class IDetailChildrenBuilder& StructBuilder, IPropertyTypeCustomizationUtils& StructCustomizationUtils ) override;
 	virtual void CustomizeHeader( TSharedRef<class IPropertyHandle> StructPropertyHandle, class FDetailWidgetRow& HeaderRow, IPropertyTypeCustomizationUtils& StructCustomizationUtils ) override;
 };
+
+
+/**
+ * Implements a details view customization for the FBlendablesWrapper structure.
+ */
+class FWeightedBlendableCustomization
+	: public IPropertyTypeCustomization
+{
+public:
+
+	/**
+	 * Creates an instance of this class.
+	 *
+	 * @return The new instance.
+	 */
+	static TSharedRef<IPropertyTypeCustomization> MakeInstance( )
+	{
+		return MakeShareable(new FWeightedBlendableCustomization());
+	}
+
+public:
+
+	// IPropertyTypeCustomization interface
+
+	virtual void CustomizeChildren( TSharedRef<class IPropertyHandle> StructPropertyHandle, class IDetailChildrenBuilder& StructBuilder, IPropertyTypeCustomizationUtils& StructCustomizationUtils ) override;
+	virtual void CustomizeHeader( TSharedRef<class IPropertyHandle> StructPropertyHandle, class FDetailWidgetRow& HeaderRow, IPropertyTypeCustomizationUtils& StructCustomizationUtils ) override;
+
+	TSharedRef<SWidget> GenerateContentWidget(TSharedRef<IPropertyHandle> StructPropertyHandle, UPackage* Package, TSharedPtr<IPropertyHandle> Weight, TSharedPtr<IPropertyHandle> Value);
+
+	// "Outer" is the object that has the blendables container
+	void AddDirectAsset(TSharedRef<IPropertyHandle> StructPropertyHandle, UPackage* Package, TSharedPtr<IPropertyHandle> Weight, TSharedPtr<IPropertyHandle> Value);
+	// "Outer" is a package
+	void AddIndirectAsset(TSharedPtr<IPropertyHandle> Weight);
+
+	// open the details panel for this object
+	FReply JumpToDirectAsset(TSharedPtr<IPropertyHandle> Value);
+	
+	// to get the blendable type as string
+	FText GetDirectAssetName(TSharedPtr<IPropertyHandle> Value) const;
+
+	// @return 0:choose 1:direct, 2:indirect
+	int32 ComputeSwitcherIndex(TSharedRef<IPropertyHandle> StructPropertyHandle, UPackage* Package, TSharedPtr<IPropertyHandle> Weight, TSharedPtr<IPropertyHandle> Value) const;
+
+	// the weight is only visible if the user choose the type (or the reference)
+	EVisibility IsWeightVisible(TSharedPtr<IPropertyHandle> Weight) const;
+};
diff --git a/Engine/Source/Runtime/Engine/Classes/Camera/CameraComponent.h b/Engine/Source/Runtime/Engine/Classes/Camera/CameraComponent.h
--- a/Engine/Source/Runtime/Engine/Classes/Camera/CameraComponent.h
+++ b/Engine/Source/Runtime/Engine/Classes/Camera/CameraComponent.h
@@ -73,6 +73,7 @@ class UCameraComponent : public USceneComponent
 	// UObject interface
 	ENGINE_API virtual void PostEditChangeProperty(FPropertyChangedEvent& PropertyChangedEvent) override;
 #endif
+	virtual void Serialize(FArchive& Ar) override;
 	// End of UObject interface
 
 	/**
@@ -82,6 +83,10 @@ class UCameraComponent : public USceneComponent
 	UFUNCTION(BlueprintCallable, Category=Camera)
 	ENGINE_API virtual void GetCameraView(float DeltaTime, FMinimalViewInfo& DesiredView);
 
+	/** Adds an Blendable (implements IBlendableInterface) to the array of Blendables (if it doesn't exist) and update the weight */
+	UFUNCTION(BlueprintCallable, Category="Rendering")
+	ENGINE_API void AddOrUpdateBlendable(TScriptInterface<IBlendableInterface> InBlendableObject, float InWeight = 1.0f) { PostProcessSettings.AddBlendable(InBlendableObject, InWeight); }
+
 protected:
 #if WITH_EDITORONLY_DATA
 	// The frustum component used to show visually where the camera field of view is
diff --git a/Engine/Source/Runtime/Engine/Classes/Components/PostProcessComponent.h b/Engine/Source/Runtime/Engine/Classes/Components/PostProcessComponent.h
--- a/Engine/Source/Runtime/Engine/Classes/Components/PostProcessComponent.h
+++ b/Engine/Source/Runtime/Engine/Classes/Components/PostProcessComponent.h
@@ -55,6 +55,10 @@ class UPostProcessComponent : public USceneComponent, public IInterface_PostProc
 	}
 	// End IInterface_PostProcessVolume Interface
 
+	/** Adds an Blendable (implements IBlendableInterface) to the array of Blendables (if it doesn't exist) and update the weight */
+	UFUNCTION(BlueprintCallable, Category="Rendering")
+	ENGINE_API void AddOrUpdateBlendable(TScriptInterface<IBlendableInterface> InBlendableObject, float InWeight = 1.0f) { Settings.AddBlendable(InBlendableObject, InWeight); }
+
 protected:
 
 	virtual void OnRegister() override;
diff --git a/Engine/Source/Runtime/Engine/Classes/Components/SceneCaptureComponent2D.h b/Engine/Source/Runtime/Engine/Classes/Components/SceneCaptureComponent2D.h
--- a/Engine/Source/Runtime/Engine/Classes/Components/SceneCaptureComponent2D.h
+++ b/Engine/Source/Runtime/Engine/Classes/Components/SceneCaptureComponent2D.h
@@ -58,8 +58,15 @@ public:
 #if WITH_EDITOR
 	virtual void PostEditChangeProperty(FPropertyChangedEvent& PropertyChangedEvent) override;
 #endif // WITH_EDITOR
+	
+	virtual void Serialize(FArchive& Ar);
+
 	// End UObject Interface
 
+	/** Adds an Blendable (implements IBlendableInterface) to the array of Blendables (if it doesn't exist) and update the weight */
+	UFUNCTION(BlueprintCallable, Category="Rendering")
+	ENGINE_API void AddOrUpdateBlendable(TScriptInterface<IBlendableInterface> InBlendableObject, float InWeight = 1.0f) { PostProcessSettings.AddBlendable(InBlendableObject, InWeight); }
+
 	/** Render the scene to the texture */
 	UFUNCTION(BlueprintCallable,Category = "Rendering|SceneCapture")
 	ENGINE_API void UpdateContent();
diff --git a/Engine/Source/Runtime/Engine/Classes/Engine/PostProcessVolume.h b/Engine/Source/Runtime/Engine/Classes/Engine/PostProcessVolume.h
--- a/Engine/Source/Runtime/Engine/Classes/Engine/PostProcessVolume.h
+++ b/Engine/Source/Runtime/Engine/Classes/Engine/PostProcessVolume.h
@@ -20,7 +20,7 @@ class APostProcessVolume : public AVolume, public IInterface_PostProcessVolume
 	GENERATED_UCLASS_BODY()
 
 	/** Post process settings to use for this volume. */
-	UPROPERTY(interp, Category=PostProcessVolume, meta=(DisplayName="Post Process Settings"))
+	UPROPERTY(interp, Category=PostProcessVolume, meta=(ShowOnlyInnerProperties))
 	struct FPostProcessSettings Settings;
 
 	/**
@@ -74,8 +74,13 @@ public:
 #if WITH_EDITOR
 	virtual void PostEditChangeProperty(FPropertyChangedEvent& PropertyChangedEvent) override;
 	virtual bool CanEditChange(const UProperty* InProperty) const override;
+	virtual void Serialize(FArchive& Ar);
 #endif // WITH_EDITOR
 	// End UObject interface
+
+	/** Adds an Blendable (implements IBlendableInterface) to the array of Blendables (if it doesn't exist) and update the weight */
+	UFUNCTION(BlueprintCallable, Category="Rendering")
+	ENGINE_API void AddOrUpdateBlendable(TScriptInterface<IBlendableInterface> InBlendableObject, float InWeight = 1.0f) { Settings.AddBlendable(InBlendableObject, InWeight); }
 };
 
 
diff --git a/Engine/Source/Runtime/Engine/Classes/Engine/Scene.h b/Engine/Source/Runtime/Engine/Classes/Engine/Scene.h
--- a/Engine/Source/Runtime/Engine/Classes/Engine/Scene.h
+++ b/Engine/Source/Runtime/Engine/Classes/Engine/Scene.h
@@ -5,6 +5,7 @@
 //=============================================================================
 
 #pragma once
+#include "BlendableInterface.h"
 #include "Scene.generated.h"
 
 /** Used by FPostProcessSettings Depth of Fields */
@@ -27,6 +28,46 @@ enum EAntiAliasingMethod
 	AAM_MAX,
 };
 
+USTRUCT()
+struct FWeightedBlendable
+{
+	GENERATED_USTRUCT_BODY()
+
+	/** 0:no effect .. 1:full effect */
+	UPROPERTY(interp, BlueprintReadWrite, Category=FWeightedBlendable, meta=(UIMin = "0.0", UIMax = "1.0", Delta = "0.01"))
+	float Weight;
+
+	/** should be of the IBlendableInterface* type but UProperties cannot express that */
+	UPROPERTY(EditAnywhere, BlueprintReadWrite, Category=FWeightedBlendable, meta=( AllowedClasses="BlendableInterface", Keywords="PostProcess" ))
+	UObject* Object;
+
+	// default constructor
+	FWeightedBlendable()
+		: Weight(-1)
+		, Object(0)
+	{
+	}
+
+	// constructor
+	// @param InWeight -1 is used to hide the weight and show the "Choose" UI, 0:no effect .. 1:full effect
+	FWeightedBlendable(float InWeight, UObject* InObject)
+		: Weight(InWeight)
+		, Object(InObject)
+	{
+	}
+};
+
+// for easier detail customization, needed?
+USTRUCT()
+struct FWeightedBlendables
+{
+	GENERATED_USTRUCT_BODY()
+
+	UPROPERTY(EditAnywhere, BlueprintReadWrite, Category="PostProcessSettings", meta=( Keywords="PostProcess" ))
+	TArray<FWeightedBlendable> Array;
+};
+
+
 /** To be able to use struct PostProcessSettings. */
 // Each property consists of a bool to enable it (by default off),
 // the variable declaration and further down the default value for it.
@@ -894,10 +935,69 @@ struct FPostProcessSettings
 	
 	/**
 	 * Allows custom post process materials to be defined, using a MaterialInstance with the same Material as its parent to allow blending.
-	 * Make sure you use the "PostProcess" domain type. This can be used for any UObject object implementing the IBlendableInterface (e.g. could be used to fade weather settings).
+	 * For materials this needs to be the "PostProcess" domain type. This can be used for any UObject object implementing the IBlendableInterface (e.g. could be used to fade weather settings).
 	 */
-	UPROPERTY(EditAnywhere, BlueprintReadWrite, Category="Misc", meta=( AllowedClasses="BlendableInterface", Keywords="PostProcess" ))
-	TArray<UObject*> Blendables;
+	UPROPERTY(EditAnywhere, Category="PostProcessSettings", meta=( Keywords="PostProcess", DisplayName = "Blendables" ))
+	FWeightedBlendables WeightedBlendables;
+
+	// for backwards compatibility
+	UPROPERTY()
+	TArray<UObject*> Blendables_DEPRECATED;
+
+	// for backwards compatibility
+	void OnAfterLoad()
+	{
+		for(int32 i = 0, count = Blendables_DEPRECATED.Num(); i < count; ++i)
+		{
+			if(Blendables_DEPRECATED[i])
+			{
+				FWeightedBlendable Element(1.0f, Blendables_DEPRECATED[i]);
+				WeightedBlendables.Array.Add(Element);
+			}
+		}
+		Blendables_DEPRECATED.Empty();
+	}
+
+	// Adds an Blendable (implements IBlendableInterface) to the array of Blendables (if it doesn't exist) and update the weight
+	// @param InBlendableObject silently ignores if no object is referenced
+	// @param 0..1 InWeight, values outside of the range get clampled later in the pipeline
+	void AddBlendable(TScriptInterface<IBlendableInterface> InBlendableObject, float InWeight)
+	{
+		// update weight, if the Blendable is already in the array
+		if(UObject* Object = InBlendableObject.GetObject())
+		{
+			for (int32 i = 0, count = WeightedBlendables.Array.Num(); i < count; ++i)
+			{
+				if (WeightedBlendables.Array[i].Object == Object)
+				{
+					WeightedBlendables.Array[i].Weight = InWeight;
+					// We assumes we only have one
+					return;
+				}
+			}
+
+			// add in the end
+			WeightedBlendables.Array.Add(FWeightedBlendable(InWeight, Object));
+		}
+	}
+
+	// removes one or multiple blendables from the array
+	void RemoveBlendable(TScriptInterface<IBlendableInterface> InBlendableObject)
+	{
+		if(UObject* Object = InBlendableObject.GetObject())
+		{
+			for (int32 i = 0, count = WeightedBlendables.Array.Num(); i < count; ++i)
+			{
+				if (WeightedBlendables.Array[i].Object == Object)
+				{
+					// this might remove multiple
+					WeightedBlendables.Array.RemoveAt(i);
+					--i;
+					--count;
+				}
+			}
+		}
+	}
 
 	// good start values for a new volume, by default no value is overriding
 	FPostProcessSettings()
diff --git a/Engine/Source/Runtime/Engine/Classes/Interfaces/Interface_PostProcessVolume.h b/Engine/Source/Runtime/Engine/Classes/Interfaces/Interface_PostProcessVolume.h
--- a/Engine/Source/Runtime/Engine/Classes/Interfaces/Interface_PostProcessVolume.h
+++ b/Engine/Source/Runtime/Engine/Classes/Interfaces/Interface_PostProcessVolume.h
@@ -3,12 +3,6 @@
 #pragma once
 #include "Interface_PostProcessVolume.generated.h"
 
-/** Interface for general PostProcessVolume access **/
-UINTERFACE()
-class UInterface_PostProcessVolume : public UInterface
-{
-	GENERATED_UINTERFACE_BODY()
-};
 
 struct FPostProcessVolumeProperties
 {
@@ -20,6 +14,13 @@ struct FPostProcessVolumeProperties
 	bool bIsUnbound;
 };
 
+/** Interface for general PostProcessVolume access **/
+UINTERFACE(MinimalAPI, meta = (CannotImplementInterfaceInBlueprint))
+class UInterface_PostProcessVolume : public UInterface
+{
+	GENERATED_UINTERFACE_BODY()
+};
+
 class IInterface_PostProcessVolume
 {
 	GENERATED_IINTERFACE_BODY()
diff --git a/Engine/Source/Runtime/Engine/Private/Camera/CameraComponent.cpp b/Engine/Source/Runtime/Engine/Private/Camera/CameraComponent.cpp
--- a/Engine/Source/Runtime/Engine/Private/Camera/CameraComponent.cpp
+++ b/Engine/Source/Runtime/Engine/Private/Camera/CameraComponent.cpp
@@ -191,6 +191,16 @@ void UCameraComponent::PostEditChangeProperty(FPropertyChangedEvent& PropertyCha
 }
 #endif
 
+void UCameraComponent::Serialize(FArchive& Ar)
+{
+	Super::Serialize(Ar);
+
+	if(Ar.IsLoading())
+	{
+		PostProcessSettings.OnAfterLoad();
+	}
+}
+
 void UCameraComponent::GetCameraView(float DeltaTime, FMinimalViewInfo& DesiredView)
 {
 	if (bUsePawnControlRotation)
diff --git a/Engine/Source/Runtime/Engine/Private/Components/SceneCaptureComponent.cpp b/Engine/Source/Runtime/Engine/Private/Components/SceneCaptureComponent.cpp
--- a/Engine/Source/Runtime/Engine/Private/Components/SceneCaptureComponent.cpp
+++ b/Engine/Source/Runtime/Engine/Private/Components/SceneCaptureComponent.cpp
@@ -352,6 +352,15 @@ void USceneCaptureComponent2D::PostEditChangeProperty(FPropertyChangedEvent& Pro
 }
 #endif // WITH_EDITOR
 
+void USceneCaptureComponent2D::Serialize(FArchive& Ar)
+{
+	Super::Serialize(Ar);
+
+	if(Ar.IsLoading())
+	{
+		PostProcessSettings.OnAfterLoad();
+	}
+}
 
 // -----------------------------------------------
 
diff --git a/Engine/Source/Runtime/Engine/Private/PostProcessVolume.cpp b/Engine/Source/Runtime/Engine/Private/PostProcessVolume.cpp
--- a/Engine/Source/Runtime/Engine/Private/PostProcessVolume.cpp
+++ b/Engine/Source/Runtime/Engine/Private/PostProcessVolume.cpp
@@ -22,6 +22,16 @@ bool APostProcessVolume::EncompassesPoint(FVector Point, float SphereRadius/*=0.
 	return Super::EncompassesPoint(Point, SphereRadius, OutDistanceToPoint);
 }
 
+void APostProcessVolume::Serialize(FArchive& Ar)
+{
+	Super::Serialize(Ar);
+
+	if(Ar.IsLoading())
+	{
+		Settings.OnAfterLoad();
+	}
+}
+
 #if WITH_EDITOR
 void APostProcessVolume::PostEditChangeProperty(FPropertyChangedEvent& PropertyChangedEvent)
 {
@@ -32,16 +42,14 @@ void APostProcessVolume::PostEditChangeProperty(FPropertyChangedEvent& PropertyC
 	if(PropertyChangedEvent.Property && PropertyChangedEvent.Property->GetFName() == NAME_Blendables)
 	{
 		// remove unsupported types
+		uint32 Count = Settings.WeightedBlendables.Array.Num();
+		for(uint32 i = 0; i < Count; ++i)
 		{
-			uint32 Count = Settings.Blendables.Num();
-			for(uint32 i = 0; i < Count; ++i)
-			{
-				UObject* Obj = Settings.Blendables[i];
+			UObject* Obj = Settings.WeightedBlendables.Array[i].Object;
 
-				if(!Cast<IBlendableInterface>(Obj))
-				{
-					Settings.Blendables[i] = 0;
-				}
+			if(!Cast<IBlendableInterface>(Obj))
+			{
+				Settings.WeightedBlendables.Array[i] = FWeightedBlendable();
 			}
 		}
 	}
diff --git a/Engine/Source/Runtime/Engine/Private/SceneView.cpp b/Engine/Source/Runtime/Engine/Private/SceneView.cpp
--- a/Engine/Source/Runtime/Engine/Private/SceneView.cpp
+++ b/Engine/Source/Runtime/Engine/Private/SceneView.cpp
@@ -948,12 +948,12 @@ void FSceneView::OverridePostProcessSettings(const FPostProcessSettings& Src, fl
 
 	// Blendable objects
 	{
-		uint32 Count = Src.Blendables.Num();
+		uint32 Count = Src.WeightedBlendables.Array.Num();
 
 		for(uint32 i = 0; i < Count; ++i)
 		{
-			// todo: optimize
-			UObject* Object = Src.Blendables[i];
+			float LocalWeight = Src.WeightedBlendables.Array[i].Weight;
+			UObject* Object = Src.WeightedBlendables.Array[i].Object;
 
 			if(!Object)
 			{
@@ -967,7 +967,7 @@ void FSceneView::OverridePostProcessSettings(const FPostProcessSettings& Src, fl
 				continue;
 			}
 
-			BlendableInterface->OverrideBlendableSettings(*this, Weight);
+			BlendableInterface->OverrideBlendableSettings(*this, LocalWeight * Weight);
 		}
 	}
 }
