#ue4 - Fix manually entered rotations in the editor sometimes not resulting in the correct object rotation (mainly when editing Pitch).

- UpdateComponentToWorldWithParent() no longer changes RelativeRotation. The normalization can be done in InternalSetWorldRotation instead.
- Component transform editing now makes sure to only edit the single component (Yaw, Pitch, or Roll) that has been changed, and makes sure Quat<->Rotation conversion does not mess with the end result.

UE-17482

diff --git a/Engine/Source/Editor/DetailCustomizations/Private/ComponentTransformDetails.cpp b/Engine/Source/Editor/DetailCustomizations/Private/ComponentTransformDetails.cpp
--- a/Engine/Source/Editor/DetailCustomizations/Private/ComponentTransformDetails.cpp
+++ b/Engine/Source/Editor/DetailCustomizations/Private/ComponentTransformDetails.cpp
@@ -1119,7 +1119,8 @@ void FComponentTransformDetails::OnSetRotation( float NewValue, bool bCommitted,
 				{
 					const bool bIsEditingTemplateObject = Object->IsTemplate();
 
-					FRotator& RelativeRotation = (bEditingRotationInUI && !bIsEditingTemplateObject) ? ObjectToRelativeRotationMap.FindOrAdd(Object) : SceneComponent->RelativeRotation;
+					FRotator CurrentComponentRotation = SceneComponent->RelativeRotation; // Intentionally make a copy, we don't want the FRotator& below to directly edit the component's values!
+					FRotator& RelativeRotation = (bEditingRotationInUI && !bIsEditingTemplateObject) ? ObjectToRelativeRotationMap.FindOrAdd(Object) : CurrentComponentRotation;
 					FRotator OldRelativeRotation = RelativeRotation;
 
 					float& ValueToChange = Axis == 0 ? RelativeRotation.Roll : Axis == 1 ? RelativeRotation.Pitch : RelativeRotation.Yaw;
@@ -1169,20 +1170,15 @@ void FComponentTransformDetails::OnSetRotation( float NewValue, bool bCommitted,
 
 						ValueToChange = NewValue;
 
-						if(!bIsEditingTemplateObject)
+						if( !bIsEditingTemplateObject && SelectedActorInfo.NumSelected != 0 )
 						{
-							if( SelectedActorInfo.NumSelected == 0 )
-							{
-								// HACK: Set directly if no actors are selected since this causes Rot->Quat->Rot conversion issues
-								// (recalculates relative rotation from quat which can give an equivalent but different value than the user typed)
-								SceneComponent->RelativeRotation = RelativeRotation;
-							}
-							else
-							{
-								SceneComponent->SetRelativeRotation( RelativeRotation );
-							}
+							SceneComponent->SetRelativeRotation( RelativeRotation );
 						}
 
+						// HACK: Set directly since to avoid Rot->Quat->Rot conversion issues
+						// (functions recalculate relative rotation from quat which can give an equivalent but different value than the user typed)
+						SceneComponent->RelativeRotation = RelativeRotation;
+
 						AActor* EditedActor = Cast<AActor>( Object );
 						if( !EditedActor && SceneComponent )
 						{
diff --git a/Engine/Source/Runtime/Engine/Private/Components/SceneComponent.cpp b/Engine/Source/Runtime/Engine/Private/Components/SceneComponent.cpp
--- a/Engine/Source/Runtime/Engine/Private/Components/SceneComponent.cpp
+++ b/Engine/Source/Runtime/Engine/Private/Components/SceneComponent.cpp
@@ -342,13 +342,6 @@ void USceneComponent::UpdateComponentToWorldWithParent(USceneComponent* Parent,
 	// Calculate the new ComponentToWorld transform
 	const FTransform RelativeTransform(RelativeRotationQuat, RelativeLocation, RelativeScale3D);
 	FTransform NewTransform = CalcNewComponentToWorld(RelativeTransform, Parent);
-	
-	// Minimize accumulation of errors after many composed transforms.
-	NewTransform.NormalizeRotation();
-
-	// Update RelativeRotation and cache (not done earlier because normalization can change the rotation)
-	const FQuat RelativeQuatNormalized = (Parent ? RelativeRotationQuat.GetNormalized() : NewTransform.GetRotation());
-	RelativeRotation = RelativeRotationCache.NormalizedQuatToRotator(RelativeQuatNormalized);
 
 #if !(UE_BUILD_SHIPPING || UE_BUILD_TEST)
 	ensureOnce(NewTransform.IsValid());
@@ -1852,7 +1845,8 @@ bool USceneComponent::InternalSetWorldLocationAndRotation(FVector NewLocation, c
 	if (!NewLocation.Equals(RelativeLocation) || !NewRotationQuat.Equals(RelativeRotationCache.RotatorToQuat_ReadOnly(RelativeRotation), 1e-6f))
 	{
 		RelativeLocation = NewLocation;
-		UpdateComponentToWorldWithParent(AttachParent, bNoPhysics, NewRotationQuat, Teleport);
+		RelativeRotation = RelativeRotationCache.QuatToRotator(NewRotationQuat); // Normalizes rotator, if this is a new rotation. Then we'll use it below.
+		UpdateComponentToWorldWithParent(AttachParent, bNoPhysics, RelativeRotationCache.RotatorToQuat(RelativeRotation), Teleport);
 		return true;
 	}
 
