Framework: Exposed the ability to set the tick group for actors and components in Blueprints [UE-17651]

diff --git a/Engine/Source/Editor/DetailCustomizations/Private/ActorComponentDetails.cpp b/Engine/Source/Editor/DetailCustomizations/Private/ActorComponentDetails.cpp
--- a/Engine/Source/Editor/DetailCustomizations/Private/ActorComponentDetails.cpp
+++ b/Engine/Source/Editor/DetailCustomizations/Private/ActorComponentDetails.cpp
@@ -26,6 +26,7 @@ void FActorComponentDetails::CustomizeDetails( IDetailLayoutBuilder& DetailBuild
 		TickCategory.AddProperty(PrimaryTickProperty->GetChildHandle(GET_MEMBER_NAME_CHECKED(FTickFunction, bStartWithTickEnabled)));
 		TickCategory.AddProperty(PrimaryTickProperty->GetChildHandle(GET_MEMBER_NAME_CHECKED(FTickFunction, bTickEvenWhenPaused)), EPropertyLocation::Advanced);
 		TickCategory.AddProperty(PrimaryTickProperty->GetChildHandle(GET_MEMBER_NAME_CHECKED(FTickFunction, bAllowTickOnDedicatedServer)), EPropertyLocation::Advanced);
+		TickCategory.AddProperty(PrimaryTickProperty->GetChildHandle(GET_MEMBER_NAME_CHECKED(FTickFunction, TickGroup)), EPropertyLocation::Advanced);
 	}
 
 	PrimaryTickProperty->MarkHiddenByCustomization();
diff --git a/Engine/Source/Editor/DetailCustomizations/Private/ActorDetails.cpp b/Engine/Source/Editor/DetailCustomizations/Private/ActorDetails.cpp
--- a/Engine/Source/Editor/DetailCustomizations/Private/ActorDetails.cpp
+++ b/Engine/Source/Editor/DetailCustomizations/Private/ActorDetails.cpp
@@ -137,6 +137,7 @@ void FActorDetails::CustomizeDetails( IDetailLayoutBuilder& DetailLayout )
 		TickCategory.AddProperty(PrimaryTickProperty->GetChildHandle(GET_MEMBER_NAME_CHECKED(FTickFunction, bStartWithTickEnabled)));
 		TickCategory.AddProperty(PrimaryTickProperty->GetChildHandle(GET_MEMBER_NAME_CHECKED(FTickFunction, bTickEvenWhenPaused)), EPropertyLocation::Advanced);
 		TickCategory.AddProperty(PrimaryTickProperty->GetChildHandle(GET_MEMBER_NAME_CHECKED(FTickFunction, bAllowTickOnDedicatedServer)), EPropertyLocation::Advanced);
+		TickCategory.AddProperty(PrimaryTickProperty->GetChildHandle(GET_MEMBER_NAME_CHECKED(FTickFunction, TickGroup)), EPropertyLocation::Advanced);
 	}
 
 	PrimaryTickProperty->MarkHiddenByCustomization();
diff --git a/Engine/Source/Runtime/Engine/Classes/Components/ActorComponent.h b/Engine/Source/Runtime/Engine/Classes/Components/ActorComponent.h
--- a/Engine/Source/Runtime/Engine/Classes/Components/ActorComponent.h
+++ b/Engine/Source/Runtime/Engine/Classes/Components/ActorComponent.h
@@ -640,6 +640,7 @@ public:
 	void ReregisterComponent();
 
 	/** Changes the ticking group for this component */
+	UFUNCTION(BlueprintCallable, Category="Utilities", meta=(Keywords = "dependency"))
 	void SetTickGroup(ETickingGroup NewTickGroup);
 
 	/** Make this component tick after PrerequisiteActor */
diff --git a/Engine/Source/Runtime/Engine/Classes/Engine/EngineBaseTypes.h b/Engine/Source/Runtime/Engine/Classes/Engine/EngineBaseTypes.h
--- a/Engine/Source/Runtime/Engine/Classes/Engine/EngineBaseTypes.h
+++ b/Engine/Source/Runtime/Engine/Classes/Engine/EngineBaseTypes.h
@@ -38,7 +38,7 @@ enum ELevelTick
 };
 
 /** Determines which ticking group a tick function belongs to. */
-UENUM()
+UENUM(BlueprintType)
 enum ETickingGroup
 {
 	/** Any item that needs to be executed before physics simulation starts. */
@@ -141,7 +141,7 @@ public:
 	 * @see ETickingGroup 
 	 * @see FTickFunction::AddPrerequisite()
 	 */
-	UPROPERTY()
+	UPROPERTY(EditDefaultsOnly, Category="Tick", AdvancedDisplay)
 	TEnumAsByte<enum ETickingGroup> TickGroup;
 
 protected:
diff --git a/Engine/Source/Runtime/Engine/Classes/GameFramework/Actor.h b/Engine/Source/Runtime/Engine/Classes/GameFramework/Actor.h
--- a/Engine/Source/Runtime/Engine/Classes/GameFramework/Actor.h
+++ b/Engine/Source/Runtime/Engine/Classes/GameFramework/Actor.h
@@ -2043,6 +2043,7 @@ public:
 	 * Sets the ticking group for this actor.
 	 * @param NewTickGroup the new value to assign
 	 */
+	UFUNCTION(BlueprintCallable, Category="Utilities", meta=(Keywords = "dependency"))
 	void SetTickGroup(ETickingGroup NewTickGroup);
 
 	/** Called once this actor has been deleted */
