Added a detail customization for text properties in preparation for future text property changes. No longer uses import/export text on text properties.

diff --git a/Engine/Source/Editor/DetailCustomizations/Private/DetailCustomizations.cpp b/Engine/Source/Editor/DetailCustomizations/Private/DetailCustomizations.cpp
--- a/Engine/Source/Editor/DetailCustomizations/Private/DetailCustomizations.cpp
+++ b/Engine/Source/Editor/DetailCustomizations/Private/DetailCustomizations.cpp
@@ -102,6 +102,7 @@
 #include "HierarchicalSimplificationCustomizations.h"
 #include "PostProcessSettingsCustomization.h"
 #include "ConfigEditorPropertyDetails.h"
+#include "TextCustomization.h"
 
 IMPLEMENT_MODULE( FDetailCustomizationsModule, DetailCustomizations );
 
@@ -189,6 +190,7 @@ void FDetailCustomizationsModule::RegisterPropertyTypeCustomizations()
 	RegisterCustomPropertyTypeLayout("NavigationLink", FOnGetPropertyTypeCustomizationInstance::CreateStatic(&FNavLinkStructCustomization::MakeInstance));
 	RegisterCustomPropertyTypeLayout("NavigationSegmentLink", FOnGetPropertyTypeCustomizationInstance::CreateStatic(&FNavLinkStructCustomization::MakeInstance));
 	RegisterCustomPropertyTypeLayout("Margin", FOnGetPropertyTypeCustomizationInstance::CreateStatic(&FMarginStructCustomization::MakeInstance));
+	RegisterCustomPropertyTypeLayout("TextProperty", FOnGetPropertyTypeCustomizationInstance::CreateStatic(&FTextCustomization::MakeInstance));
 	RegisterCustomPropertyTypeLayout("DirectoryPath", FOnGetPropertyTypeCustomizationInstance::CreateStatic(&FDirectoryPathStructCustomization::MakeInstance));
 	RegisterCustomPropertyTypeLayout("FilePath", FOnGetPropertyTypeCustomizationInstance::CreateStatic(&FFilePathStructCustomization::MakeInstance));
 	RegisterCustomPropertyTypeLayout("IOSBuildResourceDirectory", FOnGetPropertyTypeCustomizationInstance::CreateStatic(&FDirectoryPathStructCustomization::MakeInstance));
diff --git a/Engine/Source/Editor/DetailCustomizations/Private/TextCustomization.cpp b/Engine/Source/Editor/DetailCustomizations/Private/TextCustomization.cpp
new file mode 100644
--- /dev/null
+++ b/Engine/Source/Editor/DetailCustomizations/Private/TextCustomization.cpp
@@ -0,0 +1,93 @@
+// Copyright 1998-2015 Epic Games, Inc. All Rights Reserved.
+
+#include "DetailCustomizationsPrivatePCH.h"
+#include "TextCustomization.h"
+
+namespace
+{
+	class STextPropertyWidget : public SCompoundWidget
+	{
+		SLATE_BEGIN_ARGS(STextPropertyWidget) {}
+		SLATE_END_ARGS()
+
+	public:
+		void Construct(const FArguments& Arguments, const TSharedRef<IPropertyHandle>& InPropertyHandle);
+
+	private:
+		TSharedPtr<IPropertyHandle> PropertyHandle;
+	};
+
+	void STextPropertyWidget::Construct(const FArguments& Arguments, const TSharedRef<IPropertyHandle>& InPropertyHandle)
+	{
+		PropertyHandle = InPropertyHandle;
+
+		const auto& GetTextValue = [this]() -> FText
+		{
+			FText TextValue;
+
+			TArray<const void*> RawData;
+			PropertyHandle->AccessRawData(RawData);
+			if (RawData.Num() == 1)
+			{
+				const FText* RawDatum = reinterpret_cast<const FText*>(RawData.Top());
+				if (RawDatum)
+				{
+					TextValue = *RawDatum;
+				}
+			}
+			else if(RawData.Num() > 1)
+			{
+				TextValue = NSLOCTEXT("PropertyEditor", "MultipleValues", "Multiple Values");
+			}
+			else
+			{
+				TextValue = FText::GetEmpty();
+			}
+
+			return TextValue;
+		};
+
+		const auto& OnTextCommitted = [this](const FText& NewText, ETextCommit::Type CommitInfo)
+		{
+			PropertyHandle->NotifyPreChange();
+
+			TArray<void*> RawData;
+			PropertyHandle->AccessRawData(RawData);
+			for (void* const RawDatum : RawData)
+			{
+				FText& PropertyValue = *(reinterpret_cast<FText* const>(RawDatum));
+
+				// FText::FromString on the result of FText::ToString is intentional. For now, we want to nuke any namespace/key info and let it get regenerated from scratch,
+				// rather than risk adopting whatever came through some chain of calls. This will be replaced when preserving of identity is implemented.
+				PropertyValue = FText::FromString(NewText.ToString());
+			}
+			
+			PropertyHandle->NotifyPostChange();
+		};
+
+		ChildSlot
+			[
+				SNew(SEditableTextBox)
+				.Text_Lambda(GetTextValue)
+				.OnTextCommitted_Lambda(OnTextCommitted)
+			];
+	}
+}
+
+void FTextCustomization::CustomizeHeader( TSharedRef<class IPropertyHandle> InPropertyHandle, class FDetailWidgetRow& HeaderRow, IPropertyTypeCustomizationUtils& PropertyTypeCustomizationUtils )
+{
+	HeaderRow.FilterString(InPropertyHandle->GetPropertyDisplayName())
+		.NameContent()
+		[
+			InPropertyHandle->CreatePropertyNameWidget()
+		]
+		.ValueContent()
+		[
+			SNew(STextPropertyWidget, InPropertyHandle)
+		];
+}
+
+void FTextCustomization::CustomizeChildren( TSharedRef<class IPropertyHandle> InPropertyHandle, class IDetailChildrenBuilder& StructBuilder, IPropertyTypeCustomizationUtils& PropertyTypeCustomizationUtils )
+{
+
+}
\ No newline at end of file
diff --git a/Engine/Source/Editor/DetailCustomizations/Private/TextCustomization.h b/Engine/Source/Editor/DetailCustomizations/Private/TextCustomization.h
new file mode 100644
--- /dev/null
+++ b/Engine/Source/Editor/DetailCustomizations/Private/TextCustomization.h
@@ -0,0 +1,15 @@
+// Copyright 1998-2015 Epic Games, Inc. All Rights Reserved.
+#pragma once
+
+class FTextCustomization : public IPropertyTypeCustomization
+{
+public:
+	static TSharedRef<IPropertyTypeCustomization> MakeInstance() 
+	{
+		return MakeShareable( new FTextCustomization );
+	}
+
+	/** IPropertyTypeCustomization interface */
+	virtual void CustomizeHeader( TSharedRef<class IPropertyHandle> InPropertyHandle, class FDetailWidgetRow& HeaderRow, IPropertyTypeCustomizationUtils& PropertyTypeCustomizationUtils ) override;
+	virtual void CustomizeChildren( TSharedRef<class IPropertyHandle> InPropertyHandle, class IDetailChildrenBuilder& StructBuilder, IPropertyTypeCustomizationUtils& PropertyTypeCustomizationUtils ) override;
+};
