Renamed "Tick" in UActorComponent to "Component Tick" and Actor tick to "Actor Tick" to differentiate between the two when they are both shown in a details panel

diff --git a/Engine/Source/Editor/DetailCustomizations/Private/ActorComponentDetails.cpp b/Engine/Source/Editor/DetailCustomizations/Private/ActorComponentDetails.cpp
--- a/Engine/Source/Editor/DetailCustomizations/Private/ActorComponentDetails.cpp
+++ b/Engine/Source/Editor/DetailCustomizations/Private/ActorComponentDetails.cpp
@@ -21,7 +21,7 @@ void FActorComponentDetails::CustomizeDetails( IDetailLayoutBuilder& DetailBuild
 	// Defaults only show tick properties
 	if (DetailBuilder.GetDetailsView().HasClassDefaultObject())
 	{
-		IDetailCategoryBuilder& TickCategory = DetailBuilder.EditCategory("Tick");
+		IDetailCategoryBuilder& TickCategory = DetailBuilder.EditCategory("ComponentTick");
 
 		TickCategory.AddProperty(PrimaryTickProperty->GetChildHandle(GET_MEMBER_NAME_CHECKED(FTickFunction, bStartWithTickEnabled)));
 		TickCategory.AddProperty(PrimaryTickProperty->GetChildHandle(GET_MEMBER_NAME_CHECKED(FTickFunction, bTickEvenWhenPaused)), EPropertyLocation::Advanced);
diff --git a/Engine/Source/Editor/DetailCustomizations/Private/ActorDetails.cpp b/Engine/Source/Editor/DetailCustomizations/Private/ActorDetails.cpp
--- a/Engine/Source/Editor/DetailCustomizations/Private/ActorDetails.cpp
+++ b/Engine/Source/Editor/DetailCustomizations/Private/ActorDetails.cpp
@@ -132,7 +132,8 @@ void FActorDetails::CustomizeDetails( IDetailLayoutBuilder& DetailLayout )
 	// Defaults only show tick properties
 	if (DetailLayout.GetDetailsView().HasClassDefaultObject() && !HideCategories.Contains(TEXT("Tick")))
 	{
-		IDetailCategoryBuilder& TickCategory = DetailLayout.EditCategory("Tick");
+		// Note: the category is renamed to differentiate between 
+		IDetailCategoryBuilder& TickCategory = DetailLayout.EditCategory("Tick", LOCTEXT("TickCategoryName", "Actor Tick") );
 
 		TickCategory.AddProperty(PrimaryTickProperty->GetChildHandle(GET_MEMBER_NAME_CHECKED(FTickFunction, bStartWithTickEnabled)));
 		TickCategory.AddProperty(PrimaryTickProperty->GetChildHandle(GET_MEMBER_NAME_CHECKED(FTickFunction, bTickEvenWhenPaused)), EPropertyLocation::Advanced);
diff --git a/Engine/Source/Runtime/Engine/Classes/Components/ActorComponent.h b/Engine/Source/Runtime/Engine/Classes/Components/ActorComponent.h
--- a/Engine/Source/Runtime/Engine/Classes/Components/ActorComponent.h
+++ b/Engine/Source/Runtime/Engine/Classes/Components/ActorComponent.h
@@ -57,7 +57,7 @@ public:
 	virtual void GetLifetimeReplicatedProps(TArray<FLifetimeProperty>& OutLifetimeProps) const override;
 
 	/** Main tick function for the Actor */
-	UPROPERTY(EditDefaultsOnly, Category="Tick")
+	UPROPERTY(EditDefaultsOnly, Category="ComponentTick")
 	struct FActorComponentTickFunction PrimaryComponentTick;
 
 	/** Array of tags that can be used for grouping and categorizing. Can also be accessed from scripting. */
