Fix Recapture Scene button stretching the length of the details panel

diff --git a/Engine/Source/Editor/DetailCustomizations/Private/SkyLightComponentDetails.cpp b/Engine/Source/Editor/DetailCustomizations/Private/SkyLightComponentDetails.cpp
--- a/Engine/Source/Editor/DetailCustomizations/Private/SkyLightComponentDetails.cpp
+++ b/Engine/Source/Editor/DetailCustomizations/Private/SkyLightComponentDetails.cpp
@@ -49,20 +49,27 @@ void FSkyLightComponentDetails::CustomizeDetails( IDetailLayoutBuilder& DetailLa
 
 	DetailLayout.EditCategory( "SkyLight" )
 	.AddCustomRow( NSLOCTEXT("SkyLightDetails", "UpdateSkyLight", "Recapture Scene") )
-	[
-		SNew(SHorizontalBox)
-		+SHorizontalBox::Slot()
-		.FillWidth(1.f)
-		.Padding(10,5)
+		.NameContent()
+		[
+			SNew( STextBlock )
+			.Font( IDetailLayoutBuilder::GetDetailFont() )
+			.Text( NSLOCTEXT("SkyLightDetails", "UpdateSkyLight", "Recapture Scene") )
+		]
+		.ValueContent()
+		.MaxDesiredWidth(125.f)
+		.MinDesiredWidth(125.f)
 		[
 			SNew(SButton)
-			.ContentPadding(3)
+			.ContentPadding(2)
 			.VAlign(VAlign_Center)
 			.HAlign(HAlign_Center)
 			.OnClicked( this, &FSkyLightComponentDetails::OnUpdateSkyCapture )
-			.Text( NSLOCTEXT("SkyLightDetails", "UpdateSkyCapture", "Recapture Scene") )
-		]
-	];
+			[
+				SNew( STextBlock )
+				.Font( IDetailLayoutBuilder::GetDetailFont() )
+				.Text( NSLOCTEXT("SkyLightDetails", "UpdateSkyCapture", "Recapture") )
+			]
+		];
 }
 
 FReply FSkyLightComponentDetails::OnUpdateSkyCapture()
