Prevent text of Create External Curve box from being cut off

diff --git a/Engine/Source/Editor/DetailCustomizations/Private/CurveStructCustomization.cpp b/Engine/Source/Editor/DetailCustomizations/Private/CurveStructCustomization.cpp
--- a/Engine/Source/Editor/DetailCustomizations/Private/CurveStructCustomization.cpp
+++ b/Engine/Source/Editor/DetailCustomizations/Private/CurveStructCustomization.cpp
@@ -173,12 +173,17 @@ void FCurveStructCustomization::CustomizeChildren( TSharedRef<IPropertyHandle> I
 					+SVerticalBox::Slot()
 					.AutoHeight()
 					[
-						SNew(SButton)
-						.HAlign(HAlign_Center)
-						.Text( LOCTEXT( "CreateAssetButton", "Create External Curve" ) )
-						.ToolTipText(LOCTEXT( "CreateAssetTooltip", "Create a new CurveFloat asset from this curve") )
-						.OnClicked(this, &FCurveStructCustomization::OnCreateButtonClicked)
-						.IsEnabled(this, &FCurveStructCustomization::IsCreateButtonEnabled)
+						SNew(SHorizontalBox)
+						+SHorizontalBox::Slot()
+						.AutoWidth()
+						[
+							SNew(SButton)
+							.HAlign(HAlign_Center)
+							.Text( LOCTEXT( "CreateAssetButton", "Create External Curve" ) )
+							.ToolTipText(LOCTEXT( "CreateAssetTooltip", "Create a new CurveFloat asset from this curve") )
+							.OnClicked(this, &FCurveStructCustomization::OnCreateButtonClicked)
+							.IsEnabled(this, &FCurveStructCustomization::IsCreateButtonEnabled)
+						]
 					]
 				];
 		}
