[UE-17722] Fix propagation of material swaps when editing *native* component templates (DSOs) owned by the CDO.

change summary:
- modified FComponentMaterialCategory::OnMaterialChanged() to include instances of default subobjects (DSOs) when propagating to instances of the BPGC default object.

diff --git a/Engine/Source/Editor/DetailCustomizations/Private/ComponentMaterialCategory.cpp b/Engine/Source/Editor/DetailCustomizations/Private/ComponentMaterialCategory.cpp
--- a/Engine/Source/Editor/DetailCustomizations/Private/ComponentMaterialCategory.cpp
+++ b/Engine/Source/Editor/DetailCustomizations/Private/ComponentMaterialCategory.cpp
@@ -313,19 +313,35 @@ void FComponentMaterialCategory::OnMaterialChanged( UMaterialInterface* NewMater
 					NotifyHook->NotifyPostChange( PropertyChangedEvent, MaterialProperty );
 				}
 
-				TArray<UObject*> ArchetypeInstances;
-				if( CurrentComponent->IsTemplate() && !FApp::IsGame() )
+				// Propagate material change to instances of the edited component template
+				if( !FApp::IsGame() )
 				{
-					// Propagate material change to instances of the edited component template
-					CurrentComponent->GetArchetypeInstances(ArchetypeInstances);
-					for( auto ArchetypeInstance : ArchetypeInstances )
+					TArray<UObject*> ComponentArchetypeInstances;
+					if( CurrentComponent->HasAnyFlags(RF_ArchetypeObject) )
 					{
-						CurrentComponent = CastChecked<UActorComponent>( ArchetypeInstance );
+						CurrentComponent->GetArchetypeInstances(ComponentArchetypeInstances);
+					}
+					else if( UObject* Outer = CurrentComponent->GetOuter() )
+					{
+						TArray<UObject*> OuterArchetypeInstances;
+						Outer->GetArchetypeInstances(OuterArchetypeInstances);
+						for( auto OuterArchetypeInstance : OuterArchetypeInstances )
+						{
+							if( UObject* ArchetypeInstance = static_cast<UObject*>(FindObjectWithOuter(OuterArchetypeInstance, CurrentComponent->GetClass(), CurrentComponent->GetFName())) )
+							{
+								ComponentArchetypeInstances.Add(ArchetypeInstance);
+							}
+						}
+					}
+
+					for( auto ComponentArchetypeInstance : ComponentArchetypeInstances )
+					{
+						CurrentComponent = CastChecked<UActorComponent>( ComponentArchetypeInstance );
 						if( CurrentComponent->IsA<ULandscapeComponent>() )
 						{
-							ArchetypeInstance = CastChecked<ULandscapeComponent>(CurrentComponent)->GetLandscapeProxy();
+							ComponentArchetypeInstance = CastChecked<ULandscapeComponent>(CurrentComponent)->GetLandscapeProxy();
 						}
-						
+
 						// Reset the navigation update lock if necessary
 						UWorld* PreviousWorld = World;
 						World = CurrentComponent->GetWorld();
@@ -334,11 +350,11 @@ void FComponentMaterialCategory::OnMaterialChanged( UMaterialInterface* NewMater
 							NavUpdateLock = MakeShareable( new FNavigationLockContext(World, ENavigationLockReason::MaterialUpdate) );
 						}
 
-						ArchetypeInstance->PreEditChange( MaterialProperty );
+						ComponentArchetypeInstance->PreEditChange( MaterialProperty );
 
 						SwapMaterialLambda( CurrentComponent, It.GetMaterialIndex(), NewMaterial );
 
-						ArchetypeInstance->PostEditChangeProperty( PropertyChangedEvent );
+						ComponentArchetypeInstance->PostEditChangeProperty( PropertyChangedEvent );
 					}
 				}
 			}
