unified navigation agent selection
#ue4

diff --git a/Engine/Source/Editor/DetailCustomizations/Private/DetailCustomizations.cpp b/Engine/Source/Editor/DetailCustomizations/Private/DetailCustomizations.cpp
--- a/Engine/Source/Editor/DetailCustomizations/Private/DetailCustomizations.cpp
+++ b/Engine/Source/Editor/DetailCustomizations/Private/DetailCustomizations.cpp
@@ -65,6 +65,7 @@
 #include "SlateColorCustomization.h"
 #include "CurveStructCustomization.h"
 #include "NavLinkStructCustomization.h"
+#include "NavAgentSelectorCustomization.h"
 #include "DirectoryPathStructCustomization.h"
 #include "FilePathStructCustomization.h"
 #include "DeviceProfileDetails.h"
@@ -189,6 +190,7 @@ void FDetailCustomizationsModule::RegisterPropertyTypeCustomizations()
 	RegisterCustomPropertyTypeLayout("EnvNamedValue", FOnGetPropertyTypeCustomizationInstance::CreateStatic(&FEnvQueryParamInstanceCustomization::MakeInstance));
 	RegisterCustomPropertyTypeLayout("NavigationLink", FOnGetPropertyTypeCustomizationInstance::CreateStatic(&FNavLinkStructCustomization::MakeInstance));
 	RegisterCustomPropertyTypeLayout("NavigationSegmentLink", FOnGetPropertyTypeCustomizationInstance::CreateStatic(&FNavLinkStructCustomization::MakeInstance));
+	RegisterCustomPropertyTypeLayout("NavAgentSelector", FOnGetPropertyTypeCustomizationInstance::CreateStatic(&FNavAgentSelectorCustomization::MakeInstance));
 	RegisterCustomPropertyTypeLayout("Margin", FOnGetPropertyTypeCustomizationInstance::CreateStatic(&FMarginStructCustomization::MakeInstance));
 	RegisterCustomPropertyTypeLayout("TextProperty", FOnGetPropertyTypeCustomizationInstance::CreateStatic(&FTextCustomization::MakeInstance));
 	RegisterCustomPropertyTypeLayout("DirectoryPath", FOnGetPropertyTypeCustomizationInstance::CreateStatic(&FDirectoryPathStructCustomization::MakeInstance));
diff --git a/Engine/Source/Editor/DetailCustomizations/Private/NavAgentSelectorCustomization.cpp b/Engine/Source/Editor/DetailCustomizations/Private/NavAgentSelectorCustomization.cpp
new file mode 100644
--- /dev/null
+++ b/Engine/Source/Editor/DetailCustomizations/Private/NavAgentSelectorCustomization.cpp
@@ -0,0 +1,135 @@
+// Copyright 1998-2015 Epic Games, Inc. All Rights Reserved.
+
+#include "DetailCustomizationsPrivatePCH.h"
+#include "NavAgentSelectorCustomization.h"
+#include "AI/Navigation/NavigationTypes.h"
+#include "AI/Navigation/NavigationSystem.h"
+
+#define LOCTEXT_NAMESPACE "FNavAgentSelectorCustomization"
+
+TSharedRef<IPropertyTypeCustomization> FNavAgentSelectorCustomization::MakeInstance()
+{
+	return MakeShareable(new FNavAgentSelectorCustomization);
+}
+
+void FNavAgentSelectorCustomization::CustomizeHeader(TSharedRef<IPropertyHandle> StructPropertyHandle, class FDetailWidgetRow& HeaderRow, IPropertyTypeCustomizationUtils& StructCustomizationUtils)
+{
+	StructHandle = StructPropertyHandle;
+	OnAgentStateChanged();
+
+	HeaderRow.NameContent()
+	[
+		StructPropertyHandle->CreatePropertyNameWidget()
+	]
+	.ValueContent()
+	.MaxDesiredWidth(400.0f)
+	.VAlign(VAlign_Center)
+	[
+		SNew(STextBlock)
+		.Text(this, &FNavAgentSelectorCustomization::GetSupportedDesc)
+		.Font(StructCustomizationUtils.GetRegularFont())
+	];
+}
+
+void FNavAgentSelectorCustomization::CustomizeChildren(TSharedRef<class IPropertyHandle> StructPropertyHandle, class IDetailChildrenBuilder& StructBuilder, IPropertyTypeCustomizationUtils& StructCustomizationUtils)
+{
+	uint32 NumChildren = 0;
+	StructPropertyHandle->GetNumChildren(NumChildren);
+
+	FString AgentPrefix("bSupportsAgent");
+	const UNavigationSystem* DefNavSys = (UNavigationSystem*)(UNavigationSystem::StaticClass()->GetDefaultObject());
+	const int32 NumAgents = FMath::Min(DefNavSys->GetSupportedAgents().Num(), 16);
+
+	for (uint32 Idx = 0; Idx < NumChildren; Idx++)
+	{
+		TSharedPtr<IPropertyHandle> PropHandle = StructPropertyHandle->GetChildHandle(Idx);
+		if (PropHandle->GetProperty() && PropHandle->GetProperty()->GetName().StartsWith(AgentPrefix))
+		{
+			PropHandle->SetOnPropertyValueChanged(FSimpleDelegate::CreateSP(this, &FNavAgentSelectorCustomization::OnAgentStateChanged));
+
+			int32 AgentIdx = -1;
+			TTypeFromString<int32>::FromString(AgentIdx, *(PropHandle->GetProperty()->GetName().Mid(AgentPrefix.Len()) ));
+
+			if (AgentIdx >= 0 && AgentIdx < NumAgents)
+			{
+				FText PropName = FText::FromName(DefNavSys->GetSupportedAgents()[AgentIdx].Name);
+				StructBuilder.AddChildContent(PropName)
+					.NameContent()
+					[
+						SNew(STextBlock)
+						.Text(PropName)
+						.Font(StructCustomizationUtils.GetRegularFont())
+					]
+					.ValueContent()
+					[
+						PropHandle->CreatePropertyValueWidget()
+					];
+			}
+
+			continue;
+		}
+
+		StructBuilder.AddChildProperty(PropHandle.ToSharedRef());
+	}
+}
+
+void FNavAgentSelectorCustomization::OnAgentStateChanged()
+{
+	const UNavigationSystem* DefNavSys = (UNavigationSystem*)(UNavigationSystem::StaticClass()->GetDefaultObject());
+	const int32 NumAgents = FMath::Min(DefNavSys->GetSupportedAgents().Num(), 16);
+
+	uint32 NumChildren = 0;
+	StructHandle->GetNumChildren(NumChildren);
+
+	int32 NumSupported = 0;
+	int32 FirstSupportedIdx = -1;
+
+	FString AgentPrefix("bSupportsAgent");
+	for (uint32 Idx = 0; Idx < NumChildren; Idx++)
+	{
+		TSharedPtr<IPropertyHandle> PropHandle = StructHandle->GetChildHandle(Idx);
+		if (PropHandle->GetProperty() && PropHandle->GetProperty()->GetName().StartsWith(AgentPrefix))
+		{
+			bool bSupportsAgent = false;
+			FPropertyAccess::Result Result = PropHandle->GetValue(bSupportsAgent);
+			if (Result == FPropertyAccess::Success && bSupportsAgent)
+			{
+				int32 AgentIdx = -1;
+				TTypeFromString<int32>::FromString(AgentIdx, *(PropHandle->GetProperty()->GetName().Mid(AgentPrefix.Len())));
+
+				if (AgentIdx >= 0 && AgentIdx < NumAgents)
+				{
+					NumSupported++;
+					if (FirstSupportedIdx < 0)
+					{
+						FirstSupportedIdx = AgentIdx;
+					}
+				}
+			}
+		}
+	}
+
+	if (NumSupported == NumAgents)
+	{
+		SupportedDesc = LOCTEXT("AllAgents", "all");
+	}
+	else if (NumSupported == 0)
+	{
+		SupportedDesc = LOCTEXT("NoAgents", "none");
+	}
+	else if (NumSupported == 1)
+	{
+		SupportedDesc = FText::FromName(DefNavSys->GetSupportedAgents()[FirstSupportedIdx].Name);
+	}
+	else
+	{
+		SupportedDesc = FText::Format(FText::FromString("{0}, ..."), FText::FromName(DefNavSys->GetSupportedAgents()[FirstSupportedIdx].Name));
+	}
+}
+
+FText FNavAgentSelectorCustomization::GetSupportedDesc() const
+{
+	return SupportedDesc;
+}
+
+#undef LOCTEXT_NAMESPACE
diff --git a/Engine/Source/Editor/DetailCustomizations/Private/NavAgentSelectorCustomization.h b/Engine/Source/Editor/DetailCustomizations/Private/NavAgentSelectorCustomization.h
new file mode 100644
--- /dev/null
+++ b/Engine/Source/Editor/DetailCustomizations/Private/NavAgentSelectorCustomization.h
@@ -0,0 +1,28 @@
+// Copyright 1998-2015 Epic Games, Inc. All Rights Reserved.
+
+#pragma once
+
+class FNavAgentSelectorCustomization : public IPropertyTypeCustomization
+{
+public:
+
+	// IPropertyTypeCustomization interface
+
+	virtual void CustomizeHeader( TSharedRef<IPropertyHandle> StructPropertyHandle, class FDetailWidgetRow& HeaderRow, IPropertyTypeCustomizationUtils& StructCustomizationUtils ) override;
+	virtual void CustomizeChildren( TSharedRef<class IPropertyHandle> StructPropertyHandle, class IDetailChildrenBuilder& StructBuilder, IPropertyTypeCustomizationUtils& StructCustomizationUtils ) override;
+
+public:
+
+	/**
+	 * Creates a new instance.
+	 *
+	 * @return A new struct customization for Keys.
+	 */
+	static TSharedRef<IPropertyTypeCustomization> MakeInstance();
+
+	void OnAgentStateChanged();
+	FText GetSupportedDesc() const;
+
+	TSharedPtr<IPropertyHandle> StructHandle;
+	FText SupportedDesc;
+};
diff --git a/Engine/Source/Editor/DetailCustomizations/Private/NavLinkStructCustomization.cpp b/Engine/Source/Editor/DetailCustomizations/Private/NavLinkStructCustomization.cpp
--- a/Engine/Source/Editor/DetailCustomizations/Private/NavLinkStructCustomization.cpp
+++ b/Engine/Source/Editor/DetailCustomizations/Private/NavLinkStructCustomization.cpp
@@ -3,9 +3,6 @@
 #include "DetailCustomizationsPrivatePCH.h"
 #include "NavLinkStructCustomization.h"
 #include "AI/Navigation/NavLinkDefinition.h"
-#include "AI/Navigation/NavigationSystem.h"
-
-#define LOCTEXT_NAMESPACE "FNavLinkStructCustomization"
 
 TSharedRef<IPropertyTypeCustomization> FNavLinkStructCustomization::MakeInstance( )
 {
@@ -35,44 +32,14 @@ void FNavLinkStructCustomization::CustomizeHeader( TSharedRef<IPropertyHandle> S
 	];
 }
 
-void FNavLinkStructCustomization::CustomizeChildren( TSharedRef<class IPropertyHandle> StructPropertyHandle, class IDetailChildrenBuilder& StructBuilder, IPropertyTypeCustomizationUtils& StructCustomizationUtils )
-{
-	uint32 NumChildren = 0;
-	StructPropertyHandle->GetNumChildren(NumChildren);
-
-	FString AgentPrefix("bSupportsAgent");
-	const UNavigationSystem* DefNavSys = (UNavigationSystem*)(UNavigationSystem::StaticClass()->GetDefaultObject());
-	const int32 NumAgents = FMath::Min(DefNavSys->GetSupportedAgents().Num(), 16);
+void FNavLinkStructCustomization::CustomizeChildren(TSharedRef<class IPropertyHandle> StructPropertyHandle, class IDetailChildrenBuilder& StructBuilder, IPropertyTypeCustomizationUtils& StructCustomizationUtils)
+{ 
+	uint32 NumChildProps = 0;
+	StructPropertyHandle->GetNumChildren(NumChildProps);
 
-	for (uint32 i = 0; i < NumChildren; i++)
+	for (uint32 Idx = 0; Idx < NumChildProps; Idx++)
 	{
-		TSharedPtr<IPropertyHandle> PropHandle = StructPropertyHandle->GetChildHandle(i);
-		if (PropHandle->GetProperty() && PropHandle->GetProperty()->GetName().StartsWith(AgentPrefix))
-		{
-			int32 AgentIdx = -1;
-			TTypeFromString<int32>::FromString(AgentIdx, *(PropHandle->GetProperty()->GetName().Mid(AgentPrefix.Len()) ));
-
-			if (AgentIdx >= 0 && AgentIdx < NumAgents)
-			{
-				FText PropName = FText::Format(LOCTEXT("SupportedAgentFmt", "Supports Agent: {0}"), FText::FromName(DefNavSys->GetSupportedAgents()[AgentIdx].Name));
-				StructBuilder.AddChildContent(PropName)
-					.NameContent()
-					[
-						SNew(STextBlock)
-						.Text(PropName)
-						.Font(StructCustomizationUtils.GetRegularFont())
-					]
-					.ValueContent()
-					[
-						PropHandle->CreatePropertyValueWidget()
-					];
-			}
-
-			continue;
-		}
-
+		TSharedPtr<IPropertyHandle> PropHandle = StructPropertyHandle->GetChildHandle(Idx);
 		StructBuilder.AddChildProperty(PropHandle.ToSharedRef());
 	}
 }
-
-#undef LOCTEXT_NAMESPACE
diff --git a/Engine/Source/Runtime/Core/Public/UObject/ObjectVersion.h b/Engine/Source/Runtime/Core/Public/UObject/ObjectVersion.h
--- a/Engine/Source/Runtime/Core/Public/UObject/ObjectVersion.h
+++ b/Engine/Source/Runtime/Core/Public/UObject/ObjectVersion.h
@@ -550,6 +550,8 @@ enum EUnrealEngineObjectUE4Version
 	VER_UE4_ADDED_NON_LINEAR_TRANSITION_BLENDS,
 	// AO Material Mask texture
 	VER_UE4_AO_MATERIAL_MASK,
+	// Replaced navigation agents selection with single structure
+	VER_UE4_NAVIGATION_AGENT_SELECTOR,
 
 	// -----<new versions can be added before this line>-------------------------------------------------
 	// - this needs to be the last line (see note below)
diff --git a/Engine/Source/Runtime/Engine/Classes/AI/Navigation/NavAreas/NavArea.h b/Engine/Source/Runtime/Engine/Classes/AI/Navigation/NavAreas/NavArea.h
--- a/Engine/Source/Runtime/Engine/Classes/AI/Navigation/NavAreas/NavArea.h
+++ b/Engine/Source/Runtime/Engine/Classes/AI/Navigation/NavAreas/NavArea.h
@@ -23,43 +23,48 @@ public:
 	UPROPERTY(EditAnywhere, Category=NavArea, config)
 	FColor DrawColor;
 
+	/** restrict area only to specified agents */
+	UPROPERTY(EditAnywhere, Category=NavArea, config)
+	FNavAgentSelector SupportedAgents;
+
+	// DEPRECATED AGENT CONFIG
 #if CPP
 	union
 	{
 		struct
 		{
 #endif
-	UPROPERTY(EditAnywhere, Category=NavArea, config)
+	UPROPERTY(config)
 	uint32 bSupportsAgent0 : 1;
-	UPROPERTY(EditAnywhere, Category=NavArea, config)
+	UPROPERTY(config)
 	uint32 bSupportsAgent1 : 1;
-	UPROPERTY(EditAnywhere, Category=NavArea, config)
+	UPROPERTY(config)
 	uint32 bSupportsAgent2 : 1;
-	UPROPERTY(EditAnywhere, Category=NavArea, config)
+	UPROPERTY(config)
 	uint32 bSupportsAgent3 : 1;
-	UPROPERTY(EditAnywhere, Category=NavArea, config)
+	UPROPERTY(config)
 	uint32 bSupportsAgent4 : 1;
-	UPROPERTY(EditAnywhere, Category=NavArea, config)
+	UPROPERTY(config)
 	uint32 bSupportsAgent5 : 1;
-	UPROPERTY(EditAnywhere, Category=NavArea, config)
+	UPROPERTY(config)
 	uint32 bSupportsAgent6 : 1;
-	UPROPERTY(EditAnywhere, Category=NavArea, config)
+	UPROPERTY(config)
 	uint32 bSupportsAgent7 : 1;
-	UPROPERTY(EditAnywhere, Category=NavArea, config)
+	UPROPERTY(config)
 	uint32 bSupportsAgent8 : 1;
-	UPROPERTY(EditAnywhere, Category=NavArea, config)
+	UPROPERTY(config)
 	uint32 bSupportsAgent9 : 1;
-	UPROPERTY(EditAnywhere, Category=NavArea, config)
+	UPROPERTY(config)
 	uint32 bSupportsAgent10 : 1;
-	UPROPERTY(EditAnywhere, Category=NavArea, config)
+	UPROPERTY(config)
 	uint32 bSupportsAgent11 : 1;
-	UPROPERTY(EditAnywhere, Category=NavArea, config)
+	UPROPERTY(config)
 	uint32 bSupportsAgent12 : 1;
-	UPROPERTY(EditAnywhere, Category=NavArea, config)
+	UPROPERTY(config)
 	uint32 bSupportsAgent13 : 1;
-	UPROPERTY(EditAnywhere, Category=NavArea, config)
+	UPROPERTY(config)
 	uint32 bSupportsAgent14 : 1;
-	UPROPERTY(EditAnywhere, Category=NavArea, config)
+	UPROPERTY(config)
 	uint32 bSupportsAgent15 : 1;
 #if CPP
 		};
@@ -69,11 +74,12 @@ public:
 
 	virtual void PostInitProperties() override;
 	virtual void FinishDestroy() override;
+	virtual void PostLoad() override;
 
 	FORCEINLINE uint16 GetAreaFlags() const { return AreaFlags; }
 	FORCEINLINE bool HasFlags(uint16 InFlags) const { return (InFlags & AreaFlags) != 0; }
 
-	FORCEINLINE bool IsSupportingAgent(int32 AgentIndex) const { return (AgentIndex >= 0 && AgentIndex < 16) ? !!(SupportedAgentsBits & (1 << AgentIndex)) : false; }
+	FORCEINLINE bool IsSupportingAgent(int32 AgentIndex) const { return SupportedAgents.Contains(AgentIndex); }
 
 	/** called before adding to navigation system */
 	virtual void InitializeArea() {};
@@ -85,7 +91,7 @@ public:
 	static FColor GetColor(UClass* AreaDefinitionClass);
 
 #if WITH_EDITOR
-	/** setup bSupportsAgentX properties */
+	/** setup agent related properties */
 	virtual void UpdateAgentConfig();
 #endif
 	/** copy properties from other area */
diff --git a/Engine/Source/Runtime/Engine/Classes/AI/Navigation/NavLinkDefinition.h b/Engine/Source/Runtime/Engine/Classes/AI/Navigation/NavLinkDefinition.h
--- a/Engine/Source/Runtime/Engine/Classes/AI/Navigation/NavLinkDefinition.h
+++ b/Engine/Source/Runtime/Engine/Classes/AI/Navigation/NavLinkDefinition.h
@@ -52,44 +52,49 @@ struct ENGINE_API FNavigationLinkBase
 	UPROPERTY(EditAnywhere, Category=Default)
 	TSubclassOf<class UNavArea> AreaClass;
 
+	/** restrict area only to specified agents */
+	UPROPERTY(EditAnywhere, Category=Default)
+	FNavAgentSelector SupportedAgents;
+
+	// DEPRECATED AGENT CONFIG
 #if CPP
 	union
 	{
 		struct
 		{
 #endif
-			UPROPERTY(EditAnywhere, Category=Default)
-			uint32 bSupportsAgent0 : 1;
-			UPROPERTY(EditAnywhere, Category=Default)
-			uint32 bSupportsAgent1 : 1;
-			UPROPERTY(EditAnywhere, Category=Default)
-			uint32 bSupportsAgent2 : 1;
-			UPROPERTY(EditAnywhere, Category=Default)
-			uint32 bSupportsAgent3 : 1;
-			UPROPERTY(EditAnywhere, Category=Default)
-			uint32 bSupportsAgent4 : 1;
-			UPROPERTY(EditAnywhere, Category=Default)
-			uint32 bSupportsAgent5 : 1;
-			UPROPERTY(EditAnywhere, Category=Default)
-			uint32 bSupportsAgent6 : 1;
-			UPROPERTY(EditAnywhere, Category=Default)
-			uint32 bSupportsAgent7 : 1;
-			UPROPERTY(EditAnywhere, Category=Default)
-			uint32 bSupportsAgent8 : 1;
-			UPROPERTY(EditAnywhere, Category=Default)
-			uint32 bSupportsAgent9 : 1;
-			UPROPERTY(EditAnywhere, Category=Default)
-			uint32 bSupportsAgent10 : 1;
-			UPROPERTY(EditAnywhere, Category=Default)
-			uint32 bSupportsAgent11 : 1;
-			UPROPERTY(EditAnywhere, Category=Default)
-			uint32 bSupportsAgent12 : 1;
-			UPROPERTY(EditAnywhere, Category=Default)
-			uint32 bSupportsAgent13 : 1;
-			UPROPERTY(EditAnywhere, Category=Default)
-			uint32 bSupportsAgent14 : 1;
-			UPROPERTY(EditAnywhere, Category=Default)
-			uint32 bSupportsAgent15 : 1;
+	UPROPERTY()
+	uint32 bSupportsAgent0 : 1;
+	UPROPERTY()
+	uint32 bSupportsAgent1 : 1;
+	UPROPERTY()
+	uint32 bSupportsAgent2 : 1;
+	UPROPERTY()
+	uint32 bSupportsAgent3 : 1;
+	UPROPERTY()
+	uint32 bSupportsAgent4 : 1;
+	UPROPERTY()
+	uint32 bSupportsAgent5 : 1;
+	UPROPERTY()
+	uint32 bSupportsAgent6 : 1;
+	UPROPERTY()
+	uint32 bSupportsAgent7 : 1;
+	UPROPERTY()
+	uint32 bSupportsAgent8 : 1;
+	UPROPERTY()
+	uint32 bSupportsAgent9 : 1;
+	UPROPERTY()
+	uint32 bSupportsAgent10 : 1;
+	UPROPERTY()
+	uint32 bSupportsAgent11 : 1;
+	UPROPERTY()
+	uint32 bSupportsAgent12 : 1;
+	UPROPERTY()
+	uint32 bSupportsAgent13 : 1;
+	UPROPERTY()
+	uint32 bSupportsAgent14 : 1;
+	UPROPERTY()
+	uint32 bSupportsAgent15 : 1;
 #if CPP
 		};
 		uint32 SupportedAgentsBits;
@@ -103,6 +108,17 @@ struct ENGINE_API FNavigationLinkBase
 #endif // WITH_EDITORONLY_DATA
 
 	FNavigationLinkBase();
+
+	void PostSerialize(const FArchive& Ar);
+};
+
+template<>
+struct TStructOpsTypeTraits< FNavigationLinkBase > : public TStructOpsTypeTraitsBase
+{
+	enum
+	{
+		WithPostSerialize = true,
+	};
 };
 
 USTRUCT()
@@ -151,6 +167,20 @@ struct ENGINE_API FNavigationLink : public FNavigationLinkBase
 
 		return Result;
 	}
+
+	void PostSerialize(const FArchive& Ar)
+	{
+		FNavigationLinkBase::PostSerialize(Ar);
+	}
+};
+
+template<>
+struct TStructOpsTypeTraits< FNavigationLink > : public TStructOpsTypeTraitsBase
+{
+	enum
+	{
+		WithPostSerialize = true,
+	};
 };
 
 USTRUCT()
@@ -199,6 +229,20 @@ struct ENGINE_API FNavigationSegmentLink : public FNavigationLinkBase
 
 		return Result;
 	}
+
+	void PostSerialize(const FArchive& Ar)
+	{
+		FNavigationLinkBase::PostSerialize(Ar);
+	}
+};
+
+template<>
+struct TStructOpsTypeTraits< FNavigationSegmentLink > : public TStructOpsTypeTraitsBase
+{
+	enum
+	{
+		WithPostSerialize = true,
+	};
 };
 
 /** Class containing definition of a navigation area */
diff --git a/Engine/Source/Runtime/Engine/Classes/AI/Navigation/NavigationTypes.h b/Engine/Source/Runtime/Engine/Classes/AI/Navigation/NavigationTypes.h
--- a/Engine/Source/Runtime/Engine/Classes/AI/Navigation/NavigationTypes.h
+++ b/Engine/Source/Runtime/Engine/Classes/AI/Navigation/NavigationTypes.h
@@ -414,6 +414,84 @@ inline uint32 GetTypeHash(const FNavAgentProperties& A)
 	return ((int16(A.AgentRadius) << 16) | int16(A.AgentHeight)) ^ int32(A.AgentStepHeight);
 }
 
+USTRUCT()
+struct ENGINE_API FNavAgentSelector
+{
+	GENERATED_USTRUCT_BODY()
+
+#if CPP
+	union
+	{
+		struct
+		{
+#endif
+			UPROPERTY(EditAnywhere, Category = Default)
+			uint32 bSupportsAgent0 : 1;
+			UPROPERTY(EditAnywhere, Category = Default)
+			uint32 bSupportsAgent1 : 1;
+			UPROPERTY(EditAnywhere, Category = Default)
+			uint32 bSupportsAgent2 : 1;
+			UPROPERTY(EditAnywhere, Category = Default)
+			uint32 bSupportsAgent3 : 1;
+			UPROPERTY(EditAnywhere, Category = Default)
+			uint32 bSupportsAgent4 : 1;
+			UPROPERTY(EditAnywhere, Category = Default)
+			uint32 bSupportsAgent5 : 1;
+			UPROPERTY(EditAnywhere, Category = Default)
+			uint32 bSupportsAgent6 : 1;
+			UPROPERTY(EditAnywhere, Category = Default)
+			uint32 bSupportsAgent7 : 1;
+			UPROPERTY(EditAnywhere, Category = Default)
+			uint32 bSupportsAgent8 : 1;
+			UPROPERTY(EditAnywhere, Category = Default)
+			uint32 bSupportsAgent9 : 1;
+			UPROPERTY(EditAnywhere, Category = Default)
+			uint32 bSupportsAgent10 : 1;
+			UPROPERTY(EditAnywhere, Category = Default)
+			uint32 bSupportsAgent11 : 1;
+			UPROPERTY(EditAnywhere, Category = Default)
+			uint32 bSupportsAgent12 : 1;
+			UPROPERTY(EditAnywhere, Category = Default)
+			uint32 bSupportsAgent13 : 1;
+			UPROPERTY(EditAnywhere, Category = Default)
+			uint32 bSupportsAgent14 : 1;
+			UPROPERTY(EditAnywhere, Category = Default)
+			uint32 bSupportsAgent15 : 1;
+#if CPP
+		};
+		uint32 PackedBits;
+	};
+#endif
+
+	FNavAgentSelector();
+
+	FORCEINLINE bool Contains(int32 AgentIndex) const
+	{
+		return (AgentIndex >= 0 && AgentIndex < 16) ? !!(PackedBits & (1 << AgentIndex)) : false;
+	}
+
+	FORCEINLINE bool IsInitialized() const
+	{
+		return (PackedBits & 0x80000000) != 0;
+	}
+
+	FORCEINLINE void MarkInitialized()
+	{
+		PackedBits |= 0x80000000;
+	}
+
+	bool Serialize(FArchive& Ar);
+};
+
+template<>
+struct TStructOpsTypeTraits< FNavAgentSelector > : public TStructOpsTypeTraitsBase
+{
+	enum
+	{
+		WithSerializer = true,
+	};
+};
+
 USTRUCT()
 struct ENGINE_API FNavDataConfig : public FNavAgentProperties
 {
diff --git a/Engine/Source/Runtime/Engine/Private/AI/Navigation/NavAreas/NavArea.cpp b/Engine/Source/Runtime/Engine/Private/AI/Navigation/NavAreas/NavArea.cpp
--- a/Engine/Source/Runtime/Engine/Private/AI/Navigation/NavAreas/NavArea.cpp
+++ b/Engine/Source/Runtime/Engine/Private/AI/Navigation/NavAreas/NavArea.cpp
@@ -41,6 +41,32 @@ void UNavArea::FinishDestroy()
 	Super::FinishDestroy();
 }
 
+void UNavArea::PostLoad()
+{
+	Super::PostLoad();
+
+	if (!SupportedAgents.IsInitialized())
+	{
+		SupportedAgents.bSupportsAgent0 = bSupportsAgent0;
+		SupportedAgents.bSupportsAgent1 = bSupportsAgent1;
+		SupportedAgents.bSupportsAgent2 = bSupportsAgent2;
+		SupportedAgents.bSupportsAgent3 = bSupportsAgent3;
+		SupportedAgents.bSupportsAgent4 = bSupportsAgent4;
+		SupportedAgents.bSupportsAgent5 = bSupportsAgent5;
+		SupportedAgents.bSupportsAgent6 = bSupportsAgent6;
+		SupportedAgents.bSupportsAgent7 = bSupportsAgent7;
+		SupportedAgents.bSupportsAgent8 = bSupportsAgent8;
+		SupportedAgents.bSupportsAgent9 = bSupportsAgent9;
+		SupportedAgents.bSupportsAgent10 = bSupportsAgent10;
+		SupportedAgents.bSupportsAgent11 = bSupportsAgent11;
+		SupportedAgents.bSupportsAgent12 = bSupportsAgent12;
+		SupportedAgents.bSupportsAgent13 = bSupportsAgent13;
+		SupportedAgents.bSupportsAgent14 = bSupportsAgent14;
+		SupportedAgents.bSupportsAgent15 = bSupportsAgent15;
+		SupportedAgents.MarkInitialized();
+	}
+}
+
 FColor UNavArea::GetColor(UClass* AreaDefinitionClass)
 {
 	return AreaDefinitionClass ? AreaDefinitionClass->GetDefaultObject<UNavArea>()->DrawColor : FColor::Black;
@@ -64,33 +90,6 @@ void UNavArea::CopyFrom(TSubclassOf<UNavArea> AreaClass)
 #if WITH_EDITOR
 void UNavArea::UpdateAgentConfig()
 {
-	const UNavigationSystem* DefNavSys = (UNavigationSystem*)(UNavigationSystem::StaticClass()->GetDefaultObject());
-	check(DefNavSys);
-
-	const int32 MaxAllowedAgents = 16;
-	const int32 NumAgents = FMath::Min(DefNavSys->GetSupportedAgents().Num(), MaxAllowedAgents);
-	if (DefNavSys->GetSupportedAgents().Num() > MaxAllowedAgents)
-	{
-		UE_LOG(LogNavigation, Error, TEXT("Navigation system supports %d agents, but only %d can be shown in %s properties!"),
-			DefNavSys->GetSupportedAgents().Num(), MaxAllowedAgents, *GetClass()->GetName());
-	}
-
-	const FString CustomNameMeta = TEXT("DisplayName");
-	for (int32 i = 0; i < MaxAllowedAgents; i++)
-	{
-		const FString PropName = FString::Printf(TEXT("bSupportsAgent%d"), i);
-		UProperty* Prop = FindField<UProperty>(UNavArea::StaticClass(), *PropName);
-		check(Prop);
-
-		if (i < NumAgents && NumAgents > 1)
-		{
-			Prop->SetPropertyFlags(CPF_Edit);
-			Prop->SetMetaData(*CustomNameMeta, *FString::Printf(TEXT("Supports Agent: %s"), *DefNavSys->GetSupportedAgents()[i].Name.ToString()));
-		}
-		else
-		{
-			Prop->ClearPropertyFlags(CPF_Edit);
-		}
-	}
+	// empty in base class
 }
 #endif
\ No newline at end of file
diff --git a/Engine/Source/Runtime/Engine/Private/AI/Navigation/NavLinkRenderingComponent.cpp b/Engine/Source/Runtime/Engine/Private/AI/Navigation/NavLinkRenderingComponent.cpp
--- a/Engine/Source/Runtime/Engine/Private/AI/Navigation/NavLinkRenderingComponent.cpp
+++ b/Engine/Source/Runtime/Engine/Private/AI/Navigation/NavLinkRenderingComponent.cpp
@@ -137,7 +137,7 @@ void FNavLinkRenderingProxy::StorePointLinks(const FTransform& LocalToWorld, con
 		LinkDrawing.Color = UNavArea::GetColor(Link->AreaClass);
 		LinkDrawing.SnapRadius = Link->SnapRadius;
 		LinkDrawing.SnapHeight = Link->bUseSnapHeight ? Link->SnapHeight : -1.0f;
-		LinkDrawing.SupportedAgentsBits = Link->SupportedAgentsBits;
+		LinkDrawing.SupportedAgentsBits = Link->SupportedAgents.PackedBits;
 		OffMeshPointLinks.Add(LinkDrawing);
 	}
 }
@@ -156,7 +156,7 @@ void FNavLinkRenderingProxy::StoreSegmentLinks(const FTransform& LocalToWorld, c
 		LinkDrawing.Color = UNavArea::GetColor(Link->AreaClass);
 		LinkDrawing.SnapRadius = Link->SnapRadius;
 		LinkDrawing.SnapHeight = Link->bUseSnapHeight ? Link->SnapHeight : -1.0f;
-		LinkDrawing.SupportedAgentsBits = Link->SupportedAgentsBits;
+		LinkDrawing.SupportedAgentsBits = Link->SupportedAgents.PackedBits;
 		OffMeshSegmentLinks.Add(LinkDrawing);
 	}
 }
diff --git a/Engine/Source/Runtime/Engine/Private/AI/Navigation/NavigationModifier.cpp b/Engine/Source/Runtime/Engine/Private/AI/Navigation/NavigationModifier.cpp
--- a/Engine/Source/Runtime/Engine/Private/AI/Navigation/NavigationModifier.cpp
+++ b/Engine/Source/Runtime/Engine/Private/AI/Navigation/NavigationModifier.cpp
@@ -30,6 +30,30 @@ FNavigationLinkBase::FNavigationLinkBase()
 	SupportedAgentsBits = 0xFFFFFFFF;
 }
 
+void FNavigationLinkBase::PostSerialize(const FArchive& Ar)
+{
+	if (Ar.IsLoading() && Ar.UE4Ver() < VER_UE4_NAVIGATION_AGENT_SELECTOR)
+	{
+		SupportedAgents.bSupportsAgent0 = bSupportsAgent0;
+		SupportedAgents.bSupportsAgent1 = bSupportsAgent1;
+		SupportedAgents.bSupportsAgent2 = bSupportsAgent2;
+		SupportedAgents.bSupportsAgent3 = bSupportsAgent3;
+		SupportedAgents.bSupportsAgent4 = bSupportsAgent4;
+		SupportedAgents.bSupportsAgent5 = bSupportsAgent5;
+		SupportedAgents.bSupportsAgent6 = bSupportsAgent6;
+		SupportedAgents.bSupportsAgent7 = bSupportsAgent7;
+		SupportedAgents.bSupportsAgent8 = bSupportsAgent8;
+		SupportedAgents.bSupportsAgent9 = bSupportsAgent9;
+		SupportedAgents.bSupportsAgent10 = bSupportsAgent10;
+		SupportedAgents.bSupportsAgent11 = bSupportsAgent11;
+		SupportedAgents.bSupportsAgent12 = bSupportsAgent12;
+		SupportedAgents.bSupportsAgent13 = bSupportsAgent13;
+		SupportedAgents.bSupportsAgent14 = bSupportsAgent14;
+		SupportedAgents.bSupportsAgent15 = bSupportsAgent15;
+		SupportedAgents.MarkInitialized();
+	}
+}
+
 //----------------------------------------------------------------------//
 // UNavLinkDefinition
 //----------------------------------------------------------------------//
diff --git a/Engine/Source/Runtime/Engine/Private/AI/Navigation/NavigationSystem.cpp b/Engine/Source/Runtime/Engine/Private/AI/Navigation/NavigationSystem.cpp
--- a/Engine/Source/Runtime/Engine/Private/AI/Navigation/NavigationSystem.cpp
+++ b/Engine/Source/Runtime/Engine/Private/AI/Navigation/NavigationSystem.cpp
@@ -159,6 +159,19 @@ FNavHeightfieldSamples::FNavHeightfieldSamples()
 #endif // WITH_PHYSX
 }
 
+//----------------------------------------------------------------------//
+// FNavAgentSelector
+//----------------------------------------------------------------------//
+FNavAgentSelector::FNavAgentSelector() : PackedBits(0x7fffffff)
+{
+}
+
+bool FNavAgentSelector::Serialize(FArchive& Ar)
+{
+	Ar << PackedBits;
+	return true;
+}
+
 //----------------------------------------------------------------------//
 // FNavDataConfig
 //----------------------------------------------------------------------//
diff --git a/Engine/Source/Runtime/Engine/Private/AI/Navigation/RecastNavMeshGenerator.cpp b/Engine/Source/Runtime/Engine/Private/AI/Navigation/RecastNavMeshGenerator.cpp
--- a/Engine/Source/Runtime/Engine/Private/AI/Navigation/RecastNavMeshGenerator.cpp
+++ b/Engine/Source/Runtime/Engine/Private/AI/Navigation/RecastNavMeshGenerator.cpp
@@ -1259,12 +1259,12 @@ struct FOffMeshData
 		LinkParams.Reserve(ElementsCount);
 	}
 
-	void AddLinks(const TArray<FNavigationLink>& Links, const FTransform& LocalToWorld, uint32 AgentMask, float DefaultSnapHeight)
+	void AddLinks(const TArray<FNavigationLink>& Links, const FTransform& LocalToWorld, int32 AgentIndex, float DefaultSnapHeight)
 	{
 		for (int32 LinkIndex = 0; LinkIndex < Links.Num(); ++LinkIndex)
 		{
 			const FNavigationLink& Link = Links[LinkIndex];
-			if ((Link.SupportedAgentsBits & AgentMask) == 0)
+			if (!Link.SupportedAgents.Contains(AgentIndex))
 			{
 				continue;
 			}
@@ -1300,12 +1300,12 @@ struct FOffMeshData
 			LinkParams.Add(NewInfo);
 		}
 	}
-	void AddSegmentLinks(const TArray<FNavigationSegmentLink>& Links, const FTransform& LocalToWorld, uint32 AgentMask, float DefaultSnapHeight)
+	void AddSegmentLinks(const TArray<FNavigationSegmentLink>& Links, const FTransform& LocalToWorld, int32 AgentIndex, float DefaultSnapHeight)
 	{
 		for (int32 LinkIndex = 0; LinkIndex < Links.Num(); ++LinkIndex)
 		{
 			const FNavigationSegmentLink& Link = Links[LinkIndex];
-			if ((Link.SupportedAgentsBits & AgentMask) == 0)
+			if (!Link.SupportedAgents.Contains(AgentIndex))
 			{
 				continue;
 			}
@@ -2688,15 +2688,14 @@ bool FRecastTileGenerator::GenerateNavigationData(FNavMeshBuildContext& BuildCon
 				OffMeshData.Reserve(OffmeshLinks.Num());
 				OffMeshData.AreaClassToIdMap = &AdditionalCachedData.AreaClassToIdMap;
 				OffMeshData.FlagsPerArea = AdditionalCachedData.FlagsPerOffMeshLinkArea;
-				const uint32 AgentMask = (1 << TileConfig.AgentIndex);
 				const FSimpleLinkNavModifier* LinkModifier = OffmeshLinks.GetData();
 				const float DefaultSnapHeight = TileConfig.walkableClimb * TileConfig.ch;
 
 				for (int32 LinkModifierIndex = 0; LinkModifierIndex < OffmeshLinks.Num(); ++LinkModifierIndex, ++LinkModifier)
 				{
-					OffMeshData.AddLinks(LinkModifier->Links, LinkModifier->LocalToWorld, AgentMask, DefaultSnapHeight);
+					OffMeshData.AddLinks(LinkModifier->Links, LinkModifier->LocalToWorld, TileConfig.AgentIndex, DefaultSnapHeight);
 #if GENERATE_SEGMENT_LINKS
-					OffMeshData.AddSegmentLinks(LinkModifier->SegmentLinks, LinkModifier->LocalToWorld, AgentMask, DefaultSnapHeight);
+					OffMeshData.AddSegmentLinks(LinkModifier->SegmentLinks, LinkModifier->LocalToWorld, TileConfig.AgentIndex, DefaultSnapHeight);
 #endif // GENERATE_SEGMENT_LINKS
 				}
 			}
