UE-17069 Make/Break/SetMembersInStruct could be improved for PostProcessSettings
2nd part: the properties are grouped in categories on details pannel

diff --git a/Engine/Source/Editor/BlueprintGraph/Classes/K2Node.h b/Engine/Source/Editor/BlueprintGraph/Classes/K2Node.h
--- a/Engine/Source/Editor/BlueprintGraph/Classes/K2Node.h
+++ b/Engine/Source/Editor/BlueprintGraph/Classes/K2Node.h
@@ -32,18 +32,22 @@ struct FOptionalPinFromProperty
 
 	UPROPERTY(EditAnywhere, BlueprintReadWrite, Category=Hi)
 	bool bPropertyIsCustomized;
+	
+	UPROPERTY(EditAnywhere, BlueprintReadWrite, Category = Hi)
+	FName CategoryName;
 
 	FOptionalPinFromProperty()
 	{
 	}
 	
-	FOptionalPinFromProperty(FName InPropertyName, bool bInShowPin, bool bInCanToggleVisibility, const FString& InFriendlyName, const FText& InTooltip, bool bInPropertyIsCustomized)
+	FOptionalPinFromProperty(FName InPropertyName, bool bInShowPin, bool bInCanToggleVisibility, const FString& InFriendlyName, const FText& InTooltip, bool bInPropertyIsCustomized, FName InCategoryName)
 		: PropertyName(InPropertyName)
 		, PropertyFriendlyName(InFriendlyName)
 		, PropertyTooltip(InTooltip)
 		, bShowPin(bInShowPin)
 		, bCanToggleVisibility(bInCanToggleVisibility)
 		, bPropertyIsCustomized(bInPropertyIsCustomized)
+		, CategoryName(InCategoryName)
 	{
 	}
 };
diff --git a/Engine/Source/Editor/BlueprintGraph/Private/K2Node.cpp b/Engine/Source/Editor/BlueprintGraph/Private/K2Node.cpp
--- a/Engine/Source/Editor/BlueprintGraph/Private/K2Node.cpp
+++ b/Engine/Source/Editor/BlueprintGraph/Private/K2Node.cpp
@@ -9,6 +9,8 @@
 #include "KismetCompiler.h"
 #include "GraphEditorSettings.h"
 
+#include "ObjectEditorUtils.h"
+
 #define LOCTEXT_NAMESPACE "K2Node"
 
 // File-Scoped Globals
@@ -848,7 +850,9 @@ void FOptionalPinManager::RebuildPropertyList(TArray<FOptionalPinFromProperty>&
 			Record->PropertyName = TestProperty->GetFName();
 			Record->PropertyFriendlyName = UEditorEngine::GetFriendlyName(TestProperty, SourceStruct);
 			Record->PropertyTooltip = TestProperty->GetToolTipText();
-
+#if WITH_EDITOR
+			Record->CategoryName = FObjectEditorUtils::GetCategoryFName(TestProperty);
+#endif //WITH_EDITOR
 			// Get the defaults
 			GetRecordDefaults(TestProperty, *Record);
 
diff --git a/Engine/Source/Editor/DetailCustomizations/Private/SkeletalControlNodeDetails.cpp b/Engine/Source/Editor/DetailCustomizations/Private/SkeletalControlNodeDetails.cpp
--- a/Engine/Source/Editor/DetailCustomizations/Private/SkeletalControlNodeDetails.cpp
+++ b/Engine/Source/Editor/DetailCustomizations/Private/SkeletalControlNodeDetails.cpp
@@ -3,7 +3,6 @@
 #include "DetailCustomizationsPrivatePCH.h"
 #include "SkeletalControlNodeDetails.h"
 
-
 #define LOCTEXT_NAMESPACE "SkeletalControlNodeDetails"
 
 /////////////////////////////////////////////////////
@@ -18,40 +17,81 @@ void FSkeletalControlNodeDetails::CustomizeDetails(class IDetailLayoutBuilder& D
 {
 	DetailCategory = &DetailBuilder.EditCategory("PinOptions");
 	TSharedRef<IPropertyHandle> AvailablePins = DetailBuilder.GetProperty("ShowPinForProperties");
+	TSharedPtr<IPropertyHandleArray> ArrayProperty = AvailablePins->AsArray();
+
+	TSet<FName> UniqueCategoryNames;
+	{
+		int32 NumElements = 0;
+		{
+			uint32 UnNumElements = 0;
+			if (ArrayProperty.IsValid() && (FPropertyAccess::Success == ArrayProperty->GetNumElements(UnNumElements)))
+			{
+				NumElements = static_cast<int32>(UnNumElements);
+			}
+		}
+		for (int32 Index = 0; Index < NumElements; ++Index)
+		{
+			auto StructPropHandle = ArrayProperty->GetElement(Index);
+			auto CategoryPropeHandle = StructPropHandle->GetChildHandle("CategoryName");
+			check(CategoryPropeHandle.IsValid());
+			FName CategoryNameValue;
+			const auto Result = CategoryPropeHandle->GetValue(CategoryNameValue);
+			if (ensure(FPropertyAccess::Success == Result))
+			{
+				UniqueCategoryNames.Add(CategoryNameValue);
+			}
+		}
+	}
 
 	//@TODO: Shouldn't show this if the available pins array is empty!
 	const bool bGenerateHeader = true;
 	const bool bDisplayResetToDefault = false;
-	TSharedRef<FDetailArrayBuilder> AvailablePinsBuilder = MakeShareable( new FDetailArrayBuilder( AvailablePins, bGenerateHeader, bDisplayResetToDefault ) );
-	AvailablePinsBuilder->OnGenerateArrayElementWidget( FOnGenerateArrayElementWidget::CreateSP( this, &FSkeletalControlNodeDetails::OnGenerateElementForPropertyPin ) );
-
+	const bool bDisplayElementNum = false;
 	const bool bForAdvanced = false;
-	DetailCategory->AddCustomBuilder( AvailablePinsBuilder, bForAdvanced );
+	for (auto CategoryName : UniqueCategoryNames)
+	{
+		TSharedRef<FDetailArrayBuilder> AvailablePinsBuilder = MakeShareable(new FDetailArrayBuilder(AvailablePins, bGenerateHeader, bDisplayResetToDefault, bDisplayElementNum));
+		AvailablePinsBuilder->OnGenerateArrayElementWidget(FOnGenerateArrayElementWidget::CreateSP(this, &FSkeletalControlNodeDetails::OnGenerateElementForPropertyPin, CategoryName));
+		AvailablePinsBuilder->SetDisplayName((CategoryName == NAME_None) ? LOCTEXT("DefaultCategory", "Default Category") : FText::FromName(CategoryName));
+		DetailCategory->AddCustomBuilder(AvailablePinsBuilder, bForAdvanced);
+	}
 }
 
-void FSkeletalControlNodeDetails::OnGenerateElementForPropertyPin(TSharedRef<IPropertyHandle> ElementProperty, int32 ElementIndex, IDetailChildrenBuilder& ChildrenBuilder )
+void FSkeletalControlNodeDetails::OnGenerateElementForPropertyPin(TSharedRef<IPropertyHandle> ElementProperty, int32 ElementIndex, IDetailChildrenBuilder& ChildrenBuilder, FName CategoryName)
 {
+	{
+		auto CategoryPropeHandle = ElementProperty->GetChildHandle("CategoryName");
+		check(CategoryPropeHandle.IsValid());
+		FName CategoryNameValue;
+		const auto Result = CategoryPropeHandle->GetValue(CategoryNameValue);
+		const bool bProperCategory = ensure(FPropertyAccess::Success == Result) && (CategoryNameValue == CategoryName);
+		if (!bProperCategory)
+		{
+			return;
+		}
+	}
+
 	TSharedPtr<IPropertyHandle> PropertyNameHandle = ElementProperty->GetChildHandle("PropertyFriendlyName");
 	FText PropertyFriendlyName(LOCTEXT("Invalid", "Invalid"));
-	 
+
 	if (PropertyNameHandle.IsValid())
 	{
 		FString DisplayFriendlyName;
 		switch (PropertyNameHandle->GetValue(/*out*/ DisplayFriendlyName))
 		{
 		case FPropertyAccess::Success:
-			{
-				PropertyFriendlyName = FText::FromString(DisplayFriendlyName);
+		{
+			PropertyFriendlyName = FText::FromString(DisplayFriendlyName);
 
-				//DetailBuilder.EditCategory
-				//DisplayNameAsName
-			}
+			//DetailBuilder.EditCategory
+			//DisplayNameAsName
+		}
 			break;
 		case FPropertyAccess::MultipleValues:
-			ChildrenBuilder.AddChildContent( FText::GetEmpty() ) 
-			[
-				 SNew(STextBlock).Text(LOCTEXT("OnlyWorksInSingleSelectMode", "Multiple types selected"))
-			];
+			ChildrenBuilder.AddChildContent(FText::GetEmpty())
+				[
+					SNew(STextBlock).Text(LOCTEXT("OnlyWorksInSingleSelectMode", "Multiple types selected"))
+				];
 			return;
 		case FPropertyAccess::Fail:
 		default:
diff --git a/Engine/Source/Editor/DetailCustomizations/Private/SkeletalControlNodeDetails.h b/Engine/Source/Editor/DetailCustomizations/Private/SkeletalControlNodeDetails.h
--- a/Engine/Source/Editor/DetailCustomizations/Private/SkeletalControlNodeDetails.h
+++ b/Engine/Source/Editor/DetailCustomizations/Private/SkeletalControlNodeDetails.h
@@ -15,7 +15,7 @@ public:
 	// End of IDetailCustomization interface
 
 protected:
-	void OnGenerateElementForPropertyPin(TSharedRef<IPropertyHandle> ElementProperty, int32 ElementIndex, IDetailChildrenBuilder& ChildrenBuilder );
+	void OnGenerateElementForPropertyPin(TSharedRef<IPropertyHandle> ElementProperty, int32 ElementIndex, IDetailChildrenBuilder& ChildrenBuilder, FName CategoryName);
 private:
 	IDetailCategoryBuilder* DetailCategory;
 };
\ No newline at end of file
diff --git a/Engine/Source/Editor/PropertyEditor/Public/PropertyCustomizationHelpers.h b/Engine/Source/Editor/PropertyEditor/Public/PropertyCustomizationHelpers.h
--- a/Engine/Source/Editor/PropertyEditor/Public/PropertyCustomizationHelpers.h
+++ b/Engine/Source/Editor/PropertyEditor/Public/PropertyCustomizationHelpers.h
@@ -211,11 +211,12 @@ class FDetailArrayBuilder : public IDetailCustomNodeBuilder
 {
 public:
 
-	FDetailArrayBuilder( TSharedRef<IPropertyHandle> InBaseProperty, bool InGenerateHeader = true, bool InDisplayResetToDefault = true)
+	FDetailArrayBuilder(TSharedRef<IPropertyHandle> InBaseProperty, bool InGenerateHeader = true, bool InDisplayResetToDefault = true, bool InDisplayElementNum = true)
 		: ArrayProperty( InBaseProperty->AsArray() )
 		, BaseProperty( InBaseProperty )
 		, bGenerateHeader( InGenerateHeader)
 		, bDisplayResetToDefault(InDisplayResetToDefault)
+		, bDisplayElementNum(InDisplayElementNum)
 	{
 		check( ArrayProperty.IsValid() );
 
@@ -258,7 +259,17 @@ public:
 		if (bGenerateHeader)
 		{
 			const bool bDisplayResetToDefaultInNameContent = false;
+
 			TSharedPtr<SHorizontalBox> ContentHorizontalBox;
+			SAssignNew(ContentHorizontalBox, SHorizontalBox);
+			if (bDisplayElementNum)
+			{
+				ContentHorizontalBox->AddSlot()
+				[
+					BaseProperty->CreatePropertyValueWidget()
+				];
+			}
+
 			NodeRow
 			.FilterString(!DisplayName.IsEmpty() ? DisplayName : BaseProperty->GetPropertyDisplayName())
 			.NameContent()
@@ -267,11 +278,7 @@ public:
 			]
 			.ValueContent()
 			[
-				SAssignNew(ContentHorizontalBox, SHorizontalBox)
-				+ SHorizontalBox::Slot()
-				[
-					BaseProperty->CreatePropertyValueWidget()
-				]
+				ContentHorizontalBox.ToSharedRef()
 			];
 
 			if (bDisplayResetToDefault)
@@ -321,6 +328,7 @@ private:
 	FSimpleDelegate OnRebuildChildren;
 	bool bGenerateHeader;
 	bool bDisplayResetToDefault;
+	bool bDisplayElementNum;
 };
 
 /**
