FIX UE-16749 auto weld in a blueprint is enabled even if set to simulate

diff --git a/Engine/Source/Editor/DetailCustomizations/Private/PrimitiveComponentDetails.cpp b/Engine/Source/Editor/DetailCustomizations/Private/PrimitiveComponentDetails.cpp
--- a/Engine/Source/Editor/DetailCustomizations/Private/PrimitiveComponentDetails.cpp
+++ b/Engine/Source/Editor/DetailCustomizations/Private/PrimitiveComponentDetails.cpp
@@ -108,7 +108,7 @@ bool FPrimitiveComponentDetails::IsAutoWeldEditable() const
 		{
 			if (FBodyInstance* BI = SceneComponent->GetBodyInstance())
 			{
-				if (BI->IsInstanceSimulatingPhysics())
+				if (BI->ShouldInstanceSimulatingPhysics())
 				{
 					return false;
 				}
