Forgot to submit details change for new UI

diff --git a/Engine/Source/Editor/DetailCustomizations/Private/PrimitiveComponentDetails.cpp b/Engine/Source/Editor/DetailCustomizations/Private/PrimitiveComponentDetails.cpp
--- a/Engine/Source/Editor/DetailCustomizations/Private/PrimitiveComponentDetails.cpp
+++ b/Engine/Source/Editor/DetailCustomizations/Private/PrimitiveComponentDetails.cpp
@@ -145,6 +145,19 @@ EVisibility FPrimitiveComponentDetails::IsMassVisible(bool bOverrideMass) const
 	}
 }
 
+EVisibility FPrimitiveComponentDetails::IsMaxAngularVelocityVisible(bool bOverrideMaxAngularVelocity) const
+{
+	bool bIsMaxAngularVelocityReadOnly = IsMaxAngularVelocityReadOnly();
+	if (bOverrideMaxAngularVelocity)
+	{
+		return bIsMaxAngularVelocityReadOnly ? EVisibility::Collapsed : EVisibility::Visible;
+	}
+	else
+	{
+		return bIsMaxAngularVelocityReadOnly ? EVisibility::Visible : EVisibility::Collapsed;
+	}
+}
+
 void FPrimitiveComponentDetails::AddPhysicsCategory(IDetailLayoutBuilder& DetailBuilder)
 {
 	if (DetailBuilder.GetProperty(GET_MEMBER_NAME_CHECKED(UPrimitiveComponent, BodyInstance))->IsValidHandle())
@@ -157,7 +170,6 @@ void FPrimitiveComponentDetails::AddPhysicsCategory(IDetailLayoutBuilder& Detail
 		IDetailCategoryBuilder& PhysicsCategory = DetailBuilder.EditCategory("Physics");
 
 		bool bDisplayMass = true;
-		bool bDisplayMassOverride = true;
 		bool bDisplayConstraints = true;
 
 		for (int32 i = 0; i < ObjectsCustomized.Num(); ++i)
@@ -165,14 +177,8 @@ void FPrimitiveComponentDetails::AddPhysicsCategory(IDetailLayoutBuilder& Detail
 			if (ObjectsCustomized[i].IsValid() && ObjectsCustomized[i]->IsA(UDestructibleComponent::StaticClass()))
 			{
 				bDisplayMass = false;
-				bDisplayMassOverride = false;
 				bDisplayConstraints = false;
 			}
-
-			if (ObjectsCustomized[i].IsValid() && ObjectsCustomized[i]->IsA(USkeletalMeshComponent::StaticClass()))
-			{
-				bDisplayMassOverride = false;
-			}
 		}
 
 		bool bConstraintsUIExists = false;
@@ -355,64 +361,80 @@ void FPrimitiveComponentDetails::AddPhysicsCategory(IDetailLayoutBuilder& Detail
 					PhysicsCategory.AddProperty(ChildProperty).Visibility(TAttribute<EVisibility>(this, &FPrimitiveComponentDetails::IsAutoWeldVisible))
 						.EditCondition(TAttribute<bool>(this, &FPrimitiveComponentDetails::IsAutoWeldEditable), NULL);
 				}
-				else if (PropName == GET_MEMBER_NAME_CHECKED(FBodyInstance, MassInKg) || PropName == GET_MEMBER_NAME_CHECKED(FBodyInstance, bOverrideMass))
+				else if (PropName == GET_MEMBER_NAME_CHECKED(FBodyInstance, MassInKg))
 				{
 					if (bDisplayMass && bMassUIExists == false)
 					{
 						bMassUIExists = true;
 						TSharedPtr<IPropertyHandle> MassInKGHandle = BodyInstanceHandler->GetChildHandle(TEXT("MassInKg"));
-						TSharedPtr<IPropertyHandle> MassOverrideHandle = BodyInstanceHandler->GetChildHandle(TEXT("bOverrideMass"));
 
-						PhysicsCategory.AddCustomRow(LOCTEXT("Mass", "Mass"), false)
+						PhysicsCategory.AddProperty(MassInKGHandle).CustomWidget()
 						.NameContent()
 						[
 							MassInKGHandle->CreatePropertyNameWidget()
 						]
 						.ValueContent()
-						.MinDesiredWidth(125 * 3.f)
 						[
-							SNew(SHorizontalBox)
-							+ SHorizontalBox::Slot()
-							.AutoWidth()
+							SNew(SVerticalBox)
+							+ SVerticalBox::Slot()
+							.Padding(0.f, 0.f, 10.f, 0.f)
+							[
+								SNew(SNumericEntryBox<float>)
+								.IsEnabled(false)
+								.Font( IDetailLayoutBuilder::GetDetailFont() )
+								.Value(this, &FPrimitiveComponentDetails::OnGetBodyMass)
+								.Visibility(this, &FPrimitiveComponentDetails::IsMassVisible, false)
+							]
+
+							+ SVerticalBox::Slot()
+							.Padding(0.f, 0.f, 10.f, 0.f)
 							[
 								SNew(SVerticalBox)
+								.Visibility(this, &FPrimitiveComponentDetails::IsMassVisible, true)
 								+ SVerticalBox::Slot()
-								.Padding(0.f, 0.f, 10.f, 0.f)
+								.AutoHeight()
 								[
-									SNew(SNumericEntryBox<float>)
-									.IsEnabled(false)
-									.Font( IDetailLayoutBuilder::GetDetailFont() )
-									.Value(this, &FPrimitiveComponentDetails::OnGetBodyMass)
-									.Visibility(this, &FPrimitiveComponentDetails::IsMassVisible, false)
+									MassInKGHandle->CreatePropertyValueWidget()
 								]
+							]
+						];
+					}
+				}
+				else if (PropName == GET_MEMBER_NAME_CHECKED(FBodyInstance, MaxAngularVelocity))
+				{
+					bMassUIExists = true;
+					TSharedPtr<IPropertyHandle> MaxAngularVelocityHandle = BodyInstanceHandler->GetChildHandle(PropName);
 
-								+ SVerticalBox::Slot()
+					PhysicsCategory.AddProperty(MaxAngularVelocityHandle).CustomWidget()
+					.NameContent()
+						[
+							MaxAngularVelocityHandle->CreatePropertyNameWidget()
+						]
+					.ValueContent()
+						[
+							SNew(SVerticalBox)
+							+ SVerticalBox::Slot()
+							.Padding(0.f, 0.f, 10.f, 0.f)
+							[
+								SNew(SNumericEntryBox<float>)
+								.IsEnabled(false)
+								.Font(IDetailLayoutBuilder::GetDetailFont())
+								.Value(this, &FPrimitiveComponentDetails::OnGetBodyMaxAngularVelocity)
+								.Visibility(this, &FPrimitiveComponentDetails::IsMaxAngularVelocityVisible, false)
+							]
+
+							+ SVerticalBox::Slot()
 								.Padding(0.f, 0.f, 10.f, 0.f)
 								[
 									SNew(SVerticalBox)
-									.Visibility(this, &FPrimitiveComponentDetails::IsMassVisible, true)
+									.Visibility(this, &FPrimitiveComponentDetails::IsMaxAngularVelocityVisible, true)
 									+ SVerticalBox::Slot()
 									.AutoHeight()
 									[
-										MassInKGHandle->CreatePropertyValueWidget()
+										MaxAngularVelocityHandle->CreatePropertyValueWidget()
 									]
 								]
-							]
-
-							+ SHorizontalBox::Slot()
-							.AutoWidth()
-							.VAlign( VAlign_Center )
-							[
-								bDisplayMassOverride ? MassOverrideHandle->CreatePropertyValueWidget() : StaticCastSharedRef<SWidget>(SNew(SSpacer))
-							]
-
-							+ SHorizontalBox::Slot()
-							.AutoWidth()
-							[
-								bDisplayMassOverride ? MassOverrideHandle->CreatePropertyNameWidget() : StaticCastSharedRef<SWidget>(SNew(SSpacer))
-							]
 						];
-					}
 				}
 				else
 				{
@@ -571,6 +593,24 @@ TOptional<float> FPrimitiveComponentDetails::OnGetBodyMass() const
 	return Mass;
 }
 
+
+TOptional<float> FPrimitiveComponentDetails::OnGetBodyMaxAngularVelocity() const
+{
+	UPrimitiveComponent* Comp = nullptr;
+	const float MaxAngularVelocity = UPhysicsSettings::Get()->MaxAngularVelocity;
+
+	for (auto ObjectIt = ObjectsCustomized.CreateConstIterator(); ObjectIt; ++ObjectIt)
+	{
+		if (ObjectIt->IsValid() && (*ObjectIt)->IsA(UPrimitiveComponent::StaticClass()))
+		{
+			Comp = Cast<UPrimitiveComponent>(ObjectIt->Get());
+			Comp->BodyInstance.MaxAngularVelocity = MaxAngularVelocity;	//update max angular velocity so that overriding gives the same value initially
+		}
+	}
+
+	return MaxAngularVelocity;
+}
+
 bool FPrimitiveComponentDetails::IsBodyMassReadOnly() const
 {
 	for (auto ObjectIt = ObjectsCustomized.CreateConstIterator(); ObjectIt; ++ObjectIt)
@@ -587,6 +627,22 @@ bool FPrimitiveComponentDetails::IsBodyMassReadOnly() const
 	return false;
 }
 
+bool FPrimitiveComponentDetails::IsMaxAngularVelocityReadOnly() const
+{
+	for (auto ObjectIt = ObjectsCustomized.CreateConstIterator(); ObjectIt; ++ObjectIt)
+	{
+		if (ObjectIt->IsValid() && (*ObjectIt)->IsA(UPrimitiveComponent::StaticClass()))
+		{
+			if (UPrimitiveComponent* Comp = Cast<UPrimitiveComponent>(ObjectIt->Get()))
+			{
+				if (Comp->BodyInstance.bOverrideMaxAngularVelocity == false) { return true; }
+			}
+		}
+	}
+
+	return false;
+}
+
 
 #undef LOCTEXT_NAMESPACE
 
diff --git a/Engine/Source/Editor/DetailCustomizations/Private/PrimitiveComponentDetails.h b/Engine/Source/Editor/DetailCustomizations/Private/PrimitiveComponentDetails.h
--- a/Engine/Source/Editor/DetailCustomizations/Private/PrimitiveComponentDetails.h
+++ b/Engine/Source/Editor/DetailCustomizations/Private/PrimitiveComponentDetails.h
@@ -46,6 +46,10 @@ private:
 	bool IsAutoWeldEditable() const;
 	EVisibility IsAutoWeldVisible() const;
 	EVisibility IsMassVisible(bool bOverrideMass) const;
+	EVisibility IsMaxAngularVelocityVisible(bool bOverrideMaxAngularVelocity) const;
+	TOptional<float> OnGetBodyMaxAngularVelocity() const;
+
+	bool IsMaxAngularVelocityReadOnly() const;
 
 };
 
diff --git a/Engine/Source/Runtime/Engine/Classes/PhysicsEngine/BodyInstance.h b/Engine/Source/Runtime/Engine/Classes/PhysicsEngine/BodyInstance.h
--- a/Engine/Source/Runtime/Engine/Classes/PhysicsEngine/BodyInstance.h
+++ b/Engine/Source/Runtime/Engine/Classes/PhysicsEngine/BodyInstance.h
@@ -208,7 +208,7 @@ public:
 	uint32 bSimulatePhysics:1;
 
 	/** If true, mass will not be automatically computed and you must set it directly */
-	UPROPERTY(EditAnywhere, Category = Physics, meta=(DisplayName="Override"))
+	UPROPERTY(meta=(DisplayName="Override"))
 	uint32 bOverrideMass : 1;
 
 	/**Mass of the body in KG. By default we compute this based on physical material and mass scale.
