Moved the SceneComponent Mobility details customization to a separate class so it can be used for the FoliageType default Mobility

diff --git a/Engine/Source/Editor/DetailCustomizations/Private/MobilityCustomization.cpp b/Engine/Source/Editor/DetailCustomizations/Private/MobilityCustomization.cpp
new file mode 100644
--- /dev/null
+++ b/Engine/Source/Editor/DetailCustomizations/Private/MobilityCustomization.cpp
@@ -0,0 +1,192 @@
+// Copyright 1998-2015 Epic Games, Inc. All Rights Reserved.
+
+#include "DetailCustomizationsPrivatePCH.h"
+#include "MobilityCustomization.h"
+
+#define LOCTEXT_NAMESPACE "MobilityCustomization"
+
+void FMobilityCustomization::CreateMobilityCustomization(IDetailCategoryBuilder& Category, TSharedPtr<IPropertyHandle> InMobilityHandle, uint8 RestrictedMobilityBits)
+{
+	MobilityHandle = InMobilityHandle;
+
+	TSharedPtr<SUniformGridPanel> ButtonOptionsPanel;
+		
+	IDetailPropertyRow& MobilityRow = Category.AddProperty(MobilityHandle);
+	MobilityRow.CustomWidget()
+	.NameContent()
+	[
+		SNew(STextBlock)
+		.Text(LOCTEXT("Mobility", "Mobility"))
+		.ToolTipText(this, &FMobilityCustomization::GetMobilityToolTip)
+		.Font(IDetailLayoutBuilder::GetDetailFont())
+	]
+	.ValueContent()
+	.MaxDesiredWidth(0)
+	[
+		SAssignNew(ButtonOptionsPanel, SUniformGridPanel)
+	];
+
+	bool bShowStatic = !( RestrictedMobilityBits & StaticMobilityBitMask );
+	bool bShowStationary = !( RestrictedMobilityBits & StationaryMobilityBitMask );
+
+	int32 ColumnIndex = 0;
+
+	if ( bShowStatic )
+	{
+		// Static Mobility
+		ButtonOptionsPanel->AddSlot(0, 0)
+		[
+			SNew(SCheckBox)
+			.Style(FEditorStyle::Get(), "Property.ToggleButton.Start")
+			.IsChecked(this, &FMobilityCustomization::IsMobilityActive, EComponentMobility::Static)
+			.OnCheckStateChanged(this, &FMobilityCustomization::OnMobilityChanged, EComponentMobility::Static)
+			.ToolTipText(LOCTEXT("Mobility_Static_Tooltip", "A static object can't be changed in game.\n¡Ü Allows Baked Lighting\n¡Ü Fastest Rendering"))
+			[
+				SNew(SHorizontalBox)
+
+				+ SHorizontalBox::Slot()
+				.AutoWidth()
+				.VAlign(VAlign_Center)
+				.Padding(3, 2)
+				[
+					SNew(SImage)
+					.Image(FEditorStyle::GetBrush("Mobility.Static"))
+				]
+
+				+ SHorizontalBox::Slot()
+				.FillWidth(1.0f)
+				.VAlign(VAlign_Center)
+				.HAlign(HAlign_Center)
+				.Padding(6, 2)
+				[
+					SNew(STextBlock)
+					.Text(LOCTEXT("Static", "Static"))
+					.Font(IDetailLayoutBuilder::GetDetailFont())
+					.ColorAndOpacity(this, &FMobilityCustomization::GetMobilityTextColor, EComponentMobility::Static)
+				]
+			]
+		];
+
+		ColumnIndex++;
+	}
+
+	// Stationary Mobility
+	if ( bShowStationary )
+	{
+		ButtonOptionsPanel->AddSlot(ColumnIndex, 0)
+		[
+			SNew(SCheckBox)
+			.IsChecked(this, &FMobilityCustomization::IsMobilityActive, EComponentMobility::Stationary)
+			.Style(FEditorStyle::Get(), ( ColumnIndex == 0 ) ? "Property.ToggleButton.Start" : "Property.ToggleButton.Middle")
+			.OnCheckStateChanged(this, &FMobilityCustomization::OnMobilityChanged, EComponentMobility::Stationary)
+			.ToolTipText(LOCTEXT("Mobility_Stationary_Tooltip", "A stationary light will only have its shadowing and bounced lighting from static geometry baked by Lightmass, all other lighting will be dynamic.  It can change color and intensity in game.\n¡Ü Can't Move\n¡Ü Allows Partial Baked Lighting\n¡Ü Dynamic Shadows"))
+			[
+				SNew(SHorizontalBox)
+
+				+ SHorizontalBox::Slot()
+				.AutoWidth()
+				.VAlign(VAlign_Center)
+				.Padding(3, 2)
+				[
+					SNew(SImage)
+					.Image(FEditorStyle::GetBrush("Mobility.Stationary"))
+				]
+
+				+ SHorizontalBox::Slot()
+				.FillWidth(1.0f)
+				.VAlign(VAlign_Center)
+				.HAlign(HAlign_Center)
+				.Padding(6, 2)
+				[
+					SNew(STextBlock)
+					.Text(LOCTEXT("Stationary", "Stationary"))
+					.Font(IDetailLayoutBuilder::GetDetailFont())
+					.ColorAndOpacity(this, &FMobilityCustomization::GetMobilityTextColor, EComponentMobility::Stationary)
+				]
+			]
+		];
+
+		ColumnIndex++;
+	}
+
+	// Movable Mobility
+	ButtonOptionsPanel->AddSlot(ColumnIndex, 0)
+	[
+		SNew(SCheckBox)
+		.IsChecked(this, &FMobilityCustomization::IsMobilityActive, EComponentMobility::Movable)
+		.Style(FEditorStyle::Get(), ( ColumnIndex == 0 ) ? "Property.ToggleButton" : "Property.ToggleButton.End")
+		.OnCheckStateChanged(this, &FMobilityCustomization::OnMobilityChanged, EComponentMobility::Movable)
+		.ToolTipText(LOCTEXT("Mobility_Movable_Tooltip", "Movable objects can be moved and changed in game.\n¡Ü Totally Dynamic\n¡Ü Allows Dynamic Shadows\n¡Ü Slowest Rendering"))
+		[
+			SNew(SHorizontalBox)
+
+			+ SHorizontalBox::Slot()
+			.AutoWidth()
+			.VAlign(VAlign_Center)
+			.Padding(3, 2)
+			[
+				SNew(SImage)
+				.Image(FEditorStyle::GetBrush("Mobility.Movable"))
+			]
+
+			+ SHorizontalBox::Slot()
+			.FillWidth(1.0f)
+			.VAlign(VAlign_Center)
+			.HAlign(HAlign_Center)
+			.Padding(6, 2)
+			[
+				SNew(STextBlock)
+				.Text(LOCTEXT("Movable", "Movable"))
+				.Font(IDetailLayoutBuilder::GetDetailFont())
+				.ColorAndOpacity(this, &FMobilityCustomization::GetMobilityTextColor, EComponentMobility::Movable)
+			]
+		]
+	];
+}
+
+ECheckBoxState FMobilityCustomization::IsMobilityActive(EComponentMobility::Type InMobility) const
+{
+	if (MobilityHandle.IsValid())
+	{
+		uint8 MobilityByte;
+		MobilityHandle->GetValue(MobilityByte);
+
+		return MobilityByte == InMobility ? ECheckBoxState::Checked : ECheckBoxState::Unchecked;
+	}
+
+	return ECheckBoxState::Unchecked;
+}
+
+FSlateColor FMobilityCustomization::GetMobilityTextColor(EComponentMobility::Type InMobility) const
+{
+	if (MobilityHandle.IsValid())
+	{
+		uint8 MobilityByte;
+		MobilityHandle->GetValue(MobilityByte);
+
+		return MobilityByte == InMobility ? FSlateColor(FLinearColor(0, 0, 0)) : FSlateColor(FLinearColor(0.72f, 0.72f, 0.72f, 1.f));
+	}
+
+	return FSlateColor(FLinearColor(0.72f, 0.72f, 0.72f, 1.f));
+}
+
+void FMobilityCustomization::OnMobilityChanged(ECheckBoxState InCheckedState, EComponentMobility::Type InMobility)
+{
+	if (MobilityHandle.IsValid() && InCheckedState == ECheckBoxState::Checked)
+	{
+		MobilityHandle->SetValue((uint8)InMobility);
+	}
+}
+
+FText FMobilityCustomization::GetMobilityToolTip() const
+{
+	if (MobilityHandle.IsValid())
+	{
+		return MobilityHandle->GetToolTipText();
+	}
+
+	return FText::GetEmpty();
+}
+
+#undef LOCTEXT_NAMESPACE
+
diff --git a/Engine/Source/Editor/DetailCustomizations/Private/SceneComponentDetails.cpp b/Engine/Source/Editor/DetailCustomizations/Private/SceneComponentDetails.cpp
--- a/Engine/Source/Editor/DetailCustomizations/Private/SceneComponentDetails.cpp
+++ b/Engine/Source/Editor/DetailCustomizations/Private/SceneComponentDetails.cpp
@@ -13,6 +13,7 @@
 #include "Engine/InheritableComponentHandler.h"
 #include "Components/LightComponentBase.h"
 #include "ComponentUtils.h"
+#include "MobilityCustomization.h"
 
 #define LOCTEXT_NAMESPACE "SceneComponentDetails"
 
@@ -141,65 +142,15 @@ TSharedRef<IDetailCustomization> FSceneComponentDetails::MakeInstance()
 	return MakeShareable( new FSceneComponentDetails );
 }
 
-ECheckBoxState FSceneComponentDetails::IsMobilityActive(EComponentMobility::Type InMobility) const
-{
-	if ( MobilityHandle.IsValid() )
-	{
-		uint8 MobilityByte;
-		MobilityHandle->GetValue(MobilityByte);
-
-		return MobilityByte == InMobility ? ECheckBoxState::Checked : ECheckBoxState::Unchecked;
-	}
-
-	return ECheckBoxState::Unchecked;
-}
-
-FSlateColor FSceneComponentDetails::GetMobilityTextColor(EComponentMobility::Type InMobility) const
-{
-	if ( MobilityHandle.IsValid() )
-	{
-		uint8 MobilityByte;
-		MobilityHandle->GetValue(MobilityByte);
-
-		return MobilityByte == InMobility ? FSlateColor(FLinearColor(0, 0, 0)) : FSlateColor(FLinearColor(0.72f, 0.72f, 0.72f, 1.f));
-	}
-
-	return FSlateColor(FLinearColor(0.72f, 0.72f, 0.72f, 1.f));
-}
-
-void FSceneComponentDetails::OnMobilityChanged(ECheckBoxState InCheckedState, EComponentMobility::Type InMobility)
-{
-	if ( MobilityHandle.IsValid() && InCheckedState == ECheckBoxState::Checked)
-	{
-		MobilityHandle->SetValue((uint8)InMobility);
-	}
-}
-
-FText FSceneComponentDetails::GetMobilityToolTip() const
-{
-	if ( MobilityHandle.IsValid() )
-	{
-		return MobilityHandle->GetToolTipText();
-	}
-
-	return FText::GetEmpty();
-}
-
 void FSceneComponentDetails::CustomizeDetails( IDetailLayoutBuilder& DetailBuilder )
 {
 	MakeTransformDetails( DetailBuilder );
 
 	// Put mobility property in Transform section
 	IDetailCategoryBuilder& TransformCategory = DetailBuilder.EditCategory( "TransformCommon", LOCTEXT("TransformCommonCategory", "Transform"), ECategoryPriority::Transform );
-	MobilityHandle = DetailBuilder.GetProperty("Mobility");
+	TSharedPtr<IPropertyHandle> MobilityHandle = DetailBuilder.GetProperty("Mobility");
 
 	uint8 RestrictedMobilityBits = 0u;
-	enum 
-	{
-		StaticMobilityBitMask     = (1u << EComponentMobility::Static),
-		StationaryMobilityBitMask = (1u << EComponentMobility::Stationary),
-		MovableMobilityBitMask    = (1u << EComponentMobility::Movable),
-	};
 
 	TArray< TWeakObjectPtr<UObject> > SelectedSceneComponents;
 
@@ -220,7 +171,7 @@ void FSceneComponentDetails::CustomizeDetails( IDetailLayoutBuilder& DetailBuild
 		}
 
 		// if we haven't restricted the "Static" option yet
-		if (!(RestrictedMobilityBits & StaticMobilityBitMask))
+		if (!(RestrictedMobilityBits & FMobilityCustomization::StaticMobilityBitMask))
 		{
 			FText RestrictReason;
 			if (IsMobilitySettingProhibited(EComponentMobility::Static, SceneComponent, RestrictReason))
@@ -229,12 +180,12 @@ void FSceneComponentDetails::CustomizeDetails( IDetailLayoutBuilder& DetailBuild
 				 StaticRestriction->AddValue(TEXT("Static"));
 				 MobilityHandle->AddRestriction(StaticRestriction.ToSharedRef());
 
-				 RestrictedMobilityBits |= StaticMobilityBitMask;
+				 RestrictedMobilityBits |= FMobilityCustomization::StaticMobilityBitMask;
 			}			
 		}
 
 		// if we haven't restricted the "Stationary" option yet
-		if (!(RestrictedMobilityBits & StationaryMobilityBitMask))
+		if (!(RestrictedMobilityBits & FMobilityCustomization::StationaryMobilityBitMask))
 		{
 			FText RestrictReason;
 			if (IsMobilitySettingProhibited(EComponentMobility::Stationary, SceneComponent, RestrictReason))
@@ -243,150 +194,19 @@ void FSceneComponentDetails::CustomizeDetails( IDetailLayoutBuilder& DetailBuild
 				StationaryRestriction->AddValue(TEXT("Stationary"));
 				MobilityHandle->AddRestriction(StationaryRestriction.ToSharedRef());
 
-				RestrictedMobilityBits |= StationaryMobilityBitMask;
+				RestrictedMobilityBits |= FMobilityCustomization::StationaryMobilityBitMask;
 			}
 		}
 
 		// no need to go through all of them if we can't restrict any more
-		if ((RestrictedMobilityBits & StaticMobilityBitMask) && (RestrictedMobilityBits & StationaryMobilityBitMask))
+		if ((RestrictedMobilityBits & FMobilityCustomization::StaticMobilityBitMask) && (RestrictedMobilityBits & FMobilityCustomization::StationaryMobilityBitMask))
 		{
 			break;
 		}
 	}
 
-	TSharedPtr<SUniformGridPanel> ButtonOptionsPanel;
-
-	IDetailPropertyRow& MobilityRow = TransformCategory.AddProperty(MobilityHandle);
-	MobilityRow.CustomWidget()
-	.NameContent()
-	[
-		SNew(STextBlock)
-		.Text(LOCTEXT("Mobility", "Mobility"))
-		.ToolTipText(this, &FSceneComponentDetails::GetMobilityToolTip)
-		.Font(IDetailLayoutBuilder::GetDetailFont())
-	]
-	.ValueContent()
-	.MaxDesiredWidth(0)
-	[
-		SAssignNew(ButtonOptionsPanel, SUniformGridPanel)
-	];
-
-	bool bShowStatic = !( RestrictedMobilityBits & StaticMobilityBitMask );
-	bool bShowStationary = !( RestrictedMobilityBits & StationaryMobilityBitMask );
-
-	int32 ColumnIndex = 0;
-
-	if ( bShowStatic )
-	{
-		// Static Mobility
-		ButtonOptionsPanel->AddSlot(0, 0)
-		[
-			SNew(SCheckBox)
-			.Style(FEditorStyle::Get(), "Property.ToggleButton.Start")
-			.IsChecked(this, &FSceneComponentDetails::IsMobilityActive, EComponentMobility::Static)
-			.OnCheckStateChanged(this, &FSceneComponentDetails::OnMobilityChanged, EComponentMobility::Static)
-			.ToolTipText(LOCTEXT("Mobility_Static_Tooltip", "A static object can't be changed in game.\n● Allows Baked Lighting\n● Fastest Rendering"))
-			[
-				SNew(SHorizontalBox)
-
-				+ SHorizontalBox::Slot()
-				.AutoWidth()
-				.VAlign(VAlign_Center)
-				.Padding(3, 2)
-				[
-					SNew(SImage)
-					.Image(FEditorStyle::GetBrush("Mobility.Static"))
-				]
-
-				+ SHorizontalBox::Slot()
-				.FillWidth(1.0f)
-				.VAlign(VAlign_Center)
-				.HAlign(HAlign_Center)
-				.Padding(6, 2)
-				[
-					SNew(STextBlock)
-					.Text(LOCTEXT("Static", "Static"))
-					.Font(IDetailLayoutBuilder::GetDetailFont())
-					.ColorAndOpacity(this, &FSceneComponentDetails::GetMobilityTextColor, EComponentMobility::Static)
-				]
-			]
-		];
-
-		ColumnIndex++;
-	}
-
-	// Stationary Mobility
-	if ( bShowStationary )
-	{
-		ButtonOptionsPanel->AddSlot(ColumnIndex, 0)
-		[
-			SNew(SCheckBox)
-			.IsChecked(this, &FSceneComponentDetails::IsMobilityActive, EComponentMobility::Stationary)
-			.Style(FEditorStyle::Get(), ( ColumnIndex == 0 ) ? "Property.ToggleButton.Start" : "Property.ToggleButton.Middle")
-			.OnCheckStateChanged(this, &FSceneComponentDetails::OnMobilityChanged, EComponentMobility::Stationary)
-			.ToolTipText(LOCTEXT("Mobility_Stationary_Tooltip", "A stationary light will only have its shadowing and bounced lighting from static geometry baked by Lightmass, all other lighting will be dynamic.  It can change color and intensity in game.\n● Can't Move\n● Allows Partial Baked Lighting\n● Dynamic Shadows"))
-			[
-				SNew(SHorizontalBox)
-
-				+ SHorizontalBox::Slot()
-				.AutoWidth()
-				.VAlign(VAlign_Center)
-				.Padding(3, 2)
-				[
-					SNew(SImage)
-					.Image(FEditorStyle::GetBrush("Mobility.Stationary"))
-				]
-
-				+ SHorizontalBox::Slot()
-				.FillWidth(1.0f)
-				.VAlign(VAlign_Center)
-				.HAlign(HAlign_Center)
-				.Padding(6, 2)
-				[
-					SNew(STextBlock)
-					.Text(LOCTEXT("Stationary", "Stationary"))
-					.Font(IDetailLayoutBuilder::GetDetailFont())
-					.ColorAndOpacity(this, &FSceneComponentDetails::GetMobilityTextColor, EComponentMobility::Stationary)
-				]
-			]
-		];
-
-		ColumnIndex++;
-	}
-
-	// Movable Mobility
-	ButtonOptionsPanel->AddSlot(ColumnIndex, 0)
-	[
-		SNew(SCheckBox)
-		.IsChecked(this, &FSceneComponentDetails::IsMobilityActive, EComponentMobility::Movable)
-		.Style(FEditorStyle::Get(), ( ColumnIndex == 0 ) ? "Property.ToggleButton" : "Property.ToggleButton.End")
-		.OnCheckStateChanged(this, &FSceneComponentDetails::OnMobilityChanged, EComponentMobility::Movable)
-		.ToolTipText(LOCTEXT("Mobility_Movable_Tooltip", "Movable objects can be moved and changed in game.\n● Totally Dynamic\n● Allows Dynamic Shadows\n● Slowest Rendering"))
-		[
-			SNew(SHorizontalBox)
-
-			+ SHorizontalBox::Slot()
-			.AutoWidth()
-			.VAlign(VAlign_Center)
-			.Padding(3, 2)
-			[
-				SNew(SImage)
-				.Image(FEditorStyle::GetBrush("Mobility.Movable"))
-			]
-
-			+ SHorizontalBox::Slot()
-			.FillWidth(1.0f)
-			.VAlign(VAlign_Center)
-			.HAlign(HAlign_Center)
-			.Padding(6, 2)
-			[
-				SNew(STextBlock)
-				.Text(LOCTEXT("Movable", "Movable"))
-				.Font(IDetailLayoutBuilder::GetDetailFont())
-				.ColorAndOpacity(this, &FSceneComponentDetails::GetMobilityTextColor, EComponentMobility::Movable)
-			]
-		]
-	];
+	MobilityCustomization = MakeShareable(new FMobilityCustomization);
+	MobilityCustomization->CreateMobilityCustomization(TransformCategory, MobilityHandle, RestrictedMobilityBits);
 }
 
 void FSceneComponentDetails::MakeTransformDetails( IDetailLayoutBuilder& DetailBuilder )
diff --git a/Engine/Source/Editor/DetailCustomizations/Private/SceneComponentDetails.h b/Engine/Source/Editor/DetailCustomizations/Private/SceneComponentDetails.h
--- a/Engine/Source/Editor/DetailCustomizations/Private/SceneComponentDetails.h
+++ b/Engine/Source/Editor/DetailCustomizations/Private/SceneComponentDetails.h
@@ -2,6 +2,8 @@
 
 #pragma once
 
+#include "MobilityCustomization.h"
+
 class FSceneComponentDetails : public IDetailCustomization
 {
 public:
@@ -14,13 +16,5 @@ public:
 private:
 	void MakeTransformDetails( IDetailLayoutBuilder& DetailBuilder );
 
-	FSlateColor GetMobilityTextColor(EComponentMobility::Type InMobility) const;
-
-	ECheckBoxState IsMobilityActive(EComponentMobility::Type InMobility) const;
-
-	void OnMobilityChanged(ECheckBoxState InCheckedState, EComponentMobility::Type InMobility);
-
-	FText GetMobilityToolTip() const;
-private:
-	TSharedPtr<IPropertyHandle> MobilityHandle;
+	TSharedPtr<FMobilityCustomization> MobilityCustomization;
 };
diff --git a/Engine/Source/Editor/DetailCustomizations/Public/Customizations/MobilityCustomization.h b/Engine/Source/Editor/DetailCustomizations/Public/Customizations/MobilityCustomization.h
new file mode 100644
--- /dev/null
+++ b/Engine/Source/Editor/DetailCustomizations/Public/Customizations/MobilityCustomization.h
@@ -0,0 +1,31 @@
+// Copyright 1998-2015 Epic Games, Inc. All Rights Reserved.
+#pragma once
+
+// Helper class to create a Mobility customization for the specified Property in the specified CategoryBuilder.
+class DETAILCUSTOMIZATIONS_API FMobilityCustomization : public TSharedFromThis<FMobilityCustomization>
+{
+public:
+	enum
+	{
+		StaticMobilityBitMask = (1u << EComponentMobility::Static),
+		StationaryMobilityBitMask = (1u << EComponentMobility::Stationary),
+		MovableMobilityBitMask = (1u << EComponentMobility::Movable),
+	};
+
+	FMobilityCustomization()
+	{}
+
+	void CreateMobilityCustomization(IDetailCategoryBuilder& InCategoryBuilder, TSharedPtr<IPropertyHandle> InMobilityHandle, uint8 InRestrictedMobilityBits);
+
+private:
+
+	ECheckBoxState IsMobilityActive(EComponentMobility::Type InMobility) const;
+	FSlateColor GetMobilityTextColor(EComponentMobility::Type InMobility) const;
+	void OnMobilityChanged(ECheckBoxState InCheckedState, EComponentMobility::Type InMobility);
+	FText GetMobilityToolTip() const;
+	
+	TSharedPtr<IPropertyHandle> MobilityHandle;
+};
+
+#undef LOCTEXT_NAMESPACE
+
diff --git a/Engine/Source/Editor/FoliageEdit/FoliageEdit.build.cs b/Engine/Source/Editor/FoliageEdit/FoliageEdit.build.cs
--- a/Engine/Source/Editor/FoliageEdit/FoliageEdit.build.cs
+++ b/Engine/Source/Editor/FoliageEdit/FoliageEdit.build.cs
@@ -20,6 +20,7 @@ public class FoliageEdit : ModuleRules
 				"LevelEditor",
 				"Landscape",
                 "PropertyEditor",
+                "DetailCustomizations",
                 "AssetTools",
                 "Foliage",
 			}
diff --git a/Engine/Source/Editor/FoliageEdit/Private/FoliageTypePaintingCustomization.cpp b/Engine/Source/Editor/FoliageEdit/Private/FoliageTypePaintingCustomization.cpp
--- a/Engine/Source/Editor/FoliageEdit/Private/FoliageTypePaintingCustomization.cpp
+++ b/Engine/Source/Editor/FoliageEdit/Private/FoliageTypePaintingCustomization.cpp
@@ -6,6 +6,7 @@
 #include "FoliageEdMode.h"
 #include "FoliageTypePaintingCustomization.h"
 #include "FoliageTypeCustomizationHelpers.h"
+#include "MobilityCustomization.h"
 
 #define LOCTEXT_NAMESPACE "FoliageEd_Mode"
 
@@ -65,6 +66,9 @@ void FFoliageTypePaintingCustomization::ShowFoliagePropertiesForCategory(IDetail
 	// Properties with a HideBehind property specified should only be shown if that property is true, non-zero, or not empty
 	static const FName HideBehindKey("HideBehind");
 
+	// Mobility property name
+	static const FName MobilityName("Mobility");
+	
 	IDetailCategoryBuilder& CategoryBuilder = DetailLayoutBuilder.EditCategory(CategoryName);
 	TArray<TSharedRef<IPropertyHandle>> CategoryProperties;
 	CategoryBuilder.GetDefaultProperties(CategoryProperties, true, true);
@@ -75,45 +79,53 @@ void FFoliageTypePaintingCustomization::ShowFoliagePropertiesForCategory(IDetail
 		bool bShowingProperty = false;
 		if (UProperty* Property = PropertyHandle->GetProperty())
 		{
-			// Check to see if this property can be reapplied
-			TSharedPtr<IPropertyHandle> ReapplyConditionPropertyHandle = DetailLayoutBuilder.GetProperty(*Property->GetMetaData(ReapplyConditionKey));
-			if (ReapplyConditionPropertyHandle.IsValid() && ReapplyConditionPropertyHandle->IsValidHandle())
+			if (Property->GetFName() == MobilityName)
 			{
-				// Create a custom entry that allows explicit enabling/disabling of the property when reapplying
-				TSharedPtr<IPropertyHandle> PropertyHandlePtr = PropertyHandle;
-				OutDetailRowsByPropertyName.FindOrAdd(PropertyHandle->GetProperty()->GetFName()) =
-					&AddFoliageProperty(CategoryBuilder, PropertyHandlePtr, ReapplyConditionPropertyHandle, TAttribute<EVisibility>(), TAttribute<bool>());
+				MobilityCustomization = MakeShareable(new FMobilityCustomization);
+				MobilityCustomization->CreateMobilityCustomization(CategoryBuilder, DetailLayoutBuilder.GetProperty(MobilityName), FMobilityCustomization::StationaryMobilityBitMask);
 			}
 			else
 			{
-				TSharedPtr<IPropertyHandle> InvalidProperty;
-				TSharedPtr<IPropertyHandle> PropertyHandlePtr = PropertyHandle;
-
-				// Check to see if this property is hidden behind another
-				TSharedPtr<IPropertyHandle> HiddenBehindPropertyHandle = DetailLayoutBuilder.GetProperty(*Property->GetMetaData(HideBehindKey));
-				if (HiddenBehindPropertyHandle.IsValid() && HiddenBehindPropertyHandle->IsValidHandle())
+				// Check to see if this property can be reapplied
+				TSharedPtr<IPropertyHandle> ReapplyConditionPropertyHandle = DetailLayoutBuilder.GetProperty(*Property->GetMetaData(ReapplyConditionKey));
+				if (ReapplyConditionPropertyHandle.IsValid() && ReapplyConditionPropertyHandle->IsValidHandle())
 				{
-					TAttribute<bool> IsEnabledAttribute;
-					ReapplyConditionPropertyHandle = DetailLayoutBuilder.GetProperty(*HiddenBehindPropertyHandle->GetProperty()->GetMetaData(ReapplyConditionKey));
-					if (ReapplyConditionPropertyHandle.IsValid() && ReapplyConditionPropertyHandle->IsValidHandle())
-					{
-						// If the property this is hidden behind has a reapply condition, disable this when the condition is false
-						IsEnabledAttribute = TAttribute<bool>::Create(TAttribute<bool>::FGetter::CreateSP(this, &FFoliageTypePaintingCustomization::IsReapplyPropertyEnabled, ReapplyConditionPropertyHandle));
-					}
-
-					TAttribute<EVisibility> VisibilityAttribute;
-					GetHiddenPropertyVisibility(HiddenBehindPropertyHandle, !IsEnabledAttribute.IsSet(), VisibilityAttribute);
-
+					// Create a custom entry that allows explicit enabling/disabling of the property when reapplying
+					TSharedPtr<IPropertyHandle> PropertyHandlePtr = PropertyHandle;
 					OutDetailRowsByPropertyName.FindOrAdd(PropertyHandle->GetProperty()->GetFName()) =
-						&AddFoliageProperty(CategoryBuilder, PropertyHandlePtr, InvalidProperty, VisibilityAttribute, IsEnabledAttribute);
+						&AddFoliageProperty(CategoryBuilder, PropertyHandlePtr, ReapplyConditionPropertyHandle, TAttribute<EVisibility>(), TAttribute<bool>());
 				}
 				else
 				{
-					// This property cannot be reapplied and isn't hidden behind anything, so show it whenever the reapply tool isn't active
-					OutDetailRowsByPropertyName.FindOrAdd(PropertyHandle->GetProperty()->GetFName()) =
-						&AddFoliageProperty(CategoryBuilder, PropertyHandlePtr, InvalidProperty,
-						TAttribute<EVisibility>::Create(TAttribute<EVisibility>::FGetter::CreateSP(this, &FFoliageTypePaintingCustomization::GetNonReapplyPropertyVisibility)),
-						TAttribute<bool>());
+					TSharedPtr<IPropertyHandle> InvalidProperty;
+					TSharedPtr<IPropertyHandle> PropertyHandlePtr = PropertyHandle;
+
+					// Check to see if this property is hidden behind another
+					TSharedPtr<IPropertyHandle> HiddenBehindPropertyHandle = DetailLayoutBuilder.GetProperty(*Property->GetMetaData(HideBehindKey));
+					if (HiddenBehindPropertyHandle.IsValid() && HiddenBehindPropertyHandle->IsValidHandle())
+					{
+						TAttribute<bool> IsEnabledAttribute;
+						ReapplyConditionPropertyHandle = DetailLayoutBuilder.GetProperty(*HiddenBehindPropertyHandle->GetProperty()->GetMetaData(ReapplyConditionKey));
+						if (ReapplyConditionPropertyHandle.IsValid() && ReapplyConditionPropertyHandle->IsValidHandle())
+						{
+							// If the property this is hidden behind has a reapply condition, disable this when the condition is false
+							IsEnabledAttribute = TAttribute<bool>::Create(TAttribute<bool>::FGetter::CreateSP(this, &FFoliageTypePaintingCustomization::IsReapplyPropertyEnabled, ReapplyConditionPropertyHandle));
+						}
+
+						TAttribute<EVisibility> VisibilityAttribute;
+						GetHiddenPropertyVisibility(HiddenBehindPropertyHandle, !IsEnabledAttribute.IsSet(), VisibilityAttribute);
+
+						OutDetailRowsByPropertyName.FindOrAdd(PropertyHandle->GetProperty()->GetFName()) =
+							&AddFoliageProperty(CategoryBuilder, PropertyHandlePtr, InvalidProperty, VisibilityAttribute, IsEnabledAttribute);
+					}
+					else
+					{
+						// This property cannot be reapplied and isn't hidden behind anything, so show it whenever the reapply tool isn't active
+						OutDetailRowsByPropertyName.FindOrAdd(PropertyHandle->GetProperty()->GetFName()) =
+							&AddFoliageProperty(CategoryBuilder, PropertyHandlePtr, InvalidProperty,
+							TAttribute<EVisibility>::Create(TAttribute<EVisibility>::FGetter::CreateSP(this, &FFoliageTypePaintingCustomization::GetNonReapplyPropertyVisibility)),
+							TAttribute<bool>());
+					}
 				}
 			}
 		}
diff --git a/Engine/Source/Editor/FoliageEdit/Private/FoliageTypePaintingCustomization.h b/Engine/Source/Editor/FoliageEdit/Private/FoliageTypePaintingCustomization.h
--- a/Engine/Source/Editor/FoliageEdit/Private/FoliageTypePaintingCustomization.h
+++ b/Engine/Source/Editor/FoliageEdit/Private/FoliageTypePaintingCustomization.h
@@ -2,6 +2,7 @@
 #pragma once
 
 #include "IDetailCustomization.h"
+#include "MobilityCustomization.h"
 
 class FEdModeFoliage;
 class IDetailCategoryBuilder;
@@ -49,4 +50,5 @@ private:
 
 	TSharedPtr<IPropertyHandle> Scaling;
 	TSharedPtr<IPropertyHandle> ReapplyScaling;
+	TSharedPtr<FMobilityCustomization> MobilityCustomization;
 };
