Reformat file as UTF-8 to allow fancy LOCTEXT characters to compile.

diff --git a/Engine/Source/Editor/DetailCustomizations/Private/MobilityCustomization.cpp b/Engine/Source/Editor/DetailCustomizations/Private/MobilityCustomization.cpp
--- a/Engine/Source/Editor/DetailCustomizations/Private/MobilityCustomization.cpp
+++ b/Engine/Source/Editor/DetailCustomizations/Private/MobilityCustomization.cpp
@@ -40,7 +40,7 @@ void FMobilityCustomization::CreateMobilityCustomization(IDetailCategoryBuilder&
 			.Style(FEditorStyle::Get(), "Property.ToggleButton.Start")
 			.IsChecked(this, &FMobilityCustomization::IsMobilityActive, EComponentMobility::Static)
 			.OnCheckStateChanged(this, &FMobilityCustomization::OnMobilityChanged, EComponentMobility::Static)
-			.ToolTipText(LOCTEXT("Mobility_Static_Tooltip", "A static object can't be changed in game.\n¡Ü Allows Baked Lighting\n¡Ü Fastest Rendering"))
+			.ToolTipText(LOCTEXT("Mobility_Static_Tooltip", "A static object can't be changed in game.\nâ— Allows Baked Lighting\nâ— Fastest Rendering"))
 			[
 				SNew(SHorizontalBox)
 
@@ -79,7 +79,7 @@ void FMobilityCustomization::CreateMobilityCustomization(IDetailCategoryBuilder&
 			.IsChecked(this, &FMobilityCustomization::IsMobilityActive, EComponentMobility::Stationary)
 			.Style(FEditorStyle::Get(), ( ColumnIndex == 0 ) ? "Property.ToggleButton.Start" : "Property.ToggleButton.Middle")
 			.OnCheckStateChanged(this, &FMobilityCustomization::OnMobilityChanged, EComponentMobility::Stationary)
-			.ToolTipText(LOCTEXT("Mobility_Stationary_Tooltip", "A stationary light will only have its shadowing and bounced lighting from static geometry baked by Lightmass, all other lighting will be dynamic.  It can change color and intensity in game.\n¡Ü Can't Move\n¡Ü Allows Partial Baked Lighting\n¡Ü Dynamic Shadows"))
+			.ToolTipText(LOCTEXT("Mobility_Stationary_Tooltip", "A stationary light will only have its shadowing and bounced lighting from static geometry baked by Lightmass, all other lighting will be dynamic.  It can change color and intensity in game.\nâ— Can't Move\nâ— Allows Partial Baked Lighting\nâ— Dynamic Shadows"))
 			[
 				SNew(SHorizontalBox)
 
@@ -116,7 +116,7 @@ void FMobilityCustomization::CreateMobilityCustomization(IDetailCategoryBuilder&
 		.IsChecked(this, &FMobilityCustomization::IsMobilityActive, EComponentMobility::Movable)
 		.Style(FEditorStyle::Get(), ( ColumnIndex == 0 ) ? "Property.ToggleButton" : "Property.ToggleButton.End")
 		.OnCheckStateChanged(this, &FMobilityCustomization::OnMobilityChanged, EComponentMobility::Movable)
-		.ToolTipText(LOCTEXT("Mobility_Movable_Tooltip", "Movable objects can be moved and changed in game.\n¡Ü Totally Dynamic\n¡Ü Allows Dynamic Shadows\n¡Ü Slowest Rendering"))
+		.ToolTipText(LOCTEXT("Mobility_Movable_Tooltip", "Movable objects can be moved and changed in game.\nâ— Totally Dynamic\nâ— Allows Dynamic Shadows\nâ— Slowest Rendering"))
 		[
 			SNew(SHorizontalBox)
 
