Fix non-unity compile error

diff --git a/Engine/Source/Editor/DetailCustomizations/Public/Customizations/MobilityCustomization.h b/Engine/Source/Editor/DetailCustomizations/Public/Customizations/MobilityCustomization.h
--- a/Engine/Source/Editor/DetailCustomizations/Public/Customizations/MobilityCustomization.h
+++ b/Engine/Source/Editor/DetailCustomizations/Public/Customizations/MobilityCustomization.h
@@ -1,6 +1,9 @@
 // Copyright 1998-2015 Epic Games, Inc. All Rights Reserved.
 #pragma once
 
+class IPropertyHandle;
+class IDetailCategoryBuilder;
+
 // Helper class to create a Mobility customization for the specified Property in the specified CategoryBuilder.
 class DETAILCUSTOMIZATIONS_API FMobilityCustomization : public TSharedFromThis<FMobilityCustomization>
 {
