Fix Phat not having the ability to override body mass

diff --git a/Engine/Source/Editor/DetailCustomizations/Private/BodySetupDetails.cpp b/Engine/Source/Editor/DetailCustomizations/Private/BodySetupDetails.cpp
--- a/Engine/Source/Editor/DetailCustomizations/Private/BodySetupDetails.cpp
+++ b/Engine/Source/Editor/DetailCustomizations/Private/BodySetupDetails.cpp
@@ -6,6 +6,7 @@
 #include "ObjectEditorUtils.h"
 #include "IDocumentation.h"
 #include "PhysicsEngine/BodySetup.h"
+#include "SNumericEntryBox.h"
 
 #define LOCTEXT_NAMESPACE "BodySetupDetails"
 
@@ -64,27 +65,26 @@ void FBodySetupDetails::CustomizeDetails( IDetailLayoutBuilder& DetailBuilder )
 				}
 				else if (ChildProperty->GetProperty()->GetFName() == GET_MEMBER_NAME_CHECKED(FBodyInstance, MassInKg))
 				{
-					PhysicsCategory.AddCustomRow(LOCTEXT("MassLabel", "Mass"), false)
-						.IsEnabled(TAttribute<bool>(this, &FBodySetupDetails::IsBodyMassEnabled))
-						.NameContent()
-						[
-							ChildProperty->CreatePropertyNameWidget()
-						]
+					PhysicsCategory.AddProperty(ChildProperty).CustomWidget()
+					.NameContent()
+					[
+						ChildProperty->CreatePropertyNameWidget()
+					]
 					.ValueContent()
 						[
 							SNew(SVerticalBox)
 							+ SVerticalBox::Slot()
-							.AutoHeight()
+							.Padding(0.f, 0.f, 10.f, 0.f)
 							[
-								SNew(SEditableTextBox)
-								.Text(this, &FBodySetupDetails::OnGetBodyMass)
-								.IsReadOnly(this, &FBodySetupDetails::IsBodyMassReadOnly)
+								SNew(SNumericEntryBox<float>)
+								.IsEnabled(false)
 								.Font(IDetailLayoutBuilder::GetDetailFont())
+								.Value(this, &FBodySetupDetails::OnGetBodyMass)
 								.Visibility(this, &FBodySetupDetails::IsMassVisible, false)
 							]
 
 							+ SVerticalBox::Slot()
-								.AutoHeight()
+								.Padding(0.f, 0.f, 10.f, 0.f)
 								[
 									SNew(SVerticalBox)
 									.Visibility(this, &FBodySetupDetails::IsMassVisible, true)
@@ -94,9 +94,7 @@ void FBodySetupDetails::CustomizeDetails( IDetailLayoutBuilder& DetailBuilder )
 										ChildProperty->CreatePropertyValueWidget()
 									]
 								]
-
 						];
-
 					continue;
 				}
 				if (Category == TEXT("Physics"))
@@ -112,7 +110,7 @@ void FBodySetupDetails::CustomizeDetails( IDetailLayoutBuilder& DetailBuilder )
 	}
 }
 
-FText FBodySetupDetails::OnGetBodyMass() const
+TOptional<float> FBodySetupDetails::OnGetBodyMass() const
 {
 	UBodySetup* BS = NULL;
 
@@ -140,10 +138,10 @@ FText FBodySetupDetails::OnGetBodyMass() const
 
 	if (bMultipleValue)
 	{
-		return LOCTEXT("MultipleValues", "Multiple Values");
+		return TOptional<float>();
 	}
 
-	return FText::AsNumber(Mass);
+	return Mass;
 }
 
 EVisibility FBodySetupDetails::IsMassVisible(bool bOverrideMass) const
diff --git a/Engine/Source/Editor/DetailCustomizations/Private/BodySetupDetails.h b/Engine/Source/Editor/DetailCustomizations/Private/BodySetupDetails.h
--- a/Engine/Source/Editor/DetailCustomizations/Private/BodySetupDetails.h
+++ b/Engine/Source/Editor/DetailCustomizations/Private/BodySetupDetails.h
@@ -11,7 +11,7 @@ public:
 	/** IDetailCustomization interface */
 	virtual void CustomizeDetails( IDetailLayoutBuilder& DetailBuilder ) override;
 
-	FText OnGetBodyMass() const;
+	TOptional<float> OnGetBodyMass() const;
 	bool IsBodyMassReadOnly() const;
 	bool IsBodyMassEnabled() const { return !IsBodyMassReadOnly(); }
 	EVisibility IsMassVisible(bool bOverrideMass) const;
