Reduce FBodyInstance size down from 472 to 352. Refactor customization code for primitive component and bodysetup so that it's all in one place. This was needed to make properties show up in human order vs best packing order

diff --git a/Engine/Source/Editor/DetailCustomizations/Private/BodyInstanceCustomization.cpp b/Engine/Source/Editor/DetailCustomizations/Private/BodyInstanceCustomization.cpp
--- a/Engine/Source/Editor/DetailCustomizations/Private/BodyInstanceCustomization.cpp
+++ b/Engine/Source/Editor/DetailCustomizations/Private/BodyInstanceCustomization.cpp
@@ -5,6 +5,8 @@
 #include "ScopedTransaction.h"
 #include "Editor/Documentation/Public/IDocumentation.h"
 #include "Engine/CollisionProfile.h"
+#include "SNumericEntryBox.h"
+#include "PhysicsEngine/BodySetup.h"
 
 #define LOCTEXT_NAMESPACE "BodyInstanceCustomization"
 
@@ -34,8 +36,7 @@ void FBodyInstanceCustomization::RefreshCollisionProfiles()
 	}
 }
 
-BEGIN_SLATE_FUNCTION_BUILD_OPTIMIZATION
-void FBodyInstanceCustomization::CustomizeChildren( TSharedRef<class IPropertyHandle> StructPropertyHandle, class IDetailChildrenBuilder& StructBuilder, IPropertyTypeCustomizationUtils& StructCustomizationUtils )
+void FBodyInstanceCustomization::AddCollisionCategory(TSharedRef<class IPropertyHandle> StructPropertyHandle, class IDetailChildrenBuilder& StructBuilder, IPropertyTypeCustomizationUtils& StructCustomizationUtils)
 {
 	CollisionProfileNameHandle = StructPropertyHandle->GetChildHandle(TEXT("CollisionProfileName"));
 	CollisionEnabledHandle = StructPropertyHandle->GetChildHandle(TEXT("CollisionEnabled"));
@@ -47,18 +48,6 @@ void FBodyInstanceCustomization::CustomizeChildren( TSharedRef<class IPropertyHa
 	check (CollisionEnabledHandle.IsValid());
 	check (ObjectTypeHandle.IsValid());
 
-	// copy all bodyinstances I'm accessing right now
-	TArray<void*> StructPtrs;
-	StructPropertyHandle->AccessRawData( StructPtrs );
-	check(StructPtrs.Num()!=0);
-
-	BodyInstances.AddUninitialized(StructPtrs.Num());
-	for (auto Iter = StructPtrs.CreateIterator(); Iter; ++Iter)
-	{
-		check (*Iter);
-		BodyInstances[Iter.GetIndex()] = (FBodyInstance*)(*Iter);
-	}
-
 	// need to find profile name
 	FName ProfileName;
 	TSharedPtr< FString > DisplayName = CollisionProfileComboList[0];
@@ -128,6 +117,24 @@ void FBodyInstanceCustomization::CustomizeChildren( TSharedRef<class IPropertyHa
 	// now create custom set up
 	CreateCustomCollisionSetup( StructPropertyHandle, CollisionGroup );
 }
+
+BEGIN_SLATE_FUNCTION_BUILD_OPTIMIZATION
+void FBodyInstanceCustomization::CustomizeChildren( TSharedRef<class IPropertyHandle> StructPropertyHandle, class IDetailChildrenBuilder& StructBuilder, IPropertyTypeCustomizationUtils& StructCustomizationUtils )
+{
+	// copy all bodyinstances I'm accessing right now
+	TArray<void*> StructPtrs;
+	StructPropertyHandle->AccessRawData(StructPtrs);
+	check(StructPtrs.Num() != 0);
+
+	BodyInstances.AddUninitialized(StructPtrs.Num());
+	for (auto Iter = StructPtrs.CreateIterator(); Iter; ++Iter)
+	{
+		check(*Iter);
+		BodyInstances[Iter.GetIndex()] = (FBodyInstance*)(*Iter);
+	}
+
+	AddCollisionCategory(StructPropertyHandle, StructBuilder, StructCustomizationUtils);
+}
 END_SLATE_FUNCTION_BUILD_OPTIMIZATION
 
 int32 FBodyInstanceCustomization::InitializeObjectTypeComboList()
@@ -952,6 +959,588 @@ void FBodyInstanceCustomization::SetCollisionResponseContainer(const FCollisionR
 	CollisionResponsesHandle->NotifyPostChange();
 }
 
+FBodyInstanceCustomizationHelper::FBodyInstanceCustomizationHelper(const TArray<TWeakObjectPtr<UObject>>& InObjectsCustomized)
+: ObjectsCustomized(InObjectsCustomized)
+{
+}
+
+void FBodyInstanceCustomizationHelper::UpdateFilters()
+{
+	bDisplayMass = true;
+	bDisplayConstraints = true;
+	bDisplayEnablePhysics = true;
+
+	for (int32 i = 0; i < ObjectsCustomized.Num(); ++i)
+	{
+		if (ObjectsCustomized[i].IsValid())
+		{
+			if(ObjectsCustomized[i]->IsA(UDestructibleComponent::StaticClass()))
+			{
+				bDisplayMass = false;
+				bDisplayConstraints = false;
+			}
+			else
+			{
+				if(ObjectsCustomized[i]->IsA(UBodySetup::StaticClass()))
+				{
+					bDisplayEnablePhysics = false;
+					bDisplayConstraints = false;
+				}
+			}
+		}
+	}
+}
+
+void FBodyInstanceCustomizationHelper::CustomizeDetails( IDetailLayoutBuilder& DetailBuilder, TSharedRef<IPropertyHandle> BodyInstanceHandler)
+{
+	UpdateFilters();
+
+	IDetailCategoryBuilder& PhysicsCategory = DetailBuilder.EditCategory("Physics");
+
+	TSharedRef<IPropertyHandle> PhysicsEnable = BodyInstanceHandler->GetChildHandle(GET_MEMBER_NAME_CHECKED(FBodyInstance, bSimulatePhysics)).ToSharedRef();
+	if(bDisplayEnablePhysics)
+	{
+		PhysicsCategory.AddProperty(PhysicsEnable).EditCondition(TAttribute<bool>(this, &FBodyInstanceCustomizationHelper::IsSimulatePhysicsEditable), NULL);
+	}else
+	{
+		PhysicsEnable->MarkHiddenByCustomization();
+	}
+
+	AddMassInKg(PhysicsCategory, BodyInstanceHandler);
+
+	PhysicsCategory.AddProperty(BodyInstanceHandler->GetChildHandle(GET_MEMBER_NAME_CHECKED(FBodyInstance, LinearDamping)));
+	PhysicsCategory.AddProperty(BodyInstanceHandler->GetChildHandle(GET_MEMBER_NAME_CHECKED(FBodyInstance, AngularDamping)));
+	PhysicsCategory.AddProperty(BodyInstanceHandler->GetChildHandle(GET_MEMBER_NAME_CHECKED(FBodyInstance, bEnableGravity)));
+
+	AddBodyConstraint(PhysicsCategory, BodyInstanceHandler);
+
+	//ADVANCED
+	PhysicsCategory.AddProperty(BodyInstanceHandler->GetChildHandle(GET_MEMBER_NAME_CHECKED(FBodyInstance, bAutoWeld)))
+	.Visibility(TAttribute<EVisibility>(this, &FBodyInstanceCustomizationHelper::IsAutoWeldVisible))
+	.EditCondition(TAttribute<bool>(this, &FBodyInstanceCustomizationHelper::IsAutoWeldEditable), NULL);
+
+	PhysicsCategory.AddProperty(BodyInstanceHandler->GetChildHandle(GET_MEMBER_NAME_CHECKED(FBodyInstance, bStartAwake)));
+
+	PhysicsCategory.AddProperty(BodyInstanceHandler->GetChildHandle(GET_MEMBER_NAME_CHECKED(FBodyInstance, COMNudge)));
+	PhysicsCategory.AddProperty(BodyInstanceHandler->GetChildHandle(GET_MEMBER_NAME_CHECKED(FBodyInstance, MassScale)));
+
+	AddMaxAngularVelocity(PhysicsCategory, BodyInstanceHandler);
+
+	TSharedRef<IPropertyHandle> AsyncEnabled = BodyInstanceHandler->GetChildHandle(GET_MEMBER_NAME_CHECKED(FBodyInstance, bUseAsyncScene)).ToSharedRef();
+	PhysicsCategory.AddProperty(AsyncEnabled).EditCondition(TAttribute<bool>(this, &FBodyInstanceCustomizationHelper::IsUseAsyncEditable), NULL);
+
+	//Add the rest
+	uint32 NumChildren = 0;
+	BodyInstanceHandler->GetNumChildren(NumChildren);
+	for(uint32 ChildIdx = 0; ChildIdx < NumChildren; ++ChildIdx)
+	{
+		TSharedPtr<IPropertyHandle> ChildProp = BodyInstanceHandler->GetChildHandle(ChildIdx);
+		
+		FString Category = FObjectEditorUtils::GetCategory(ChildProp->GetProperty());
+		if(ChildProp->IsCustomized() == false && Category == TEXT("Physics"))	//add the rest of the physics properties
+		{
+			PhysicsCategory.AddProperty(ChildProp);
+		}
+	}
+}
+
+bool FBodyInstanceCustomizationHelper::IsSimulatePhysicsEditable() const
+{
+	// Check whether to enable editing of bSimulatePhysics - this will happen if all objects are UPrimitiveComponents & have collision geometry.
+	bool bEnableSimulatePhysics = ObjectsCustomized.Num() > 0;
+	for (TWeakObjectPtr<UObject> CustomizedObject : ObjectsCustomized)
+	{
+		if (UPrimitiveComponent* PrimitiveComponent = Cast<UPrimitiveComponent>(CustomizedObject.Get()))
+		{
+			if (!PrimitiveComponent->CanEditSimulatePhysics())
+			{
+				bEnableSimulatePhysics = false;
+				break;
+			}
+		}
+	}
+
+	return bEnableSimulatePhysics;
+}
+
+bool FBodyInstanceCustomizationHelper::IsUseAsyncEditable() const
+{
+	// Check whether to enable editing of bUseAsyncScene - this will happen if all objects are movable and the project uses an AsyncScene
+	if (!UPhysicsSettings::Get()->bEnableAsyncScene)
+	{
+		return false;
+	}
+
+	bool bEnableUseAsyncScene = ObjectsCustomized.Num() > 0;
+	for (auto ObjectIt = ObjectsCustomized.CreateConstIterator(); ObjectIt; ++ObjectIt)
+	{
+		if (ObjectIt->IsValid() && (*ObjectIt)->IsA(UPrimitiveComponent::StaticClass()))
+		{
+			TWeakObjectPtr<USceneComponent> SceneComponent = CastChecked<USceneComponent>(ObjectIt->Get());
+
+			if (SceneComponent.IsValid() && SceneComponent->Mobility != EComponentMobility::Movable)
+			{
+				bEnableUseAsyncScene = false;
+				break;
+			}
+
+			// Skeletal mesh uses a physics asset which will have multiple bodies - these bodies have their own bUseAsyncScene which is what we actually use - the flag on the skeletal mesh is not used
+			TWeakObjectPtr<USkeletalMeshComponent> SkeletalMeshComponent = Cast<USkeletalMeshComponent>(ObjectIt->Get());
+			if (SkeletalMeshComponent.IsValid())
+			{
+				bEnableUseAsyncScene = false;
+				break;
+			}
+		}
+		else if (ObjectIt->IsValid() && (*ObjectIt)->IsA(UBodySetup::StaticClass()))
+		{
+			continue;
+		}
+		else
+		{
+			bEnableUseAsyncScene = false;
+			break;
+		}
+	}
+
+	return bEnableUseAsyncScene;
+}
+
+
+EVisibility FBodyInstanceCustomizationHelper::IsMassVisible(bool bOverrideMass) const
+{
+	bool bIsMassReadOnly = IsBodyMassReadOnly();
+	if (bOverrideMass)
+	{
+		return bIsMassReadOnly ? EVisibility::Collapsed : EVisibility::Visible;
+	}
+	else
+	{
+		return bIsMassReadOnly ? EVisibility::Visible : EVisibility::Collapsed;
+	}
+}
+
+bool FBodyInstanceCustomizationHelper::IsBodyMassReadOnly() const
+{
+	for (auto ObjectIt = ObjectsCustomized.CreateConstIterator(); ObjectIt; ++ObjectIt)
+	{
+		if (ObjectIt->IsValid() && (*ObjectIt)->IsA(UPrimitiveComponent::StaticClass()))
+		{
+			if (UPrimitiveComponent* Comp = Cast<UPrimitiveComponent>(ObjectIt->Get()))
+			{
+				if (Comp->BodyInstance.bOverrideMass == false) { return true; }
+			}
+		}else if(ObjectIt->IsValid() && (*ObjectIt)->IsA(UBodySetup::StaticClass()))
+		{
+			UBodySetup* BS = Cast<UBodySetup>(ObjectIt->Get());
+			if (BS->DefaultInstance.bOverrideMass == false)
+			{
+				return true;
+			}
+		}
+	}
+
+	return false;
+}
+
+TOptional<float> FBodyInstanceCustomizationHelper::OnGetBodyMaxAngularVelocity() const
+{
+	UPrimitiveComponent* Comp = nullptr;
+	const float MaxAngularVelocity = UPhysicsSettings::Get()->MaxAngularVelocity;
+
+	for (auto ObjectIt = ObjectsCustomized.CreateConstIterator(); ObjectIt; ++ObjectIt)
+	{
+		if (ObjectIt->IsValid() && (*ObjectIt)->IsA(UPrimitiveComponent::StaticClass()))
+		{
+			Comp = Cast<UPrimitiveComponent>(ObjectIt->Get());
+			Comp->BodyInstance.MaxAngularVelocity = MaxAngularVelocity;	//update max angular velocity so that overriding gives the same value initially
+		}
+	}
+
+	return MaxAngularVelocity;
+}
+
+bool FBodyInstanceCustomizationHelper::IsMaxAngularVelocityReadOnly() const
+{
+	for (auto ObjectIt = ObjectsCustomized.CreateConstIterator(); ObjectIt; ++ObjectIt)
+	{
+		if (ObjectIt->IsValid() && (*ObjectIt)->IsA(UPrimitiveComponent::StaticClass()))
+		{
+			if (UPrimitiveComponent* Comp = Cast<UPrimitiveComponent>(ObjectIt->Get()))
+			{
+				if (Comp->BodyInstance.bOverrideMaxAngularVelocity == false) { return true; }
+			}
+		}
+	}
+
+	return false;
+}
+
+EVisibility FBodyInstanceCustomizationHelper::IsMaxAngularVelocityVisible(bool bOverrideMaxAngularVelocity) const
+{
+	bool bIsMaxAngularVelocityReadOnly = IsMaxAngularVelocityReadOnly();
+	if (bOverrideMaxAngularVelocity)
+	{
+		return bIsMaxAngularVelocityReadOnly ? EVisibility::Collapsed : EVisibility::Visible;
+	}
+	else
+	{
+		return bIsMaxAngularVelocityReadOnly ? EVisibility::Visible : EVisibility::Collapsed;
+	}
+}
+
+bool FBodyInstanceCustomizationHelper::IsAutoWeldEditable() const
+{
+	for (int32 i = 0; i < ObjectsCustomized.Num(); ++i)
+	{
+		if (UPrimitiveComponent * SceneComponent = Cast<UPrimitiveComponent>(ObjectsCustomized[i].Get()))
+		{
+			if (FBodyInstance* BI = SceneComponent->GetBodyInstance())
+			{
+				if (BI->ShouldInstanceSimulatingPhysics())
+				{
+					return false;
+				}
+			}
+		}
+	}
+
+	return true;
+}
+
+EVisibility FBodyInstanceCustomizationHelper::IsAutoWeldVisible() const
+{
+	for (int32 i = 0; i < ObjectsCustomized.Num(); ++i)
+	{
+		if (ObjectsCustomized[i].IsValid() && !(ObjectsCustomized[i]->IsA(UStaticMeshComponent::StaticClass()) || ObjectsCustomized[i]->IsA(UShapeComponent::StaticClass())))
+		{
+			return EVisibility::Collapsed;
+		}
+	}
+
+	return EVisibility::Visible;
+}
+
+
+TOptional<float> FBodyInstanceCustomizationHelper::OnGetBodyMass() const
+{
+	UPrimitiveComponent* Comp = nullptr;
+	UBodySetup* BS = nullptr;
+
+	float Mass = 0.0f;
+	bool bMultipleValue = false;
+
+	for (auto ObjectIt = ObjectsCustomized.CreateConstIterator(); ObjectIt; ++ObjectIt)
+	{
+		if (ObjectIt->IsValid() && (*ObjectIt)->IsA(UPrimitiveComponent::StaticClass()))
+		{
+			Comp = Cast<UPrimitiveComponent>(ObjectIt->Get());
+
+			float CompMass = Comp->CalculateMass();
+			Comp->BodyInstance.MassInKg = CompMass;	//update the component's override mass to be the calculated one
+			if (Mass == 0.0f || FMath::Abs(Mass - CompMass) < 0.1f)
+			{
+				Mass = CompMass;
+			}
+			else
+			{
+				bMultipleValue = true;
+				break;
+			}
+		}
+		else if (ObjectIt->IsValid() && (*ObjectIt)->IsA(UBodySetup::StaticClass()))
+		{
+			BS = Cast<UBodySetup>(ObjectIt->Get());
+
+			float BSMass = BS->CalculateMass();
+			if (Mass == 0.0f || FMath::Abs(Mass - BSMass) < 0.1f)
+			{
+				Mass = BSMass;
+			}
+			else
+			{
+				bMultipleValue = true;
+				break;
+			}
+		}
+	}
+
+	if (bMultipleValue)
+	{
+		return TOptional<float>();
+	}
+
+	return Mass;
+}
+
+EVisibility FBodyInstanceCustomizationHelper::IsDOFMode(EDOFMode::Type Mode) const
+{
+	bool bVisible = false;
+	if (DOFModeProperty.IsValid() && bDisplayConstraints)
+	{
+		uint8 CurrentMode;
+		if (DOFModeProperty->GetValue(CurrentMode) == FPropertyAccess::Success)
+		{
+			EDOFMode::Type PropertyDOF = FBodyInstance::ResolveDOFMode(static_cast<EDOFMode::Type>(CurrentMode));
+			bVisible = PropertyDOF == Mode;
+		}
+	}
+
+	return bVisible ? EVisibility::Visible : EVisibility::Collapsed;
+}
+
+
+void FBodyInstanceCustomizationHelper::AddMassInKg(IDetailCategoryBuilder& PhysicsCategory, TSharedRef<IPropertyHandle> BodyInstanceHandler)
+{
+	if (bDisplayMass)
+	{
+		TSharedRef<IPropertyHandle> MassInKGHandle = BodyInstanceHandler->GetChildHandle(GET_MEMBER_NAME_CHECKED(FBodyInstance, MassInKg)).ToSharedRef();
+
+		PhysicsCategory.AddProperty(MassInKGHandle).CustomWidget()
+		.NameContent()
+		[
+			MassInKGHandle->CreatePropertyNameWidget()
+		]
+		.ValueContent()
+		[
+			SNew(SVerticalBox)
+			+ SVerticalBox::Slot()
+			.Padding(0.f, 0.f, 10.f, 0.f)
+			[
+				SNew(SNumericEntryBox<float>)
+				.IsEnabled(false)
+				.Font(IDetailLayoutBuilder::GetDetailFont())
+				.Value(this, &FBodyInstanceCustomizationHelper::OnGetBodyMass)
+				.Visibility(this, &FBodyInstanceCustomizationHelper::IsMassVisible, false)
+			]
+
+			+ SVerticalBox::Slot()
+			.Padding(0.f, 0.f, 10.f, 0.f)
+			[
+				SNew(SVerticalBox)
+				.Visibility(this, &FBodyInstanceCustomizationHelper::IsMassVisible, true)
+				+ SVerticalBox::Slot()
+				.AutoHeight()
+				[
+					MassInKGHandle->CreatePropertyValueWidget()
+				]
+			]
+		];
+	}
+}
+
+void FBodyInstanceCustomizationHelper::AddMaxAngularVelocity(IDetailCategoryBuilder& PhysicsCategory, TSharedRef<IPropertyHandle> BodyInstanceHandler)
+{
+	TSharedRef<IPropertyHandle> MaxAngularVelocityHandle = BodyInstanceHandler->GetChildHandle(GET_MEMBER_NAME_CHECKED(FBodyInstance, MaxAngularVelocity)).ToSharedRef();
+	
+	PhysicsCategory.AddProperty(MaxAngularVelocityHandle).CustomWidget()
+	.NameContent()
+	[
+		MaxAngularVelocityHandle->CreatePropertyNameWidget()
+	]
+	.ValueContent()
+	[
+		SNew(SVerticalBox)
+		+ SVerticalBox::Slot()
+		.Padding(0.f, 0.f, 10.f, 0.f)
+		[
+			SNew(SNumericEntryBox<float>)
+			.IsEnabled(false)
+			.Font(IDetailLayoutBuilder::GetDetailFont())
+			.Value(this, &FBodyInstanceCustomizationHelper::OnGetBodyMaxAngularVelocity)
+			.Visibility(this, &FBodyInstanceCustomizationHelper::IsMaxAngularVelocityVisible, false)
+		]
+
+		+ SVerticalBox::Slot()
+		.Padding(0.f, 0.f, 10.f, 0.f)
+		[
+			SNew(SVerticalBox)
+			.Visibility(this, &FBodyInstanceCustomizationHelper::IsMaxAngularVelocityVisible, true)
+			+ SVerticalBox::Slot()
+			.AutoHeight()
+			[
+				MaxAngularVelocityHandle->CreatePropertyValueWidget()
+			]
+		]
+	];
+}
+
+void FBodyInstanceCustomizationHelper::AddBodyConstraint(IDetailCategoryBuilder& PhysicsCategory, TSharedRef<IPropertyHandle> BodyInstanceHandler)
+{
+	const float XYZPadding = 5.f;
+	
+	TSharedPtr<IPropertyHandle> bLockXTranslation = BodyInstanceHandler->GetChildHandle(GET_MEMBER_NAME_CHECKED(FBodyInstance, bLockXTranslation));
+	bLockXTranslation->MarkHiddenByCustomization();
+
+	TSharedPtr<IPropertyHandle> bLockYTranslation = BodyInstanceHandler->GetChildHandle(GET_MEMBER_NAME_CHECKED(FBodyInstance, bLockYTranslation));
+	bLockYTranslation->MarkHiddenByCustomization();
+
+	TSharedPtr<IPropertyHandle> bLockZTranslation = BodyInstanceHandler->GetChildHandle(GET_MEMBER_NAME_CHECKED(FBodyInstance, bLockZTranslation));
+	bLockZTranslation->MarkHiddenByCustomization();
+
+	TSharedPtr<IPropertyHandle> bLockXRotation = BodyInstanceHandler->GetChildHandle(GET_MEMBER_NAME_CHECKED(FBodyInstance, bLockXRotation));
+	bLockXRotation->MarkHiddenByCustomization();
+
+	TSharedPtr<IPropertyHandle> bLockYRotation = BodyInstanceHandler->GetChildHandle(GET_MEMBER_NAME_CHECKED(FBodyInstance, bLockYRotation));
+	bLockYRotation->MarkHiddenByCustomization();
+
+	TSharedPtr<IPropertyHandle> bLockZRotation = BodyInstanceHandler->GetChildHandle(GET_MEMBER_NAME_CHECKED(FBodyInstance, bLockZRotation));
+	bLockZRotation->MarkHiddenByCustomization();
+
+	DOFModeProperty = BodyInstanceHandler->GetChildHandle(GET_MEMBER_NAME_CHECKED(FBodyInstance, DOFMode)).ToSharedRef();
+	DOFModeProperty->MarkHiddenByCustomization();
+
+	TSharedRef<IPropertyHandle> bLockRotation = BodyInstanceHandler->GetChildHandle(GET_MEMBER_NAME_CHECKED(FBodyInstance, bLockRotation)).ToSharedRef();
+	bLockRotation->MarkHiddenByCustomization();
+
+	TSharedRef<IPropertyHandle> bLockTranslation = BodyInstanceHandler->GetChildHandle(GET_MEMBER_NAME_CHECKED(FBodyInstance, bLockTranslation)).ToSharedRef();
+	bLockTranslation->MarkHiddenByCustomization();
+
+	TSharedRef<IPropertyHandle> CustomDOFPlaneNormal = BodyInstanceHandler->GetChildHandle(GET_MEMBER_NAME_CHECKED(FBodyInstance, CustomDOFPlaneNormal)).ToSharedRef();
+	CustomDOFPlaneNormal->MarkHiddenByCustomization();
+
+	//the above are all marked hidden even if we don't display constraints because the user wants to hide it anyway
+
+	if (bDisplayConstraints)
+	{
+		IDetailGroup& ConstraintsGroup = PhysicsCategory.AddGroup(TEXT("ConstraintsGroup"), LOCTEXT("Constraints", "Constraints"));
+
+		ConstraintsGroup.AddWidgetRow()
+		.Visibility(TAttribute<EVisibility>::Create(TAttribute<EVisibility>::FGetter::CreateSP(this, &FBodyInstanceCustomizationHelper::IsDOFMode, EDOFMode::SixDOF)))
+		.NameContent()
+		[
+			SNew(STextBlock)
+			.Text(LOCTEXT("LockPositionLabel", "Lock Position"))
+			.ToolTipText(LOCTEXT("LockPositionTooltip", "Locks movement along the specified axis"))
+			.Font(IDetailLayoutBuilder::GetDetailFont())
+		]
+		.ValueContent()
+		[
+			SNew(SHorizontalBox)
+			+ SHorizontalBox::Slot()
+			.Padding(0.f, 0.f, XYZPadding, 0.f)
+			.AutoWidth()
+			[
+				SNew(SHorizontalBox)
+				+ SHorizontalBox::Slot()
+				.AutoWidth()
+				[
+					bLockXTranslation->CreatePropertyNameWidget()
+				]
+				+ SHorizontalBox::Slot()
+				.AutoWidth()
+				[
+					bLockXTranslation->CreatePropertyValueWidget()
+				]
+			]
+
+			+ SHorizontalBox::Slot()
+			.Padding(0.f, 0.f, XYZPadding, 0.f)
+			.AutoWidth()
+			[
+				SNew(SHorizontalBox)
+				+ SHorizontalBox::Slot()
+				.AutoWidth()
+				[
+					bLockYTranslation->CreatePropertyNameWidget()
+				]
+				+ SHorizontalBox::Slot()
+				.AutoWidth()
+				[
+					bLockYTranslation->CreatePropertyValueWidget()
+				]
+			]
+
+			+ SHorizontalBox::Slot()
+			.Padding(0.f, 0.f, XYZPadding, 0.f)
+			.AutoWidth()
+			[
+				SNew(SHorizontalBox)
+				+ SHorizontalBox::Slot()
+				.AutoWidth()
+				[
+					bLockZTranslation->CreatePropertyNameWidget()
+				]
+				+ SHorizontalBox::Slot()
+				.AutoWidth()
+				[
+					bLockZTranslation->CreatePropertyValueWidget()
+				]
+			]
+		];
+
+		ConstraintsGroup.AddWidgetRow()
+		.Visibility(TAttribute<EVisibility>::Create(TAttribute<EVisibility>::FGetter::CreateSP(this, &FBodyInstanceCustomizationHelper::IsDOFMode, EDOFMode::SixDOF)))
+		.NameContent()
+		[
+			SNew(STextBlock)
+			.Text(LOCTEXT("LockRotationLabel", "Lock Rotation"))
+			.ToolTipText(LOCTEXT("LockRotationTooltip", "Locks rotation about the specified axis"))
+			.Font(IDetailLayoutBuilder::GetDetailFont())
+		]
+		.ValueContent()
+		[
+			SNew(SHorizontalBox)
+			+ SHorizontalBox::Slot()
+			.Padding(0.f, 0.f, XYZPadding, 0.f)
+			.AutoWidth()
+			[
+				SNew(SHorizontalBox)
+				+ SHorizontalBox::Slot()
+				.AutoWidth()
+				[
+					bLockXRotation->CreatePropertyNameWidget()
+				]
+				+ SHorizontalBox::Slot()
+				.AutoWidth()
+				[
+					bLockXRotation->CreatePropertyValueWidget()
+				]
+			]
+
+			+ SHorizontalBox::Slot()
+			.Padding(0.f, 0.f, XYZPadding, 0.f)
+			.AutoWidth()
+			[
+				SNew(SHorizontalBox)
+				+ SHorizontalBox::Slot()
+				.AutoWidth()
+				[
+					bLockYRotation->CreatePropertyNameWidget()
+				]
+				+ SHorizontalBox::Slot()
+				.AutoWidth()
+				[
+					bLockYRotation->CreatePropertyValueWidget()
+				]
+			]
+
+			+ SHorizontalBox::Slot()
+			.Padding(0.f, 0.f, XYZPadding, 0.f)
+			.AutoWidth()
+			[
+				SNew(SHorizontalBox)
+				+ SHorizontalBox::Slot()
+				.AutoWidth()
+				[
+					bLockZRotation->CreatePropertyNameWidget()
+				]
+				+ SHorizontalBox::Slot()
+				.AutoWidth()
+				[
+					bLockZRotation->CreatePropertyValueWidget()
+				]
+			]
+		];
+
+		//we only show the custom plane normal if we've selected that mode
+		ConstraintsGroup.AddPropertyRow(CustomDOFPlaneNormal).Visibility(TAttribute<EVisibility>::Create(TAttribute<EVisibility>::FGetter::CreateSP(this, &FBodyInstanceCustomizationHelper::IsDOFMode, EDOFMode::CustomPlane)));
+		ConstraintsGroup.AddPropertyRow(bLockTranslation).Visibility(TAttribute<EVisibility>::Create(TAttribute<EVisibility>::FGetter::CreateSP(this, &FBodyInstanceCustomizationHelper::IsDOFMode, EDOFMode::CustomPlane)));
+		ConstraintsGroup.AddPropertyRow(bLockRotation).Visibility(TAttribute<EVisibility>::Create(TAttribute<EVisibility>::FGetter::CreateSP(this, &FBodyInstanceCustomizationHelper::IsDOFMode, EDOFMode::CustomPlane)));
+		ConstraintsGroup.AddPropertyRow(DOFModeProperty.ToSharedRef());
+	}
+}
+
 ////////////////////////////////////////////////////////////////
 
 #undef LOCTEXT_NAMESPACE
diff --git a/Engine/Source/Editor/DetailCustomizations/Private/BodyInstanceCustomization.h b/Engine/Source/Editor/DetailCustomizations/Private/BodyInstanceCustomization.h
--- a/Engine/Source/Editor/DetailCustomizations/Private/BodyInstanceCustomization.h
+++ b/Engine/Source/Editor/DetailCustomizations/Private/BodyInstanceCustomization.h
@@ -67,6 +67,9 @@ private:
 
 	void UpdateValidCollisionChannels();
 
+	void AddPhysicsCategory(TSharedRef<class IPropertyHandle> StructPropertyHandle, class IDetailChildrenBuilder& StructBuilder, IPropertyTypeCustomizationUtils& StructCustomizationUtils);
+	void AddCollisionCategory(TSharedRef<class IPropertyHandle> StructPropertyHandle, class IDetailChildrenBuilder& StructBuilder, IPropertyTypeCustomizationUtils& StructCustomizationUtils);
+
 private:
 	// property handles
 	TSharedPtr<IPropertyHandle> CollisionProfileNameHandle;
@@ -93,3 +96,40 @@ private:
 
 	void RefreshCollisionProfiles();
 };
+
+class FBodyInstanceCustomizationHelper  : public TSharedFromThis<FBodyInstanceCustomizationHelper>
+{
+public:
+	FBodyInstanceCustomizationHelper(const TArray<TWeakObjectPtr<UObject>>& InObjectsCustomized);
+	void CustomizeDetails( IDetailLayoutBuilder& DetailBuilder, TSharedRef<IPropertyHandle> BodyInstanceHandler );
+
+private:
+	void UpdateFilters();
+	bool IsSimulatePhysicsEditable() const;
+	bool IsUseAsyncEditable() const;
+
+	TOptional<float> OnGetBodyMass() const;
+	bool IsBodyMassReadOnly() const;
+	EVisibility IsMassVisible(bool bOverrideMass) const;
+	bool IsBodyMassEnabled() const { return !IsBodyMassReadOnly(); }
+	void AddMassInKg(IDetailCategoryBuilder& PhysicsCategory, TSharedRef<IPropertyHandle> BodyInstanceHandler);
+	void AddBodyConstraint(IDetailCategoryBuilder& PhysicsCategory, TSharedRef<IPropertyHandle> BodyInstanceHandler);
+	void AddMaxAngularVelocity(IDetailCategoryBuilder& PhysicsCategory, TSharedRef<IPropertyHandle> BodyInstanceHandler);
+
+	bool IsAutoWeldEditable() const;
+	EVisibility IsAutoWeldVisible() const;
+
+	EVisibility IsMaxAngularVelocityVisible(bool bOverrideMaxAngularVelocity) const;
+	TOptional<float> OnGetBodyMaxAngularVelocity() const;
+
+	bool IsMaxAngularVelocityReadOnly() const;
+	EVisibility IsDOFMode(EDOFMode::Type Mode) const;
+
+private:
+	bool bDisplayMass;
+	bool bDisplayConstraints;
+	bool bDisplayEnablePhysics;
+
+	TSharedPtr<IPropertyHandle> DOFModeProperty;
+	TArray<TWeakObjectPtr<UObject>> ObjectsCustomized;
+};
diff --git a/Engine/Source/Editor/DetailCustomizations/Private/BodySetupDetails.cpp b/Engine/Source/Editor/DetailCustomizations/Private/BodySetupDetails.cpp
--- a/Engine/Source/Editor/DetailCustomizations/Private/BodySetupDetails.cpp
+++ b/Engine/Source/Editor/DetailCustomizations/Private/BodySetupDetails.cpp
@@ -21,7 +21,10 @@ void FBodySetupDetails::CustomizeDetails( IDetailLayoutBuilder& DetailBuilder )
 	{
 		if ( DetailBuilder.GetProperty(GET_MEMBER_NAME_CHECKED(UBodySetup, DefaultInstance))->IsValidHandle() )
 		{
-			IDetailCategoryBuilder& PhysicsCategory = DetailBuilder.EditCategory("Physics");
+			DetailBuilder.GetObjectsBeingCustomized(ObjectsCustomized);
+			BodyInstanceCustomizationHelper = MakeShareable(new FBodyInstanceCustomizationHelper(ObjectsCustomized));
+			BodyInstanceCustomizationHelper->CustomizeDetails(DetailBuilder, DetailBuilder.GetProperty(GET_MEMBER_NAME_CHECKED(UBodySetup, DefaultInstance)));
+
 			IDetailCategoryBuilder& CollisionCategory = DetailBuilder.EditCategory("Collision");
 
 			TSharedPtr<IPropertyHandle> BodyInstanceHandler = DetailBuilder.GetProperty(GET_MEMBER_NAME_CHECKED(UBodySetup, DefaultInstance));
@@ -34,74 +37,14 @@ void FBodySetupDetails::CustomizeDetails( IDetailLayoutBuilder& DetailBuilder )
 			uint32 NumChildren = 0;
 			BodyInstanceHandler->GetNumChildren(NumChildren);
 
-			// Get the objects being customized so we can enable/disable editing of 'Simulate Physics'
-			DetailBuilder.GetObjectsBeingCustomized(ObjectsCustomized);
-
 			// add all properties of this now - after adding 
 			for (uint32 ChildIndex=0; ChildIndex < NumChildren; ++ChildIndex)
 			{
 				TSharedPtr<IPropertyHandle> ChildProperty = BodyInstanceHandler->GetChildHandle(ChildIndex);
 				FString Category = FObjectEditorUtils::GetCategory(ChildProperty->GetProperty());
-				if (ChildProperty->GetProperty()->GetFName() == GET_MEMBER_NAME_CHECKED(FBodyInstance, bSimulatePhysics)
-					|| ChildProperty->GetProperty()->GetFName() == GET_MEMBER_NAME_CHECKED(FBodyInstance, bAutoWeld)
-					|| ChildProperty->GetProperty()->GetFName() == GET_MEMBER_NAME_CHECKED(FBodyInstance, DOFMode)
-					|| ChildProperty->GetProperty()->GetFName() == GET_MEMBER_NAME_CHECKED(FBodyInstance, CustomDOFPlaneNormal)
-					|| ChildProperty->GetProperty()->GetFName() == GET_MEMBER_NAME_CHECKED(FBodyInstance, bLockTranslation)
-					|| ChildProperty->GetProperty()->GetFName() == GET_MEMBER_NAME_CHECKED(FBodyInstance, bLockXTranslation)
-					|| ChildProperty->GetProperty()->GetFName() == GET_MEMBER_NAME_CHECKED(FBodyInstance, bLockYTranslation)
-					|| ChildProperty->GetProperty()->GetFName() == GET_MEMBER_NAME_CHECKED(FBodyInstance, bLockZTranslation)
-					|| ChildProperty->GetProperty()->GetFName() == GET_MEMBER_NAME_CHECKED(FBodyInstance, bLockRotation)
-					|| ChildProperty->GetProperty()->GetFName() == GET_MEMBER_NAME_CHECKED(FBodyInstance, bLockXRotation)
-					|| ChildProperty->GetProperty()->GetFName() == GET_MEMBER_NAME_CHECKED(FBodyInstance, bLockYRotation)
-					|| ChildProperty->GetProperty()->GetFName() == GET_MEMBER_NAME_CHECKED(FBodyInstance, bLockZRotation))
-				{
-					// skip bSimulatePhysics
-					// this is because we don't want bSimulatePhysics to show up 
-					// phat editor 
-					// staitc mesh already hides everything else not interested in
-					// so phat editor just should not show this option
-					//also hide bAutoWeld for phat
-					continue;
-				}
-				else if (ChildProperty->GetProperty()->GetFName() == GET_MEMBER_NAME_CHECKED(FBodyInstance, MassInKg))
-				{
-					PhysicsCategory.AddProperty(ChildProperty).CustomWidget()
-					.NameContent()
-					[
-						ChildProperty->CreatePropertyNameWidget()
-					]
-					.ValueContent()
-						[
-							SNew(SVerticalBox)
-							+ SVerticalBox::Slot()
-							.Padding(0.f, 0.f, 10.f, 0.f)
-							[
-								SNew(SNumericEntryBox<float>)
-								.IsEnabled(false)
-								.Font(IDetailLayoutBuilder::GetDetailFont())
-								.Value(this, &FBodySetupDetails::OnGetBodyMass)
-								.Visibility(this, &FBodySetupDetails::IsMassVisible, false)
-							]
 
-							+ SVerticalBox::Slot()
-								.Padding(0.f, 0.f, 10.f, 0.f)
-								[
-									SNew(SVerticalBox)
-									.Visibility(this, &FBodySetupDetails::IsMassVisible, true)
-									+ SVerticalBox::Slot()
-									.AutoHeight()
-									[
-										ChildProperty->CreatePropertyValueWidget()
-									]
-								]
-						];
-					continue;
-				}
-				if (Category == TEXT("Physics"))
-				{
-					PhysicsCategory.AddProperty(ChildProperty);
-				}
-				else if (Category == TEXT("Collision"))
+				
+				if (Category == TEXT("Collision"))
 				{
 					CollisionCategory.AddProperty(ChildProperty);
 				}
@@ -110,71 +53,5 @@ void FBodySetupDetails::CustomizeDetails( IDetailLayoutBuilder& DetailBuilder )
 	}
 }
 
-TOptional<float> FBodySetupDetails::OnGetBodyMass() const
-{
-	UBodySetup* BS = NULL;
-
-	float Mass = 0.0f;
-	bool bMultipleValue = false;
-
-	for (auto ObjectIt = ObjectsCustomized.CreateConstIterator(); ObjectIt; ++ObjectIt)
-	{
-		if(ObjectIt->IsValid() && (*ObjectIt)->IsA(UBodySetup::StaticClass()))
-		{
-			BS = Cast<UBodySetup>(ObjectIt->Get());
-			
-			float BSMass = BS->CalculateMass();
-			if (Mass == 0.0f || FMath::Abs(Mass - BSMass) < 0.1f)
-			{
-				Mass = BSMass;
-			}
-			else
-			{
-				bMultipleValue = true;
-				break;
-			}
-		}
-	}
-
-	if (bMultipleValue)
-	{
-		return TOptional<float>();
-	}
-
-	return Mass;
-}
-
-EVisibility FBodySetupDetails::IsMassVisible(bool bOverrideMass) const
-{
-	bool bIsMassReadOnly = IsBodyMassReadOnly();
-	if (bOverrideMass)
-	{
-		return bIsMassReadOnly ? EVisibility::Collapsed : EVisibility::Visible;
-	}
-	else
-	{
-		return bIsMassReadOnly ? EVisibility::Visible : EVisibility::Collapsed;
-	}
-}
-
-
-bool FBodySetupDetails::IsBodyMassReadOnly() const
-{
-	for (auto ObjectIt = ObjectsCustomized.CreateConstIterator(); ObjectIt; ++ObjectIt)
-	{
-		if (ObjectIt->IsValid() && (*ObjectIt)->IsA(UBodySetup::StaticClass()))
-		{
-			UBodySetup* BS = Cast<UBodySetup>(ObjectIt->Get());
-			if (BS->DefaultInstance.bOverrideMass == false)
-			{
-				return true;
-			}
-		}
-	}
-
-	return false;
-}
-
-
 #undef LOCTEXT_NAMESPACE
 
diff --git a/Engine/Source/Editor/DetailCustomizations/Private/BodySetupDetails.h b/Engine/Source/Editor/DetailCustomizations/Private/BodySetupDetails.h
--- a/Engine/Source/Editor/DetailCustomizations/Private/BodySetupDetails.h
+++ b/Engine/Source/Editor/DetailCustomizations/Private/BodySetupDetails.h
@@ -2,6 +2,8 @@
 
 #pragma once
 
+#include "BodyInstanceCustomization.h"
+
 class FBodySetupDetails : public IDetailCustomization
 {
 public:
@@ -11,12 +13,8 @@ public:
 	/** IDetailCustomization interface */
 	virtual void CustomizeDetails( IDetailLayoutBuilder& DetailBuilder ) override;
 
-	TOptional<float> OnGetBodyMass() const;
-	bool IsBodyMassReadOnly() const;
-	bool IsBodyMassEnabled() const { return !IsBodyMassReadOnly(); }
-	EVisibility IsMassVisible(bool bOverrideMass) const;
-
 private:
+	TSharedPtr<class FBodyInstanceCustomizationHelper> BodyInstanceCustomizationHelper;
 	TArray< TWeakObjectPtr<UObject> > ObjectsCustomized;
 };
 
diff --git a/Engine/Source/Editor/DetailCustomizations/Private/PrimitiveComponentDetails.cpp b/Engine/Source/Editor/DetailCustomizations/Private/PrimitiveComponentDetails.cpp
--- a/Engine/Source/Editor/DetailCustomizations/Private/PrimitiveComponentDetails.cpp
+++ b/Engine/Source/Editor/DetailCustomizations/Private/PrimitiveComponentDetails.cpp
@@ -26,423 +26,10 @@ TSharedRef<IDetailCustomization> FPrimitiveComponentDetails::MakeInstance()
 	return MakeShareable( new FPrimitiveComponentDetails );
 }
 
-bool FPrimitiveComponentDetails::IsSimulatePhysicsEditable() const
-{
-	// Check whether to enable editing of bSimulatePhysics - this will happen if all objects are UPrimitiveComponents & have collision geometry.
-	bool bEnableSimulatePhysics = ObjectsCustomized.Num() > 0;
-	for (TWeakObjectPtr<UObject> CustomizedObject : ObjectsCustomized)
-	{
-		if (UPrimitiveComponent* PrimitiveComponent = Cast<UPrimitiveComponent>(CustomizedObject.Get()))
-		{
-			if (!PrimitiveComponent->CanEditSimulatePhysics())
-			{
-				bEnableSimulatePhysics = false;
-				break;
-			}
-		}
-	}
-
-	return bEnableSimulatePhysics;
-}
-
-bool FPrimitiveComponentDetails::IsUseAsyncEditable() const
-{
-	// Check whether to enable editing of bUseAsyncScene - this will happen if all objects are movable and the project uses an AsyncScene
-	if (!UPhysicsSettings::Get()->bEnableAsyncScene)
-	{
-		return false;
-	}
-	
-	bool bEnableUseAsyncScene = ObjectsCustomized.Num() > 0;
-	for (auto ObjectIt = ObjectsCustomized.CreateConstIterator(); ObjectIt; ++ObjectIt)
-	{
-		if (ObjectIt->IsValid() && (*ObjectIt)->IsA(UPrimitiveComponent::StaticClass()))
-		{
-			TWeakObjectPtr<USceneComponent> SceneComponent = CastChecked<USceneComponent>(ObjectIt->Get());
-
-			if (SceneComponent.IsValid() && SceneComponent->Mobility != EComponentMobility::Movable)
-			{
-				bEnableUseAsyncScene = false;
-				break;
-			}
-
-			// Skeletal mesh uses a physics asset which will have multiple bodies - these bodies have their own bUseAsyncScene which is what we actually use - the flag on the skeletal mesh is not used
-			TWeakObjectPtr<USkeletalMeshComponent> SkeletalMeshComponent = Cast<USkeletalMeshComponent>(ObjectIt->Get());
-			if (SkeletalMeshComponent.IsValid())
-			{
-				bEnableUseAsyncScene = false;
-				break;
-			}
-		}
-		else
-		{
-			bEnableUseAsyncScene = false;
-			break;
-		}
-	}
-
-	return bEnableUseAsyncScene;
-}
-
-EVisibility FPrimitiveComponentDetails::IsDOFMode(EDOFMode::Type Mode) const
-{
-	bool bVisible = false;
-	if (DOFModeProperty.IsValid())
-	{
-		uint8 CurrentMode;
-		if (DOFModeProperty->GetValue(CurrentMode) == FPropertyAccess::Success)
-		{
-			EDOFMode::Type PropertyDOF = FBodyInstance::ResolveDOFMode(static_cast<EDOFMode::Type>(CurrentMode));
-			bVisible = PropertyDOF == Mode;
-		}
-	}
-
-	return bVisible ? EVisibility::Visible : EVisibility::Collapsed;
-}
-
-bool FPrimitiveComponentDetails::IsAutoWeldEditable() const
-{
-	for (int32 i = 0; i < ObjectsCustomized.Num(); ++i)
-	{
-		if (UPrimitiveComponent * SceneComponent = Cast<UPrimitiveComponent>(ObjectsCustomized[i].Get()))
-		{
-			if (FBodyInstance* BI = SceneComponent->GetBodyInstance())
-			{
-				if (BI->ShouldInstanceSimulatingPhysics())
-				{
-					return false;
-				}
-			}
-		}
-	}
-
-	return true;
-}
-
-EVisibility FPrimitiveComponentDetails::IsAutoWeldVisible() const
-{
-	for (int32 i = 0; i < ObjectsCustomized.Num(); ++i)
-	{
-		if (ObjectsCustomized[i].IsValid() && !( ObjectsCustomized[i]->IsA(UStaticMeshComponent::StaticClass()) || ObjectsCustomized[i]->IsA(UShapeComponent::StaticClass())))
-		{
-			return EVisibility::Collapsed;
-		}
-	}
-
-	return EVisibility::Visible;
-}
-
-EVisibility FPrimitiveComponentDetails::IsMassVisible(bool bOverrideMass) const
-{
-	bool bIsMassReadOnly = IsBodyMassReadOnly();
-	if (bOverrideMass)
-	{
-		return bIsMassReadOnly ? EVisibility::Collapsed : EVisibility::Visible;
-	}
-	else
-	{
-		return bIsMassReadOnly ? EVisibility::Visible : EVisibility::Collapsed;
-	}
-}
-
-EVisibility FPrimitiveComponentDetails::IsMaxAngularVelocityVisible(bool bOverrideMaxAngularVelocity) const
-{
-	bool bIsMaxAngularVelocityReadOnly = IsMaxAngularVelocityReadOnly();
-	if (bOverrideMaxAngularVelocity)
-	{
-		return bIsMaxAngularVelocityReadOnly ? EVisibility::Collapsed : EVisibility::Visible;
-	}
-	else
-	{
-		return bIsMaxAngularVelocityReadOnly ? EVisibility::Visible : EVisibility::Collapsed;
-	}
-}
-
 void FPrimitiveComponentDetails::AddPhysicsCategory(IDetailLayoutBuilder& DetailBuilder)
 {
-	if (DetailBuilder.GetProperty(GET_MEMBER_NAME_CHECKED(UPrimitiveComponent, BodyInstance))->IsValidHandle())
-	{
-		TSharedPtr<IPropertyHandle> BodyInstanceHandler = DetailBuilder.GetProperty(GET_MEMBER_NAME_CHECKED(UPrimitiveComponent, BodyInstance));
-		uint32 NumChildren = 0;
-		BodyInstanceHandler->GetNumChildren(NumChildren);
-
-
-		IDetailCategoryBuilder& PhysicsCategory = DetailBuilder.EditCategory("Physics");
-
-		bool bDisplayMass = true;
-		bool bDisplayConstraints = true;
-
-		for (int32 i = 0; i < ObjectsCustomized.Num(); ++i)
-		{
-			if (ObjectsCustomized[i].IsValid() && ObjectsCustomized[i]->IsA(UDestructibleComponent::StaticClass()))
-			{
-				bDisplayMass = false;
-				bDisplayConstraints = false;
-			}
-		}
-
-		bool bConstraintsUIExists = false;
-		bool bMassUIExists = false;
-
-		// add all physics properties now - after adding mass
-		for (uint32 ChildIndex = 0; ChildIndex < NumChildren; ++ChildIndex)
-		{
-			TSharedPtr<IPropertyHandle> ChildProperty = BodyInstanceHandler->GetChildHandle(ChildIndex);
-			FString Category = FObjectEditorUtils::GetCategory(ChildProperty->GetProperty());
-			FName PropName = ChildProperty->GetProperty()->GetFName();
-			if (Category == TEXT("Physics"))
-			{
-				// Only permit modifying bSimulatePhysics when the body has some geometry.
-				if (PropName == GET_MEMBER_NAME_CHECKED(FBodyInstance, bSimulatePhysics))
-				{
-					PhysicsCategory.AddProperty(ChildProperty).EditCondition(TAttribute<bool>(this, &FPrimitiveComponentDetails::IsSimulatePhysicsEditable), NULL);
-				}
-				else if (PropName == FName("bUseAsyncScene")) // Can't use GET_MEMBER_NAME_CHECKED because protected
-				{
-					//we only enable bUseAsyncScene if the project uses an AsyncScene
-					PhysicsCategory.AddProperty(ChildProperty).EditCondition(TAttribute<bool>(this, &FPrimitiveComponentDetails::IsUseAsyncEditable), NULL);
-				}
-				else if (PropName == GET_MEMBER_NAME_CHECKED(FBodyInstance, bLockXTranslation)
-					|| PropName == GET_MEMBER_NAME_CHECKED(FBodyInstance, bLockYTranslation)
-					|| PropName == GET_MEMBER_NAME_CHECKED(FBodyInstance, bLockZTranslation)
-					|| PropName == GET_MEMBER_NAME_CHECKED(FBodyInstance, bLockXRotation)
-					|| PropName == GET_MEMBER_NAME_CHECKED(FBodyInstance, bLockYRotation)
-					|| PropName == GET_MEMBER_NAME_CHECKED(FBodyInstance, bLockZRotation)
-					|| PropName == GET_MEMBER_NAME_CHECKED(FBodyInstance, DOFMode)
-					|| PropName == GET_MEMBER_NAME_CHECKED(FBodyInstance, CustomDOFPlaneNormal)
-					|| PropName == GET_MEMBER_NAME_CHECKED(FBodyInstance, bLockTranslation)
-					|| PropName == GET_MEMBER_NAME_CHECKED(FBodyInstance, bLockRotation))
-				{
-					// We treat all of these constraint parameters as a big group
-					if (bConstraintsUIExists == false && bDisplayConstraints)
-					{
-						const float XYZPadding = 5.f;
-						bConstraintsUIExists = true;
-						DOFModeProperty = BodyInstanceHandler->GetChildHandle(TEXT("DOFMode"));
-
-						IDetailGroup& ConstraintsGroup = PhysicsCategory.AddGroup(TEXT("ConstraintsGroup"), LOCTEXT("Constraints", "Constraints"));
-
-						ConstraintsGroup.AddWidgetRow()
-						.Visibility(TAttribute<EVisibility>::Create(TAttribute<EVisibility>::FGetter::CreateSP(this, &FPrimitiveComponentDetails::IsDOFMode, EDOFMode::SixDOF)))
-						.NameContent()
-						[
-							SNew(STextBlock)
-							.Text(LOCTEXT("LockPositionLabel", "Lock Position"))
-							.ToolTipText(LOCTEXT("LockPositionTooltip", "Locks movement along the specified axis"))
-						]
-						.ValueContent()
-						[
-							SNew(SHorizontalBox)
-							+ SHorizontalBox::Slot()
-							.Padding(0.f, 0.f, XYZPadding, 0.f)
-							.AutoWidth()
-							[
-								SNew(SHorizontalBox)
-								+ SHorizontalBox::Slot()
-								.AutoWidth()
-								[
-									BodyInstanceHandler->GetChildHandle(TEXT("bLockXTranslation"))->CreatePropertyNameWidget()
-								]
-								+ SHorizontalBox::Slot()
-								.AutoWidth()
-								[
-									BodyInstanceHandler->GetChildHandle(TEXT("bLockXTranslation"))->CreatePropertyValueWidget()
-								]
-							]
-
-							+ SHorizontalBox::Slot()
-							.Padding(0.f, 0.f, XYZPadding, 0.f)
-							.AutoWidth()
-							[
-								SNew(SHorizontalBox)
-								+ SHorizontalBox::Slot()
-								.AutoWidth()
-								[
-									BodyInstanceHandler->GetChildHandle(TEXT("bLockYTranslation"))->CreatePropertyNameWidget()
-								]
-								+ SHorizontalBox::Slot()
-									.AutoWidth()
-									[
-										BodyInstanceHandler->GetChildHandle(TEXT("bLockYTranslation"))->CreatePropertyValueWidget()
-									]
-							]
-
-							+ SHorizontalBox::Slot()
-							.Padding(0.f, 0.f, XYZPadding, 0.f)
-							.AutoWidth()
-							[
-								SNew(SHorizontalBox)
-								+ SHorizontalBox::Slot()
-								.AutoWidth()
-								[
-									BodyInstanceHandler->GetChildHandle(TEXT("bLockZTranslation"))->CreatePropertyNameWidget()
-								]
-								+ SHorizontalBox::Slot()
-								.AutoWidth()
-								[
-									BodyInstanceHandler->GetChildHandle(TEXT("bLockZTranslation"))->CreatePropertyValueWidget()
-								]
-							]
-						];
-
-						ConstraintsGroup.AddWidgetRow()
-						.Visibility(TAttribute<EVisibility>::Create(TAttribute<EVisibility>::FGetter::CreateSP(this, &FPrimitiveComponentDetails::IsDOFMode, EDOFMode::SixDOF)))
-						.NameContent()
-						[
-							SNew(STextBlock)
-							.Text(LOCTEXT("LockRotationLabel", "Lock Rotation"))
-							.ToolTipText(LOCTEXT("LockRotationTooltip", "Locks rotation about the specified axis"))
-						]
-						.ValueContent()
-							[
-								SNew(SHorizontalBox)
-								+ SHorizontalBox::Slot()
-								.Padding(0.f, 0.f, XYZPadding, 0.f)
-								.AutoWidth()
-								[
-									SNew(SHorizontalBox)
-									+ SHorizontalBox::Slot()
-									.AutoWidth()
-									[
-										BodyInstanceHandler->GetChildHandle(TEXT("bLockXRotation"))->CreatePropertyNameWidget()
-									]
-									+ SHorizontalBox::Slot()
-									.AutoWidth()
-									[
-										BodyInstanceHandler->GetChildHandle(TEXT("bLockXRotation"))->CreatePropertyValueWidget()
-									]
-								]
-
-								+ SHorizontalBox::Slot()
-								.Padding(0.f, 0.f, XYZPadding, 0.f)
-								.AutoWidth()
-								[
-									SNew(SHorizontalBox)
-									+ SHorizontalBox::Slot()
-									.AutoWidth()
-									[
-										BodyInstanceHandler->GetChildHandle(TEXT("bLockYRotation"))->CreatePropertyNameWidget()
-									]
-									+ SHorizontalBox::Slot()
-									.AutoWidth()
-									[
-										BodyInstanceHandler->GetChildHandle(TEXT("bLockYRotation"))->CreatePropertyValueWidget()
-									]
-								]
-
-								+ SHorizontalBox::Slot()
-								.Padding(0.f, 0.f, XYZPadding, 0.f)
-								.AutoWidth()
-								[
-									SNew(SHorizontalBox)
-									+ SHorizontalBox::Slot()
-									.AutoWidth()
-									[
-										BodyInstanceHandler->GetChildHandle(TEXT("bLockZRotation"))->CreatePropertyNameWidget()
-									]
-									+ SHorizontalBox::Slot()
-									.AutoWidth()
-									[
-										BodyInstanceHandler->GetChildHandle(TEXT("bLockZRotation"))->CreatePropertyValueWidget()
-									]
-								]
-
-							];
-
-						//we only show the custom plane normal if we've selected that mode
-						ConstraintsGroup.AddPropertyRow(BodyInstanceHandler->GetChildHandle(TEXT("CustomDOFPlaneNormal")).ToSharedRef()).Visibility(TAttribute<EVisibility>::Create(TAttribute<EVisibility>::FGetter::CreateSP(this, &FPrimitiveComponentDetails::IsDOFMode, EDOFMode::CustomPlane)));
-						ConstraintsGroup.AddPropertyRow(BodyInstanceHandler->GetChildHandle(TEXT("bLockTranslation")).ToSharedRef()).Visibility(TAttribute<EVisibility>::Create(TAttribute<EVisibility>::FGetter::CreateSP(this, &FPrimitiveComponentDetails::IsDOFMode, EDOFMode::CustomPlane)));
-						ConstraintsGroup.AddPropertyRow(BodyInstanceHandler->GetChildHandle(TEXT("bLockRotation")).ToSharedRef()).Visibility(TAttribute<EVisibility>::Create(TAttribute<EVisibility>::FGetter::CreateSP(this, &FPrimitiveComponentDetails::IsDOFMode, EDOFMode::CustomPlane)));
-						ConstraintsGroup.AddPropertyRow(DOFModeProperty.ToSharedRef());
-					}
-				}
-				else if (PropName == GET_MEMBER_NAME_CHECKED(FBodyInstance, bAutoWeld))
-				{
-					PhysicsCategory.AddProperty(ChildProperty).Visibility(TAttribute<EVisibility>(this, &FPrimitiveComponentDetails::IsAutoWeldVisible))
-						.EditCondition(TAttribute<bool>(this, &FPrimitiveComponentDetails::IsAutoWeldEditable), NULL);
-				}
-				else if (PropName == GET_MEMBER_NAME_CHECKED(FBodyInstance, MassInKg))
-				{
-					if (bDisplayMass && bMassUIExists == false)
-					{
-						bMassUIExists = true;
-						TSharedPtr<IPropertyHandle> MassInKGHandle = BodyInstanceHandler->GetChildHandle(TEXT("MassInKg"));
-
-						PhysicsCategory.AddProperty(MassInKGHandle).CustomWidget()
-						.NameContent()
-						[
-							MassInKGHandle->CreatePropertyNameWidget()
-						]
-						.ValueContent()
-						[
-							SNew(SVerticalBox)
-							+ SVerticalBox::Slot()
-							.Padding(0.f, 0.f, 10.f, 0.f)
-							[
-								SNew(SNumericEntryBox<float>)
-								.IsEnabled(false)
-								.Font( IDetailLayoutBuilder::GetDetailFont() )
-								.Value(this, &FPrimitiveComponentDetails::OnGetBodyMass)
-								.Visibility(this, &FPrimitiveComponentDetails::IsMassVisible, false)
-							]
-
-							+ SVerticalBox::Slot()
-							.Padding(0.f, 0.f, 10.f, 0.f)
-							[
-								SNew(SVerticalBox)
-								.Visibility(this, &FPrimitiveComponentDetails::IsMassVisible, true)
-								+ SVerticalBox::Slot()
-								.AutoHeight()
-								[
-									MassInKGHandle->CreatePropertyValueWidget()
-								]
-							]
-						];
-					}
-				}
-				else if (PropName == GET_MEMBER_NAME_CHECKED(FBodyInstance, MaxAngularVelocity))
-				{
-					bMassUIExists = true;
-					TSharedPtr<IPropertyHandle> MaxAngularVelocityHandle = BodyInstanceHandler->GetChildHandle(PropName);
-
-					PhysicsCategory.AddProperty(MaxAngularVelocityHandle).CustomWidget()
-					.NameContent()
-						[
-							MaxAngularVelocityHandle->CreatePropertyNameWidget()
-						]
-					.ValueContent()
-						[
-							SNew(SVerticalBox)
-							+ SVerticalBox::Slot()
-							.Padding(0.f, 0.f, 10.f, 0.f)
-							[
-								SNew(SNumericEntryBox<float>)
-								.IsEnabled(false)
-								.Font(IDetailLayoutBuilder::GetDetailFont())
-								.Value(this, &FPrimitiveComponentDetails::OnGetBodyMaxAngularVelocity)
-								.Visibility(this, &FPrimitiveComponentDetails::IsMaxAngularVelocityVisible, false)
-							]
-
-							+ SVerticalBox::Slot()
-								.Padding(0.f, 0.f, 10.f, 0.f)
-								[
-									SNew(SVerticalBox)
-									.Visibility(this, &FPrimitiveComponentDetails::IsMaxAngularVelocityVisible, true)
-									+ SVerticalBox::Slot()
-									.AutoHeight()
-									[
-										MaxAngularVelocityHandle->CreatePropertyValueWidget()
-									]
-								]
-						];
-				}
-				else
-				{
-					PhysicsCategory.AddProperty(ChildProperty);
-				}
-			}
-		}
-	}
+	BodyInstanceCustomizationHelper = MakeShareable(new FBodyInstanceCustomizationHelper(ObjectsCustomized));
+	BodyInstanceCustomizationHelper->CustomizeDetails(DetailBuilder, DetailBuilder.GetProperty(GET_MEMBER_NAME_CHECKED(UPrimitiveComponent, BodyInstance)));
 }
 
 void FPrimitiveComponentDetails::AddCollisionCategory(IDetailLayoutBuilder& DetailBuilder)
@@ -558,91 +145,6 @@ void FPrimitiveComponentDetails::AddAdvancedSubCategory( IDetailLayoutBuilder& D
 
 }
 
-TOptional<float> FPrimitiveComponentDetails::OnGetBodyMass() const
-{
-	UPrimitiveComponent* Comp = NULL;
-
-	float Mass = 0.0f;
-	bool bMultipleValue = false;
-
-	for (auto ObjectIt = ObjectsCustomized.CreateConstIterator(); ObjectIt; ++ObjectIt)
-	{
-		if(ObjectIt->IsValid() && (*ObjectIt)->IsA(UPrimitiveComponent::StaticClass()))
-		{
-			Comp = Cast<UPrimitiveComponent>(ObjectIt->Get());
-
-			float CompMass = Comp->CalculateMass();
-			Comp->BodyInstance.MassInKg = CompMass;	//update the component's override mass to be the calculated one
-			if (Mass == 0.0f || FMath::Abs(Mass - CompMass) < 0.1f)
-			{
-				Mass = CompMass;
-			}
-			else
-			{
-				bMultipleValue = true;
-				break;
-			}
-		}
-	}
-
-	if (bMultipleValue)
-	{
-		return TOptional<float>();
-	}
-
-	return Mass;
-}
-
-
-TOptional<float> FPrimitiveComponentDetails::OnGetBodyMaxAngularVelocity() const
-{
-	UPrimitiveComponent* Comp = nullptr;
-	const float MaxAngularVelocity = UPhysicsSettings::Get()->MaxAngularVelocity;
-
-	for (auto ObjectIt = ObjectsCustomized.CreateConstIterator(); ObjectIt; ++ObjectIt)
-	{
-		if (ObjectIt->IsValid() && (*ObjectIt)->IsA(UPrimitiveComponent::StaticClass()))
-		{
-			Comp = Cast<UPrimitiveComponent>(ObjectIt->Get());
-			Comp->BodyInstance.MaxAngularVelocity = MaxAngularVelocity;	//update max angular velocity so that overriding gives the same value initially
-		}
-	}
-
-	return MaxAngularVelocity;
-}
-
-bool FPrimitiveComponentDetails::IsBodyMassReadOnly() const
-{
-	for (auto ObjectIt = ObjectsCustomized.CreateConstIterator(); ObjectIt; ++ObjectIt)
-	{
-		if (ObjectIt->IsValid() && (*ObjectIt)->IsA(UPrimitiveComponent::StaticClass()))
-		{
-			if (UPrimitiveComponent* Comp = Cast<UPrimitiveComponent>(ObjectIt->Get()))
-			{
-				if (Comp->BodyInstance.bOverrideMass == false) { return true; }
-			}
-		}
-	}
-
-	return false;
-}
-
-bool FPrimitiveComponentDetails::IsMaxAngularVelocityReadOnly() const
-{
-	for (auto ObjectIt = ObjectsCustomized.CreateConstIterator(); ObjectIt; ++ObjectIt)
-	{
-		if (ObjectIt->IsValid() && (*ObjectIt)->IsA(UPrimitiveComponent::StaticClass()))
-		{
-			if (UPrimitiveComponent* Comp = Cast<UPrimitiveComponent>(ObjectIt->Get()))
-			{
-				if (Comp->BodyInstance.bOverrideMaxAngularVelocity == false) { return true; }
-			}
-		}
-	}
-
-	return false;
-}
-
 
 #undef LOCTEXT_NAMESPACE
 
diff --git a/Engine/Source/Editor/DetailCustomizations/Private/PrimitiveComponentDetails.h b/Engine/Source/Editor/DetailCustomizations/Private/PrimitiveComponentDetails.h
--- a/Engine/Source/Editor/DetailCustomizations/Private/PrimitiveComponentDetails.h
+++ b/Engine/Source/Editor/DetailCustomizations/Private/PrimitiveComponentDetails.h
@@ -2,6 +2,8 @@
 
 #pragma once
 
+#include "BodyInstanceCustomization.h"
+
 class FPrimitiveComponentDetails : public IDetailCustomization
 {
 public:
@@ -39,17 +41,9 @@ private:
 private:
 	/** Objects being customized so we can update the 'Simulate Physics' state if physics geometry is added/removed */
 	TArray< TWeakObjectPtr<UObject> > ObjectsCustomized;
-	TSharedPtr<IPropertyHandle> DOFModeProperty;
-	TSharedPtr<class FComponentMaterialCategory> MaterialCategory;
-	EVisibility IsDOFMode(EDOFMode::Type Mode) const;
 
-	bool IsAutoWeldEditable() const;
-	EVisibility IsAutoWeldVisible() const;
-	EVisibility IsMassVisible(bool bOverrideMass) const;
-	EVisibility IsMaxAngularVelocityVisible(bool bOverrideMaxAngularVelocity) const;
-	TOptional<float> OnGetBodyMaxAngularVelocity() const;
-
-	bool IsMaxAngularVelocityReadOnly() const;
+	TSharedPtr<class FComponentMaterialCategory> MaterialCategory;
+	TSharedPtr<class FBodyInstanceCustomizationHelper> BodyInstanceCustomizationHelper;
 
 };
 
diff --git a/Engine/Source/Runtime/Engine/Classes/PhysicsEngine/BodyInstance.h b/Engine/Source/Runtime/Engine/Classes/PhysicsEngine/BodyInstance.h
--- a/Engine/Source/Runtime/Engine/Classes/PhysicsEngine/BodyInstance.h
+++ b/Engine/Source/Runtime/Engine/Classes/PhysicsEngine/BodyInstance.h
@@ -142,20 +142,20 @@
 	 *	Index of this BodyInstance within the SkeletalMeshComponent/PhysicsAsset. 
 	 *	Is INDEX_NONE if a single body component
 	 */
-	int32 InstanceBodyIndex;
+	int16 InstanceBodyIndex;
 
 	/** When we are a body within a SkeletalMeshComponent, we cache the index of the bone we represent, to speed up sync'ing physics to anim. */
-	int32 InstanceBoneIndex;
+	int16 InstanceBoneIndex;
 
 	/** Current scale of physics - used to know when and how physics must be rescaled to match current transform of OwnerComponent. */
 	UPROPERTY()
 	FVector Scale3D;
 
 	/** Physics scene index for the synchronous scene. */
-	int32 SceneIndexSync;
+	int16 SceneIndexSync;
 
 	/** Physics scene index for the asynchronous scene. */
-	int32 SceneIndexAsync;
+	int16 SceneIndexAsync;
 
 	/////////
 	// COLLISION SETTINGS
@@ -167,63 +167,36 @@
 	struct FCollisionResponseContainer ResponseToChannels_DEPRECATED;
 
 private:
+
 	/** Collision Profile Name **/
 	UPROPERTY(EditAnywhere, Category=Custom)
 	FName CollisionProfileName;
 
-	/**
-	 * Type of collision enabled.
-	 * 
-	 *	No Collision      : No collision is performed against this body.
-	 *	Query Only        : This body is used only for collision queries (raycasts, sweeps, and overlaps).
-	 *	Physics Only      : This body is used only for physics collision.
-	 *	Collision Enabled : This body interacts with all collision (Query and Physics).
-	 */
-	UPROPERTY(EditAnywhere, Category=Custom)
-	TEnumAsByte<ECollisionEnabled::Type> CollisionEnabled;
-
-	/** Enum indicating what type of object this should be considered as when it moves */
-	UPROPERTY(EditAnywhere, Category=Custom)
-	TEnumAsByte<enum ECollisionChannel> ObjectType;
-
 	/** Custom Channels for Responses*/
-	UPROPERTY(EditAnywhere, Category=Custom) 
+	UPROPERTY(EditAnywhere, Category = Custom)
 	struct FCollisionResponse CollisionResponses;
 
 public:
 
 	/** If true Continuous Collision Detection (CCD) will be used for this component */
-	UPROPERTY(EditAnywhere, AdvancedDisplay, BlueprintReadOnly, Category=Collision)
-	uint32 bUseCCD:1;
+	UPROPERTY(EditAnywhere, AdvancedDisplay, BlueprintReadOnly, Category = Collision)
+	uint32 bUseCCD : 1;
 
 	/**	Should 'Hit' events fire when this object collides during physics simulation. */
-	UPROPERTY(EditAnywhere, BlueprintReadOnly, Category=Collision, meta=(DisplayName="Simulation Generates Hit Events"))
-	uint32 bNotifyRigidBodyCollision:1;
+	UPROPERTY(EditAnywhere, BlueprintReadOnly, Category = Collision, meta = (DisplayName = "Simulation Generates Hit Events"))
+	uint32 bNotifyRigidBodyCollision : 1;
 
 	/////////
 	// SIM SETTINGS
 
 	/** If true, this body will use simulation. If false, will be 'fixed' (ie kinematic) and move where it is told. */
-	UPROPERTY(EditAnywhere, BlueprintReadWrite, Category=Physics)
-	uint32 bSimulatePhysics:1;
+	UPROPERTY(EditAnywhere, BlueprintReadWrite, Category = Physics)
+	uint32 bSimulatePhysics : 1;
 
 	/** If true, mass will not be automatically computed and you must set it directly */
-	UPROPERTY(meta=(DisplayName="Override"))
+	UPROPERTY(meta = (DisplayName = "Override"))
 	uint32 bOverrideMass : 1;
 
-	/**Mass of the body in KG. By default we compute this based on physical material and mass scale.
-	*@see bOverrideMass to set this directly */
-	UPROPERTY(EditAnywhere, BlueprintReadOnly, Category = Physics, meta = (editcondition = "bOverrideMass", ClampMin = "0.001", UIMin = "0.001"))
-	float MassInKg;
-
-	/** 'Drag' force added to reduce linear movement */
-	UPROPERTY(EditAnywhere, BlueprintReadWrite, Category = Physics)
-	float LinearDamping;
-
-	/** 'Drag' force added to reduce angular movement */
-	UPROPERTY(EditAnywhere, BlueprintReadWrite, Category = Physics)
-	float AngularDamping;
-
 	/** If object should have the force of gravity applied */
 	UPROPERTY(EditAnywhere, BlueprintReadOnly, Category = Physics)
 	uint32 bEnableGravity : 1;
@@ -283,9 +256,45 @@
 	UPROPERTY(meta = (editcondition = "bSimulatePhysics"))
 	uint32 bOverrideMaxAngularVelocity : 1;
 
-	/** Locks physical movement along specified axis.*/
-	UPROPERTY(EditAnywhere, Category = Physics, meta = (DisplayName = "Mode"))
-	TEnumAsByte<EDOFMode::Type> DOFMode;
+	/** When initializing dynamic instances their component or velocity can override the bStartAwake flag */
+	uint32 bWokenExternally : 1;
+protected:
 
+	/**
+	 * If true, this body will be put into the asynchronous physics scene. If false, it will be put into the synchronous physics scene.
+	 * If the body is static, it will be placed into both scenes regardless of the value of bUseAsyncScene.
+	 */
+	UPROPERTY(EditAnywhere, AdvancedDisplay, BlueprintReadWrite, Category=Physics)
+	uint32 bUseAsyncScene:1;
+
+	/** Whether this body instance has its own custom MaxDepenetrationVelocity*/
+	UPROPERTY()
+	uint32 bOverrideMaxDepenetrationVelocity : 1;
+
+	/** Whether this instance of the object has its own custom walkable slope override setting. */
+	UPROPERTY()
+	uint32 bOverrideWalkableSlopeOnInstance : 1;
+
+	uint32 bHasSharedShapes : 1;
+
+	/** The maximum velocity used to depenetrate this object*/
+	UPROPERTY(EditAnywhere, AdvancedDisplay, BlueprintReadOnly, Category = Physics, meta = (editcondition = "bOverrideMaxDepenetrationVelocity", ClampMin = "0.0", UIMin = "0.0"))
+	float MaxDepenetrationVelocity;
+
+public:
+
+	/**Mass of the body in KG. By default we compute this based on physical material and mass scale.
+	*@see bOverrideMass to set this directly */
+	UPROPERTY(EditAnywhere, BlueprintReadOnly, Category = Physics, meta = (editcondition = "bOverrideMass", ClampMin = "0.001", UIMin = "0.001"))
+	float MassInKg;
+
+	/** 'Drag' force added to reduce linear movement */
+	UPROPERTY(EditAnywhere, BlueprintReadWrite, Category = Physics)
+	float LinearDamping;
+
+	/** 'Drag' force added to reduce angular movement */
+	UPROPERTY(EditAnywhere, BlueprintReadWrite, Category = Physics)
+	float AngularDamping;
+
 	/** Locks physical movement along a custom plane for a given normal.*/
 	UPROPERTY(EditAnywhere, Category = Physics, meta = (DisplayName = "Plane Normal"))
@@ -299,11 +308,6 @@
 	UPROPERTY(EditAnywhere, AdvancedDisplay, BlueprintReadWrite, Category = Physics)
 	float MassScale;
 
-	/** The maximum angular velocity for this instance */
-	UPROPERTY(EditAnywhere, AdvancedDisplay, BlueprintReadOnly, Category = Physics, meta = (editcondition = "bOverrideMaxAngularVelocity"))
-	float MaxAngularVelocity;
-
-
 	/** Locks physical movement along axis. */
 	void SetDOFLock(EDOFMode::Type NewDOFMode);
 
@@ -327,24 +331,5 @@
 protected:
 
 	/**
-	 * If true, this body will be put into the asynchronous physics scene. If false, it will be put into the synchronous physics scene.
-	 * If the body is static, it will be placed into both scenes regardless of the value of bUseAsyncScene.
-	 */
-	UPROPERTY(EditAnywhere, AdvancedDisplay, BlueprintReadWrite, Category=Physics)
-	uint32 bUseAsyncScene:1;
-
-	/** Whether this body instance has its own custom MaxDepenetrationVelocity*/
-	UPROPERTY()
-	uint32 bOverrideMaxDepenetrationVelocity : 1;
-
-	/** The maximum velocity used to depenetrate this object*/
-	UPROPERTY(EditAnywhere, AdvancedDisplay, BlueprintReadOnly, Category = Physics, meta = (editcondition = "bOverrideMaxDepenetrationVelocity", ClampMin = "0.0", UIMin = "0.0"))
-	float MaxDepenetrationVelocity;
-
-	/** Whether this instance of the object has its own custom walkable slope override setting. */
-	UPROPERTY()
-	uint32 bOverrideWalkableSlopeOnInstance : 1;
-
-	/**
 	* Custom walkable slope override setting for this instance.
 	* @see GetWalkableSlopeOverride(), SetWalkableSlopeOverride()
@@ -357,9 +342,10 @@
 	class UPhysicalMaterial* PhysMaterialOverride;
 
 public:
-	/** The set of values used in considering when put this body to sleep. */
-	UPROPERTY(EditAnywhere, AdvancedDisplay, BlueprintReadWrite, Category=Physics)
-	ESleepFamily SleepFamily;
+	/** The maximum angular velocity for this instance */
+	UPROPERTY(EditAnywhere, AdvancedDisplay, BlueprintReadOnly, Category = Physics, meta = (editcondition = "bOverrideMaxAngularVelocity"))
+	float MaxAngularVelocity;
 
+
 	/** If the SleepFamily is set to custom, multiply the natural sleep threshold by this amount. A higher number will cause the body to sleep sooner. */
 	UPROPERTY(EditAnywhere, AdvancedDisplay, BlueprintReadOnly, Category=Physics)
@@ -374,10 +360,6 @@
 	UPROPERTY(EditAnywhere, AdvancedDisplay, BlueprintReadOnly, Category=Physics)
 	int32 PositionSolverIterationCount;
 
-	/** This physics body's solver iteration count for velocity. Increasing this will be more CPU intensive, but better stabilized. */
-	UPROPERTY(EditAnywhere, AdvancedDisplay, BlueprintReadOnly, Category=Physics)
-	int32 VelocitySolverIterationCount;
-
 public:
 
 #if WITH_PHYSX
@@ -401,13 +383,14 @@
 	UPROPERTY()
 	uint64 RigidActorAsyncId;
 
+	/** This physics body's solver iteration count for velocity. Increasing this will be more CPU intensive, but better stabilized. */
+	UPROPERTY(EditAnywhere, AdvancedDisplay, BlueprintReadOnly, Category = Physics)
+	int32 VelocitySolverIterationCount;
+
 	/** Per instance data used to initialize dynamic instances */
 	/** Initial physx velocity to apply to dynamic instances */
 	FVector InitialLinearVelocity;
 
-	/** When initializing dynamic instances their component or velocity can override the bStartAwake flag */
-	bool bWokenExternally;
-
 	/** PrimitiveComponent containing this body.   */
 	TWeakObjectPtr<class UPrimitiveComponent> OwnerComponent;
 
@@ -1011,9 +994,6 @@
 public:
 	FPhysxUserData PhysxUserData;
 
-	// Current state of the physx body for tracking deferred addition and removal.
-	BodyInstanceSceneState CurrentSceneState;
-
 	/** Returns the original owning body instance. This is needed for welding */
 	const FBodyInstance* GetOriginalBodyInstance(const physx::PxShape* PShape) const;
 
@@ -1062,6 +1042,7 @@
 	friend struct FInitBodiesHelper<true>;
 	friend struct FInitBodiesHelper<false>;
 	friend class FDerivedDataPhysXBinarySerializer;
+	friend class FBodyInstanceCustomizationHelper;
 
 #if WITH_BOX2D
 
@@ -1085,10 +1066,37 @@
 
 #if WITH_PHYSX
 	/** Used to map between shapes and welded bodies. We do not create entries if the owning body instance is root*/
-	TMap<physx::PxShape*, FWeldInfo> ShapeToBodiesMap;
+	TSharedPtr<TMap<physx::PxShape*, FWeldInfo>> ShapeToBodiesMap;
 
 #endif
-	bool bHasSharedShapes;
+
+public:
+	// Current state of the physx body for tracking deferred addition and removal.
+	BodyInstanceSceneState CurrentSceneState;
+
+	/** The set of values used in considering when put this body to sleep. */
+	UPROPERTY(EditAnywhere, AdvancedDisplay, BlueprintReadWrite, Category = Physics)
+	ESleepFamily SleepFamily;
+
+	/** Locks physical movement along specified axis.*/
+	UPROPERTY(EditAnywhere, Category = Physics, meta = (DisplayName = "Mode"))
+	TEnumAsByte<EDOFMode::Type> DOFMode;
+
+private:
+	/**
+	 * Type of collision enabled.
+	 * 
+	 *	No Collision      : No collision is performed against this body.
+	 *	Query Only        : This body is used only for collision queries (raycasts, sweeps, and overlaps).
+	 *	Physics Only      : This body is used only for physics collision.
+	 *	Collision Enabled : This body interacts with all collision (Query and Physics).
+	 */
+	UPROPERTY(EditAnywhere, Category=Custom)
+	TEnumAsByte<ECollisionEnabled::Type> CollisionEnabled;
+
+	/** Enum indicating what type of object this should be considered as when it moves */
+	UPROPERTY(EditAnywhere, Category=Custom)
+	TEnumAsByte<enum ECollisionChannel> ObjectType;
 };
 
 template<>
diff --git a/Engine/Source/Runtime/Engine/Private/PhysicsEngine/BodyInstance.cpp b/Engine/Source/Runtime/Engine/Private/PhysicsEngine/BodyInstance.cpp
--- a/Engine/Source/Runtime/Engine/Private/PhysicsEngine/BodyInstance.cpp
+++ b/Engine/Source/Runtime/Engine/Private/PhysicsEngine/BodyInstance.cpp
@@ -1677,14 +1677,14 @@ FVector GetInitialLinearVelocity(const AActor* OwningActor, bool& bComponentAwak
 const FBodyInstance* FBodyInstance::GetOriginalBodyInstance(const PxShape* PShape) const
 {
 	const FBodyInstance* BI = WeldParent ? WeldParent : this;
-	const FWeldInfo* Result = BI->ShapeToBodiesMap.Find(PShape);
+	const FWeldInfo* Result = BI->ShapeToBodiesMap.IsValid() ? BI->ShapeToBodiesMap->Find(PShape) : nullptr;
 	return Result ? Result->ChildBI : this;
 }
 
 const FTransform& FBodyInstance::GetRelativeBodyTransform(const physx::PxShape* PShape) const
 {
 	const FBodyInstance* BI = WeldParent ? WeldParent : this;
-	const FWeldInfo* Result = BI->ShapeToBodiesMap.Find(PShape);
+	const FWeldInfo* Result = BI->ShapeToBodiesMap.IsValid() ? BI->ShapeToBodiesMap->Find(PShape) : nullptr;
 	return Result ? Result->RelativeTM : FTransform::Identity;
 }
 
@@ -1732,7 +1732,7 @@ TArray<int32> FBodyInstance::AddCollisionNotifyInfo(const FBodyInstance* Body0,
 }
 
 //helper function for TermBody to avoid code duplication between scenes
-void TermBodyHelper(int32& SceneIndex, PxRigidActor*& PRigidActor, FBodyInstance* BodyInstance)
+void TermBodyHelper(int16& SceneIndex, PxRigidActor*& PRigidActor, FBodyInstance* BodyInstance)
 {
 	if (SceneIndex)
 	{
@@ -1890,13 +1890,19 @@ bool FBodyInstance::Weld(FBodyInstance* TheirBody, const FTransform& TheirTM)
 			}
 		}
 
-		for (int32 ShapeIdx = 0; ShapeIdx < PNewShapes.Num(); ++ShapeIdx)
+		if(PNewShapes.Num())
 		{
-			PxShape* PShape = PNewShapes[ShapeIdx];
-			ShapeToBodiesMap.Add(PShape, FWeldInfo(TheirBody, RelativeTM));
-		}
-
+			if(!ShapeToBodiesMap.IsValid())
+			{
+				ShapeToBodiesMap = TSharedPtr<TMap<physx::PxShape*, FWeldInfo>> (new TMap<physx::PxShape*, FWeldInfo>());
+			}
 
+			for (int32 ShapeIdx = 0; ShapeIdx < PNewShapes.Num(); ++ShapeIdx)
+			{
+				PxShape* PShape = PNewShapes[ShapeIdx];
+				ShapeToBodiesMap->Add(PShape, FWeldInfo(TheirBody, RelativeTM));
+			}
+		}
 
 		PostShapeChange();
 
@@ -1925,25 +1931,25 @@ void FBodyInstance::UnWeld(FBodyInstance* TheirBI)
 		TArray<physx::PxShape *> PShapes;
 		const int32 NumSyncShapes = GetAllShapes_AssumesLocked(PShapes);
 
-	for (int32 ShapeIdx = 0; ShapeIdx < NumSyncShapes; ++ShapeIdx)
-	{
-		PxShape* PShape = PShapes[ShapeIdx];
-			const FBodyInstance* BI = GetOriginalBodyInstance(PShape);
-			if (TheirBI == BI)
-			{
-				ShapeToBodiesMap.Remove(PShape);
-				RigidActorSync->detachShape(*PShape);
-				bShapesChanged = true;
-			}
+		for (int32 ShapeIdx = 0; ShapeIdx < NumSyncShapes; ++ShapeIdx)
+		{
+			PxShape* PShape = PShapes[ShapeIdx];
+				const FBodyInstance* BI = GetOriginalBodyInstance(PShape);
+				if (TheirBI == BI)
+				{
+					ShapeToBodiesMap->Remove(PShape);
+					RigidActorSync->detachShape(*PShape);
+					bShapesChanged = true;
+				}
 		}
 
-	for (int32 ShapeIdx = NumSyncShapes; ShapeIdx <PShapes.Num(); ++ShapeIdx)
-	{
-		PxShape* PShape = PShapes[ShapeIdx];
+		for (int32 ShapeIdx = NumSyncShapes; ShapeIdx <PShapes.Num(); ++ShapeIdx)
+		{
+			PxShape* PShape = PShapes[ShapeIdx];
 			const FBodyInstance* BI = GetOriginalBodyInstance(PShape);
 			if (TheirBI == BI)
 			{
-				ShapeToBodiesMap.Remove(PShape);
+				ShapeToBodiesMap->Remove(PShape);
 				RigidActorAsync->detachShape(*PShape);
 				bShapesChanged = true;
 			}
diff --git a/Engine/Source/Runtime/Engine/Private/PhysicsEngine/PhysScene.cpp b/Engine/Source/Runtime/Engine/Private/PhysicsEngine/PhysScene.cpp
--- a/Engine/Source/Runtime/Engine/Private/PhysicsEngine/PhysScene.cpp
+++ b/Engine/Source/Runtime/Engine/Private/PhysicsEngine/PhysScene.cpp
@@ -24,7 +24,7 @@
 #define USE_ADAPTIVE_FORCES_FOR_ASYNC_SCENE			1
 #define USE_SPECIAL_FRICTION_MODEL_FOR_ASYNC_SCENE	0
 
-static int32 PhysXSceneCount = 1;
+static int16 PhysXSceneCount = 1;
 static const int PhysXSlowRebuildRate = 10;
 
 EPhysicsSceneType FPhysScene::SceneType_AssumesLocked(const FBodyInstance* BodyInstance) const
diff --git a/Engine/Source/Runtime/Engine/Private/PhysicsEngine/PhysXSupport.cpp b/Engine/Source/Runtime/Engine/Private/PhysicsEngine/PhysXSupport.cpp
--- a/Engine/Source/Runtime/Engine/Private/PhysicsEngine/PhysXSupport.cpp
+++ b/Engine/Source/Runtime/Engine/Private/PhysicsEngine/PhysXSupport.cpp
@@ -35,13 +35,13 @@ ENGINE_API NxModule*				GApexModuleLegacy = NULL;
 ENGINE_API NxModuleClothing*		GApexModuleClothing		= NULL;	
 #endif //WITH_APEX_CLOTHING
 
-TMap<int32, NxApexScene*>				GPhysXSceneMap;
+TMap<int16, NxApexScene*>				GPhysXSceneMap;
 FApexNullRenderResourceManager		GApexNullRenderResourceManager;
 FApexResourceCallback				GApexResourceCallback;
 FApexPhysX3Interface				GApexPhysX3Interface;
 FApexChunkReport					GApexChunkReport;
 #else	// #if WITH_APEX
-TMap<int32, PxScene*>		GPhysXSceneMap;
+TMap<int16, PxScene*>		GPhysXSceneMap;
 #endif	// #if WITH_APEX
 
 int32						GNumPhysXConvexMeshes = 0;
diff --git a/Engine/Source/Runtime/Engine/Private/PhysicsEngine/PhysXSupport.h b/Engine/Source/Runtime/Engine/Private/PhysicsEngine/PhysXSupport.h
--- a/Engine/Source/Runtime/Engine/Private/PhysicsEngine/PhysXSupport.h
+++ b/Engine/Source/Runtime/Engine/Private/PhysicsEngine/PhysXSupport.h
@@ -105,13 +105,13 @@ extern PxProfileZoneManager*	GPhysXProfileZoneManager;
  *	Map from SceneIndex to actual NxApexScene. This indirection allows us to set it to null when we kill the scene, 
  *	and therefore abort trying to destroy PhysX objects after the scene has been destroyed (eg. on game exit). 
  */
-extern TMap<int32, NxApexScene*>	GPhysXSceneMap;
+extern TMap<int16, NxApexScene*>	GPhysXSceneMap;
 #else // #if WITH_APEX
 /** 
  *	Map from SceneIndex to actual PxScene. This indirection allows us to set it to null when we kill the scene, 
  *	and therefore abort trying to destroy PhysX objects after the scene has been destroyed (eg. on game exit). 
  */
-extern TMap<int32, PxScene*>	GPhysXSceneMap;
+extern TMap<int16, PxScene*>	GPhysXSceneMap;
 #endif // #if WITH_APEX
 
 /** Total number of PhysX convex meshes around currently. */
diff --git a/Engine/Source/Runtime/Engine/Public/PhysicsPublic.h b/Engine/Source/Runtime/Engine/Public/PhysicsPublic.h
--- a/Engine/Source/Runtime/Engine/Public/PhysicsPublic.h
+++ b/Engine/Source/Runtime/Engine/Public/PhysicsPublic.h
@@ -289,7 +289,7 @@ public:
 	UWorld*							OwningWorld;
 
 	/** These indices are used to get the actual PxScene or NxApexScene from the GPhysXSceneMap. */
-	int32								PhysXSceneIndex[PST_MAX];
+	int16								PhysXSceneIndex[PST_MAX];
 
 	/** Whether or not the given scene is between its execute and sync point. */
 	bool							bPhysXSceneExecuting[PST_MAX];
