Fix NonUnity build

diff --git a/Engine/Source/Editor/DetailCustomizations/Private/BodyInstanceCustomization.cpp b/Engine/Source/Editor/DetailCustomizations/Private/BodyInstanceCustomization.cpp
--- a/Engine/Source/Editor/DetailCustomizations/Private/BodyInstanceCustomization.cpp
+++ b/Engine/Source/Editor/DetailCustomizations/Private/BodyInstanceCustomization.cpp
@@ -7,6 +7,7 @@
 #include "Engine/CollisionProfile.h"
 #include "SNumericEntryBox.h"
 #include "PhysicsEngine/BodySetup.h"
+#include "ObjectEditorUtils.h"
 
 #define LOCTEXT_NAMESPACE "BodyInstanceCustomization"
 
