Fixed a crash selecting any particle system in the level editor.  The collision category was hidden but the customization did not take that into account

https://jira.ol.epicgames.net/browse/UE-18763

diff --git a/Engine/Source/Editor/DetailCustomizations/Private/BodyInstanceCustomization.cpp b/Engine/Source/Editor/DetailCustomizations/Private/BodyInstanceCustomization.cpp
--- a/Engine/Source/Editor/DetailCustomizations/Private/BodyInstanceCustomization.cpp
+++ b/Engine/Source/Editor/DetailCustomizations/Private/BodyInstanceCustomization.cpp
@@ -994,53 +994,57 @@ void FBodyInstanceCustomizationHelper::UpdateFilters()
 
 void FBodyInstanceCustomizationHelper::CustomizeDetails( IDetailLayoutBuilder& DetailBuilder, TSharedRef<IPropertyHandle> BodyInstanceHandler)
 {
-	UpdateFilters();
+	if( BodyInstanceHandler->IsValidHandle() )
+	{
+		UpdateFilters();
 
-	IDetailCategoryBuilder& PhysicsCategory = DetailBuilder.EditCategory("Physics");
+		IDetailCategoryBuilder& PhysicsCategory = DetailBuilder.EditCategory("Physics");
 
-	TSharedRef<IPropertyHandle> PhysicsEnable = BodyInstanceHandler->GetChildHandle(GET_MEMBER_NAME_CHECKED(FBodyInstance, bSimulatePhysics)).ToSharedRef();
-	if(bDisplayEnablePhysics)
-	{
-		PhysicsCategory.AddProperty(PhysicsEnable).EditCondition(TAttribute<bool>(this, &FBodyInstanceCustomizationHelper::IsSimulatePhysicsEditable), NULL);
-	}else
-	{
-		PhysicsEnable->MarkHiddenByCustomization();
-	}
+		TSharedRef<IPropertyHandle> PhysicsEnable = BodyInstanceHandler->GetChildHandle(GET_MEMBER_NAME_CHECKED(FBodyInstance, bSimulatePhysics)).ToSharedRef();
+		if(bDisplayEnablePhysics)
+		{
+			PhysicsCategory.AddProperty(PhysicsEnable).EditCondition(TAttribute<bool>(this, &FBodyInstanceCustomizationHelper::IsSimulatePhysicsEditable), NULL);
+		}
+		else
+		{
+			PhysicsEnable->MarkHiddenByCustomization();
+		}
 
-	AddMassInKg(PhysicsCategory, BodyInstanceHandler);
+		AddMassInKg(PhysicsCategory, BodyInstanceHandler);
 
-	PhysicsCategory.AddProperty(BodyInstanceHandler->GetChildHandle(GET_MEMBER_NAME_CHECKED(FBodyInstance, LinearDamping)));
-	PhysicsCategory.AddProperty(BodyInstanceHandler->GetChildHandle(GET_MEMBER_NAME_CHECKED(FBodyInstance, AngularDamping)));
-	PhysicsCategory.AddProperty(BodyInstanceHandler->GetChildHandle(GET_MEMBER_NAME_CHECKED(FBodyInstance, bEnableGravity)));
+		PhysicsCategory.AddProperty(BodyInstanceHandler->GetChildHandle(GET_MEMBER_NAME_CHECKED(FBodyInstance, LinearDamping)));
+		PhysicsCategory.AddProperty(BodyInstanceHandler->GetChildHandle(GET_MEMBER_NAME_CHECKED(FBodyInstance, AngularDamping)));
+		PhysicsCategory.AddProperty(BodyInstanceHandler->GetChildHandle(GET_MEMBER_NAME_CHECKED(FBodyInstance, bEnableGravity)));
 
-	AddBodyConstraint(PhysicsCategory, BodyInstanceHandler);
+		AddBodyConstraint(PhysicsCategory, BodyInstanceHandler);
 
-	//ADVANCED
-	PhysicsCategory.AddProperty(BodyInstanceHandler->GetChildHandle(GET_MEMBER_NAME_CHECKED(FBodyInstance, bAutoWeld)))
-	.Visibility(TAttribute<EVisibility>(this, &FBodyInstanceCustomizationHelper::IsAutoWeldVisible))
-	.EditCondition(TAttribute<bool>(this, &FBodyInstanceCustomizationHelper::IsAutoWeldEditable), NULL);
+		//ADVANCED
+		PhysicsCategory.AddProperty(BodyInstanceHandler->GetChildHandle(GET_MEMBER_NAME_CHECKED(FBodyInstance, bAutoWeld)))
+			.Visibility(TAttribute<EVisibility>(this, &FBodyInstanceCustomizationHelper::IsAutoWeldVisible))
+			.EditCondition(TAttribute<bool>(this, &FBodyInstanceCustomizationHelper::IsAutoWeldEditable), NULL);
 
-	PhysicsCategory.AddProperty(BodyInstanceHandler->GetChildHandle(GET_MEMBER_NAME_CHECKED(FBodyInstance, bStartAwake)));
+		PhysicsCategory.AddProperty(BodyInstanceHandler->GetChildHandle(GET_MEMBER_NAME_CHECKED(FBodyInstance, bStartAwake)));
 
-	PhysicsCategory.AddProperty(BodyInstanceHandler->GetChildHandle(GET_MEMBER_NAME_CHECKED(FBodyInstance, COMNudge)));
-	PhysicsCategory.AddProperty(BodyInstanceHandler->GetChildHandle(GET_MEMBER_NAME_CHECKED(FBodyInstance, MassScale)));
+		PhysicsCategory.AddProperty(BodyInstanceHandler->GetChildHandle(GET_MEMBER_NAME_CHECKED(FBodyInstance, COMNudge)));
+		PhysicsCategory.AddProperty(BodyInstanceHandler->GetChildHandle(GET_MEMBER_NAME_CHECKED(FBodyInstance, MassScale)));
 
-	AddMaxAngularVelocity(PhysicsCategory, BodyInstanceHandler);
+		AddMaxAngularVelocity(PhysicsCategory, BodyInstanceHandler);
 
-	TSharedRef<IPropertyHandle> AsyncEnabled = BodyInstanceHandler->GetChildHandle(GET_MEMBER_NAME_CHECKED(FBodyInstance, bUseAsyncScene)).ToSharedRef();
-	PhysicsCategory.AddProperty(AsyncEnabled).EditCondition(TAttribute<bool>(this, &FBodyInstanceCustomizationHelper::IsUseAsyncEditable), NULL);
+		TSharedRef<IPropertyHandle> AsyncEnabled = BodyInstanceHandler->GetChildHandle(GET_MEMBER_NAME_CHECKED(FBodyInstance, bUseAsyncScene)).ToSharedRef();
+		PhysicsCategory.AddProperty(AsyncEnabled).EditCondition(TAttribute<bool>(this, &FBodyInstanceCustomizationHelper::IsUseAsyncEditable), NULL);
 
-	//Add the rest
-	uint32 NumChildren = 0;
-	BodyInstanceHandler->GetNumChildren(NumChildren);
-	for(uint32 ChildIdx = 0; ChildIdx < NumChildren; ++ChildIdx)
-	{
-		TSharedPtr<IPropertyHandle> ChildProp = BodyInstanceHandler->GetChildHandle(ChildIdx);
-		
-		FString Category = FObjectEditorUtils::GetCategory(ChildProp->GetProperty());
-		if(ChildProp->IsCustomized() == false && Category == TEXT("Physics"))	//add the rest of the physics properties
+		//Add the rest
+		uint32 NumChildren = 0;
+		BodyInstanceHandler->GetNumChildren(NumChildren);
+		for(uint32 ChildIdx = 0; ChildIdx < NumChildren; ++ChildIdx)
 		{
-			PhysicsCategory.AddProperty(ChildProp);
+			TSharedPtr<IPropertyHandle> ChildProp = BodyInstanceHandler->GetChildHandle(ChildIdx);
+
+			FString Category = FObjectEditorUtils::GetCategory(ChildProp->GetProperty());
+			if(ChildProp->IsCustomized() == false && Category == TEXT("Physics"))	//add the rest of the physics properties
+			{
+				PhysicsCategory.AddProperty(ChildProp);
+			}
 		}
 	}
 }
