Editing FText properties in the user defined struct editor and for local Blueprint function variables will no longer erase the inputed information upon closing and re-opening the editor.

Added FPropertyHandleBase::NotifyFinishedChangingProperties to handle generic calls to NotifyFinishedChangingProperties.

Text Customization will now NotifyFinishedChangingProperties so the data can properly be captured.

diff --git a/Engine/Source/Editor/DetailCustomizations/Private/TextCustomization.cpp b/Engine/Source/Editor/DetailCustomizations/Private/TextCustomization.cpp
--- a/Engine/Source/Editor/DetailCustomizations/Private/TextCustomization.cpp
+++ b/Engine/Source/Editor/DetailCustomizations/Private/TextCustomization.cpp
@@ -49,20 +49,25 @@ namespace
 
 		const auto& OnTextCommitted = [this](const FText& NewText, ETextCommit::Type CommitInfo)
 		{
-			PropertyHandle->NotifyPreChange();
-
 			TArray<void*> RawData;
 			PropertyHandle->AccessRawData(RawData);
-			for (void* const RawDatum : RawData)
+
+			if (RawData.Num() > 0)
 			{
-				FText& PropertyValue = *(reinterpret_cast<FText* const>(RawDatum));
+				PropertyHandle->NotifyPreChange();
+
+				for (void* const RawDatum : RawData)
+				{
+					FText& PropertyValue = *(reinterpret_cast<FText* const>(RawDatum));
+
+					// FText::FromString on the result of FText::ToString is intentional. For now, we want to nuke any namespace/key info and let it get regenerated from scratch,
+					// rather than risk adopting whatever came through some chain of calls. This will be replaced when preserving of identity is implemented.
+					PropertyValue = FText::FromString(NewText.ToString());
+				}
 
-				// FText::FromString on the result of FText::ToString is intentional. For now, we want to nuke any namespace/key info and let it get regenerated from scratch,
-				// rather than risk adopting whatever came through some chain of calls. This will be replaced when preserving of identity is implemented.
-				PropertyValue = FText::FromString(NewText.ToString());
+				PropertyHandle->NotifyPostChange();
+				PropertyHandle->NotifyFinishedChangingProperties();
 			}
-			
-			PropertyHandle->NotifyPostChange();
 		};
 
 		ChildSlot
diff --git a/Engine/Source/Editor/PropertyEditor/Private/PropertyHandleImpl.cpp b/Engine/Source/Editor/PropertyEditor/Private/PropertyHandleImpl.cpp
--- a/Engine/Source/Editor/PropertyEditor/Private/PropertyHandleImpl.cpp
+++ b/Engine/Source/Editor/PropertyEditor/Private/PropertyHandleImpl.cpp
@@ -2003,6 +2003,17 @@ void FPropertyHandleBase::NotifyPostChange()
 	}
 }
 
+void FPropertyHandleBase::NotifyFinishedChangingProperties()
+{
+	TSharedPtr<FPropertyNode> PropertyNode = Implementation->GetPropertyNode();
+	if( PropertyNode.IsValid() )
+	{
+		FPropertyChangedEvent ChangeEvent(PropertyNode->GetProperty(), EPropertyChangeType::ValueSet);
+		PropertyNode->FixPropertiesInEvent(ChangeEvent);
+		Implementation->GetPropertyUtilities()->NotifyFinishedChangingProperties(ChangeEvent);
+	}
+}
+
 FPropertyAccess::Result FPropertyHandleBase::SetObjectValueFromSelection()
 {
 	return Implementation->OnUseSelected();
diff --git a/Engine/Source/Editor/PropertyEditor/Private/PropertyHandleImpl.h b/Engine/Source/Editor/PropertyEditor/Private/PropertyHandleImpl.h
--- a/Engine/Source/Editor/PropertyEditor/Private/PropertyHandleImpl.h
+++ b/Engine/Source/Editor/PropertyEditor/Private/PropertyHandleImpl.h
@@ -406,6 +406,7 @@ public:
 	virtual FPropertyAccess::Result SetObjectValueFromSelection() override;
 	virtual void NotifyPreChange() override;
 	virtual void NotifyPostChange() override;
+	virtual void NotifyFinishedChangingProperties() override;
 	virtual void AddRestriction( TSharedRef<const FPropertyRestriction> Restriction )override;
 	virtual bool IsRestricted(const FString& Value) const override;
 	virtual bool IsRestricted(const FString& Value, TArray<FText>& OutReasons) const override;
diff --git a/Engine/Source/Editor/PropertyEditor/Public/PropertyHandle.h b/Engine/Source/Editor/PropertyEditor/Public/PropertyHandle.h
--- a/Engine/Source/Editor/PropertyEditor/Public/PropertyHandle.h
+++ b/Engine/Source/Editor/PropertyEditor/Public/PropertyHandle.h
@@ -236,6 +236,12 @@ public:
 	 */
 	virtual void NotifyPostChange() = 0;
 
+	/**
+	 * Called to manually notify root objects that this property has finished changing
+	 * This does not need to be called when SetValue functions are used since it will be called automatically
+	 */
+	virtual void NotifyFinishedChangingProperties() = 0;
+
 	/**
 	 * Sets the object value from the current editor selection
 	 * Will fail if this handle isn't an object property
