Fixed text property detail customization support for multiline text.

diff --git a/Engine/Source/Editor/DetailCustomizations/Private/TextCustomization.cpp b/Engine/Source/Editor/DetailCustomizations/Private/TextCustomization.cpp
--- a/Engine/Source/Editor/DetailCustomizations/Private/TextCustomization.cpp
+++ b/Engine/Source/Editor/DetailCustomizations/Private/TextCustomization.cpp
@@ -7,19 +7,41 @@ namespace
 {
 	class STextPropertyWidget : public SCompoundWidget
 	{
-		SLATE_BEGIN_ARGS(STextPropertyWidget) {}
+		SLATE_BEGIN_ARGS(STextPropertyWidget)
+			: _Font( FEditorStyle::GetFontStyle( TEXT("PropertyWindow.NormalFont") ) ) 
+			{}
+			SLATE_ATTRIBUTE( FSlateFontInfo, Font )
 		SLATE_END_ARGS()
 
 	public:
-		void Construct(const FArguments& Arguments, const TSharedRef<IPropertyHandle>& InPropertyHandle);
+		void Construct(const FArguments& Arguments, const TSharedRef<IPropertyHandle>& InPropertyHandle, const TSharedPtr<IPropertyUtilities>& InPropertyUtilities);
+		void GetDesiredWidth( float& OutMinDesiredWidth, float& OutMaxDesiredWidth );
+		bool SupportsKeyboardFocus() const;
+		FReply OnFocusReceived( const FGeometry& MyGeometry, const FFocusEvent& InFocusEvent );
+		void Tick( const FGeometry& AllottedGeometry, const double InCurrentTime, const float InDeltaTime );
+		bool CanEdit() const;
+		bool IsReadOnly() const;
 
 	private:
 		TSharedPtr<IPropertyHandle> PropertyHandle;
+
+		TSharedPtr<IPropertyUtilities> PropertyUtilities;
+
+		TSharedPtr< class SWidget > PrimaryWidget;
+
+		TSharedPtr<SMultiLineEditableTextBox> MultiLineWidget;
+
+		TSharedPtr<SEditableTextBox> SingleLineWidget;
+
+		TOptional< float > PreviousHeight;
+
+		bool bIsMultiLine;
 	};
 
-	void STextPropertyWidget::Construct(const FArguments& Arguments, const TSharedRef<IPropertyHandle>& InPropertyHandle)
+	void STextPropertyWidget::Construct(const FArguments& InArgs, const TSharedRef<IPropertyHandle>& InPropertyHandle, const TSharedPtr<IPropertyUtilities>& InPropertyUtilities)
 	{
 		PropertyHandle = InPropertyHandle;
+		PropertyUtilities = InPropertyUtilities;
 
 		const auto& GetTextValue = [this]() -> FText
 		{
@@ -70,12 +92,101 @@ namespace
 			}
 		};
 
-		ChildSlot
-			[
-				SNew(SEditableTextBox)
-				.Text_Lambda(GetTextValue)
-				.OnTextCommitted_Lambda(OnTextCommitted)
-			];
+		TSharedPtr<SHorizontalBox> HorizontalBox;
+
+		bIsMultiLine = PropertyHandle->GetBoolMetaData("MultiLine");
+		if(bIsMultiLine)
+		{
+			ChildSlot
+				[
+					SAssignNew(HorizontalBox, SHorizontalBox)
+					+SHorizontalBox::Slot()
+					.FillWidth(1.0f)
+					[
+						SAssignNew(MultiLineWidget, SMultiLineEditableTextBox)
+						.Text_Lambda(GetTextValue)
+						.Font(InArgs._Font)
+						.SelectAllTextWhenFocused(false)
+						.ClearKeyboardFocusOnCommit(false)
+						.OnTextCommitted_Lambda(OnTextCommitted)
+						.SelectAllTextOnCommit(false)
+						.IsReadOnly(this, &STextPropertyWidget::IsReadOnly)
+						.AutoWrapText(true)
+						.ModiferKeyForNewLine(EModifierKey::Shift)
+					]
+				];
+
+			PrimaryWidget = MultiLineWidget;
+		}
+		else
+		{
+			ChildSlot
+				[
+					SAssignNew(HorizontalBox, SHorizontalBox)
+					+SHorizontalBox::Slot()
+					.FillWidth(1.0f)
+					[
+						SAssignNew( SingleLineWidget, SEditableTextBox )
+						.Text_Lambda(GetTextValue)
+						.Font( InArgs._Font )
+						.SelectAllTextWhenFocused( true )
+						.ClearKeyboardFocusOnCommit(false)
+						.OnTextCommitted_Lambda(OnTextCommitted)
+						.SelectAllTextOnCommit( true )
+						.IsReadOnly(this, &STextPropertyWidget::IsReadOnly)
+
+					]
+				];
+
+			PrimaryWidget = SingleLineWidget;
+		}
+
+		SetEnabled( TAttribute<bool>( this, &STextPropertyWidget::CanEdit ) );
+	}
+
+	void STextPropertyWidget::GetDesiredWidth( float& OutMinDesiredWidth, float& OutMaxDesiredWidth )
+	{
+		if(bIsMultiLine)
+		{
+			OutMinDesiredWidth = 250.0f;
+		}
+		else
+		{
+			OutMinDesiredWidth = 125.0f;
+		}
+
+		OutMaxDesiredWidth = 600.0f;
+	}
+
+	bool STextPropertyWidget::SupportsKeyboardFocus() const
+	{
+		return PrimaryWidget.IsValid() && PrimaryWidget->SupportsKeyboardFocus();
+	}
+
+	FReply STextPropertyWidget::OnFocusReceived( const FGeometry& MyGeometry, const FFocusEvent& InFocusEvent )
+	{
+		// Forward keyboard focus to our editable text widget
+		return FReply::Handled().SetUserFocus(PrimaryWidget.ToSharedRef(), InFocusEvent.GetCause());
+	}
+
+	void STextPropertyWidget::Tick( const FGeometry& AllottedGeometry, const double InCurrentTime, const float InDeltaTime )
+	{
+		const float CurrentHeight = AllottedGeometry.GetLocalSize().Y;
+		if (bIsMultiLine && PreviousHeight.IsSet() && PreviousHeight.GetValue() != CurrentHeight && PropertyUtilities.IsValid())
+		{
+			PropertyUtilities->RequestRefresh();
+		}
+		PreviousHeight = CurrentHeight;
+	}
+
+	bool STextPropertyWidget::CanEdit() const
+	{
+		return PropertyHandle.IsValid() ? !PropertyHandle->IsEditConst() : true;
+	}
+
+	bool STextPropertyWidget::IsReadOnly() const
+	{
+		return !CanEdit();
 	}
 }
 
@@ -88,7 +199,7 @@ void FTextCustomization::CustomizeHeader( TSharedRef<class IPropertyHandle> InPr
 		]
 		.ValueContent()
 		[
-			SNew(STextPropertyWidget, InPropertyHandle)
+			SNew(STextPropertyWidget, InPropertyHandle, PropertyTypeCustomizationUtils.GetPropertyUtilities())
 		];
 }
 
