Property editor widgets for Timespan and DateTime call NotifyFinishedChangingProperties (it's necessary for USStruct editor).

diff --git a/Engine/Source/Editor/DetailCustomizations/Private/DateTimeStructCustomization.cpp b/Engine/Source/Editor/DetailCustomizations/Private/DateTimeStructCustomization.cpp
--- a/Engine/Source/Editor/DetailCustomizations/Private/DateTimeStructCustomization.cpp
+++ b/Engine/Source/Editor/DetailCustomizations/Private/DateTimeStructCustomization.cpp
@@ -99,6 +99,7 @@ void FDateTimeStructCustomization::HandleTextBoxTextCommited( const FText& NewTe
 			*(FDateTime*)RawDataInstance = ParsedDateTime;
 		}
 		PropertyHandle->NotifyPostChange();
+		PropertyHandle->NotifyFinishedChangingProperties();
 	}
 }
 
diff --git a/Engine/Source/Editor/DetailCustomizations/Private/TimespanStructCustomization.cpp b/Engine/Source/Editor/DetailCustomizations/Private/TimespanStructCustomization.cpp
--- a/Engine/Source/Editor/DetailCustomizations/Private/TimespanStructCustomization.cpp
+++ b/Engine/Source/Editor/DetailCustomizations/Private/TimespanStructCustomization.cpp
@@ -100,6 +100,7 @@ void FTimespanStructCustomization::HandleTextBoxTextCommited( const FText& NewTe
 			*(FTimespan*)RawDataInstance = ParsedTimespan;
 		}
 		PropertyHandle->NotifyPostChange();
+		PropertyHandle->NotifyFinishedChangingProperties();
 	}
 }
 
