"Reset to default" on the Image Size property of a brush now uses the size of the texture if one is set.
[UETOOL-125] - Have the "return to defaults" arrow on a brush's texture actually go back to the size of the texture

#change Customization of property rows now supports passing reset to default customizations to child properties so that structs like vectors can have their components handled by the same custom handler as the parent. Properties identify themselves by passing their handle to the delegate.
#change Reset to default is now set on properties using FResetToDefaultOverride class containing the delegates and ovrride behaviors.
#change Slate brush customization sets custom reset behavior on the Image Size property and uses the texture size as the default size.

diff --git a/Engine/Source/Editor/DetailCustomizations/Private/SlateBrushCustomization.cpp b/Engine/Source/Editor/DetailCustomizations/Private/SlateBrushCustomization.cpp
--- a/Engine/Source/Editor/DetailCustomizations/Private/SlateBrushCustomization.cpp
+++ b/Engine/Source/Editor/DetailCustomizations/Private/SlateBrushCustomization.cpp
@@ -994,9 +994,9 @@ class SSlateBrushStaticPreview : public SCompoundWidget
 
 		if (ResourceObjectProperty.IsValid() && ResourceObjectProperty->IsValidHandle())
 		{
-			ResourceObjectProperty->AccessRawData(RawData);
+		ResourceObjectProperty->AccessRawData(RawData);
 
-			check(RawData[0] != NULL);
+		check(RawData[0] != NULL);
 			TemporaryBrush = *static_cast<FSlateBrush*>(RawData[0]);
 		}
 	}
@@ -1235,12 +1235,12 @@ void FSlateBrushStructCustomization::CustomizeHeader( TSharedRef<IPropertyHandle
 void FSlateBrushStructCustomization::CustomizeChildren( TSharedRef<IPropertyHandle> StructPropertyHandle, IDetailChildrenBuilder& StructBuilder, IPropertyTypeCustomizationUtils& StructCustomizationUtils )
 {
 	// Add the child properties
-	TSharedPtr<IPropertyHandle> ImageSizeProperty = StructPropertyHandle->GetChildHandle( TEXT("ImageSize") );
+	ImageSizeProperty = StructPropertyHandle->GetChildHandle( TEXT("ImageSize") );
 	DrawAsProperty = StructPropertyHandle->GetChildHandle( TEXT("DrawAs") );
 	TSharedPtr<IPropertyHandle> TilingProperty = StructPropertyHandle->GetChildHandle( TEXT("Tiling") );
 	TSharedPtr<IPropertyHandle> MarginProperty = StructPropertyHandle->GetChildHandle( TEXT("Margin") );
 	TSharedPtr<IPropertyHandle> TintProperty = StructPropertyHandle->GetChildHandle( TEXT("TintColor") );
-	TSharedPtr<IPropertyHandle> ResourceObjectProperty = StructPropertyHandle->GetChildHandle( TEXT("ResourceObject") );
+	ResourceObjectProperty = StructPropertyHandle->GetChildHandle( TEXT("ResourceObject") );
 	
 	StructBuilder.AddChildContent( NSLOCTEXT( "SlateBrushCustomization", "ResourceObjectFilterString", "Image" ) )
 	.NameContent()
@@ -1254,7 +1254,14 @@ void FSlateBrushStructCustomization::CustomizeChildren( TSharedRef<IPropertyHand
 		SNew( SBrushResourceObjectBox, &StructCustomizationUtils, ResourceObjectProperty, ImageSizeProperty )
 	];
 
-	StructBuilder.AddChildProperty( ImageSizeProperty.ToSharedRef() );
+	// Add the image size property with custom reset delegates that also affect the child properties (the components)
+	const bool bOverrideDefaultOnVectorChildren = true;
+	StructBuilder.AddChildProperty( ImageSizeProperty.ToSharedRef() )
+	.OverrideResetToDefault(FResetToDefaultOverride::Create(
+		FIsResetToDefaultVisible::CreateSP(this, &FSlateBrushStructCustomization::IsImageSizeResetToDefaultVisible),
+		FResetToDefaultHandler::CreateSP(this, &FSlateBrushStructCustomization::OnImageSizeResetToDefault),
+		bOverrideDefaultOnVectorChildren));
+
 	StructBuilder.AddChildProperty(TintProperty.ToSharedRef());
 	StructBuilder.AddChildProperty( DrawAsProperty.ToSharedRef() );
 	StructBuilder.AddChildProperty( TilingProperty.ToSharedRef() )
@@ -1324,3 +1331,74 @@ EVisibility FSlateBrushStructCustomization::GetMarginPropertyVisibility() const
 
 	return (Result == FPropertyAccess::MultipleValues || DrawAsType == ESlateBrushDrawType::Box || DrawAsType == ESlateBrushDrawType::Border) ? EVisibility::Visible : EVisibility::Collapsed;
 }
+
+bool FSlateBrushStructCustomization::IsImageSizeResetToDefaultVisible(TSharedRef<IPropertyHandle> PropertyHandle) const
+{
+	// get texture size from ResourceObjectProperty and compare to image size prop value
+	const FVector2D SizeDefault = GetDefaultImageSize();
+
+	FVector2D Size;
+	ImageSizeProperty->GetValue(Size);
+
+	if (PropertyHandle->GetProperty() == ImageSizeProperty->GetProperty())
+	{
+		// reseting the whole vector
+		return SizeDefault != Size;
+	}
+	else if (PropertyHandle->GetProperty() == ImageSizeProperty->GetChildHandle(0)->GetProperty())	// X
+	{
+		// reseting the vector.X
+		return SizeDefault.X != Size.X;
+	}
+	else if (PropertyHandle->GetProperty() == ImageSizeProperty->GetChildHandle(1)->GetProperty())	// Y
+	{
+		// reseting the vector.Y
+		return SizeDefault.Y != Size.Y;
+	}
+
+	ensureMsgf(false, TEXT("Property handle mismatch in brush size FVector2D struct"));
+	return false;
+}
+
+void FSlateBrushStructCustomization::OnImageSizeResetToDefault(TSharedRef<IPropertyHandle> PropertyHandle) const
+{
+	// set image size prop value to the texture size in ResourceObjectProperty
+	const FVector2D SizeDefault = GetDefaultImageSize();
+
+	if (PropertyHandle->GetProperty() == ImageSizeProperty->GetProperty())
+	{
+		// reseting the whole vector
+		PropertyHandle->SetValue(SizeDefault);
+	}
+	else if (PropertyHandle->GetProperty() == ImageSizeProperty->GetChildHandle(0)->GetProperty())	// X
+	{
+		// reseting the vector.X
+		PropertyHandle->SetValue(SizeDefault.X);
+	}
+	else if (PropertyHandle->GetProperty() == ImageSizeProperty->GetChildHandle(1)->GetProperty())	// Y
+	{
+		// reseting the vector.Y
+		PropertyHandle->SetValue(SizeDefault.Y);
+	}
+	else
+	{
+		ensureMsgf(false, TEXT("Property handle mismatch in brush size FVector2D struct"));
+	}
+}
+
+FVector2D FSlateBrushStructCustomization::GetDefaultImageSize() const
+{
+	// Custom default behavior using the texture's size, if one is set as the resource object
+	UObject* ResourceObject;
+	if (FPropertyAccess::Success == ResourceObjectProperty->GetValue(ResourceObject))
+	{
+		UTexture2D* Texture = Cast<UTexture2D>(ResourceObject);
+		if (Texture != nullptr)
+		{
+			return FVector2D(Texture->GetSizeX(), Texture->GetSizeY());
+		}
+	}
+
+	// Fall back on the standard default size for brush images
+	return FVector2D(SlateBrushDefs::DefaultImageSize, SlateBrushDefs::DefaultImageSize);
+}
diff --git a/Engine/Source/Editor/DetailCustomizations/Public/Customizations/SlateBrushCustomization.h b/Engine/Source/Editor/DetailCustomizations/Public/Customizations/SlateBrushCustomization.h
--- a/Engine/Source/Editor/DetailCustomizations/Public/Customizations/SlateBrushCustomization.h
+++ b/Engine/Source/Editor/DetailCustomizations/Public/Customizations/SlateBrushCustomization.h
@@ -3,6 +3,9 @@
 #pragma once
 #include "IPropertyTypeCustomization.h"
 
+class IPropertyHandle;
+
+
 class DETAILCUSTOMIZATIONS_API FSlateBrushStructCustomization : public IPropertyTypeCustomization
 {
 public:
@@ -11,9 +14,9 @@ public:
 	FSlateBrushStructCustomization(bool bIncludePreview);
 
 	/** IPropertyTypeCustomization interface */
-	virtual void CustomizeHeader( TSharedRef<class IPropertyHandle> StructPropertyHandle, class FDetailWidgetRow& HeaderRow, IPropertyTypeCustomizationUtils& StructCustomizationUtils ) override;
+	virtual void CustomizeHeader( TSharedRef<IPropertyHandle> StructPropertyHandle, class FDetailWidgetRow& HeaderRow, IPropertyTypeCustomizationUtils& StructCustomizationUtils ) override;
 
-	virtual void CustomizeChildren( TSharedRef<class IPropertyHandle> StructPropertyHandle, class IDetailChildrenBuilder& StructBuilder, IPropertyTypeCustomizationUtils& StructCustomizationUtils ) override;              
+	virtual void CustomizeChildren( TSharedRef<IPropertyHandle> StructPropertyHandle, class IDetailChildrenBuilder& StructBuilder, IPropertyTypeCustomizationUtils& StructCustomizationUtils ) override;              
 
 private:
 	/**
@@ -26,9 +29,24 @@ private:
 	 */
 	EVisibility GetMarginPropertyVisibility() const;
 
+	/** Callback for determining image size reset button visibility */
+	bool IsImageSizeResetToDefaultVisible(TSharedRef<IPropertyHandle> PropertyHandle) const;
+
+	/** Callback for clicking the image size reset button */
+	void OnImageSizeResetToDefault(TSharedRef<IPropertyHandle> PropertyHandle) const;
+
+	/** Gets the current default image size based on the current texture resource */
+	FVector2D GetDefaultImageSize() const;
+
 	/** Slate Brush DrawAs property */
 	TSharedPtr<IPropertyHandle> DrawAsProperty;
 
+/** Slate Brush Image Size property */
+	TSharedPtr<IPropertyHandle> ImageSizeProperty;
+
+/** Slate Brush Resource Object property */
+	TSharedPtr<IPropertyHandle> ResourceObjectProperty;
+
 	/** Error text to display if the resource object is not valid*/
 	TSharedPtr<SErrorText> ResourceErrorText;
 
diff --git a/Engine/Source/Editor/LandscapeEditor/Private/LandscapeEditorDetailCustomization_AlphaBrush.cpp b/Engine/Source/Editor/LandscapeEditor/Private/LandscapeEditorDetailCustomization_AlphaBrush.cpp
--- a/Engine/Source/Editor/LandscapeEditor/Private/LandscapeEditorDetailCustomization_AlphaBrush.cpp
+++ b/Engine/Source/Editor/LandscapeEditor/Private/LandscapeEditorDetailCustomization_AlphaBrush.cpp
@@ -185,7 +185,7 @@ void FLandscapeEditorDetailCustomization_AlphaBrush::CustomizeDetails(IDetailLay
 	DetailBuilder.HideProperty(PropertyHandle_AlphaTextureChannel);
 
 	BrushSettingsCategory.AddProperty(PropertyHandle_AlphaTexture)
-	.OverrideResetToDefault(false, FSimpleDelegate())
+	.OverrideResetToDefault(FResetToDefaultOverride::Hide())
 	.CustomWidget()
 	.NameContent()
 	[
@@ -222,7 +222,7 @@ void FLandscapeEditorDetailCustomization_AlphaBrush::CustomizeDetails(IDetailLay
 	];
 
 	BrushSettingsCategory.AddProperty(PropertyHandle_AlphaTextureChannel)
-	.OverrideResetToDefault(false, FSimpleDelegate());
+	.OverrideResetToDefault(FResetToDefaultOverride::Hide());
 
 	if (IsBrushSetActive("BrushSet_Pattern"))
 	{
diff --git a/Engine/Source/Editor/LandscapeEditor/Private/LandscapeEditorDetailCustomization_ResizeLandscape.cpp b/Engine/Source/Editor/LandscapeEditor/Private/LandscapeEditorDetailCustomization_ResizeLandscape.cpp
--- a/Engine/Source/Editor/LandscapeEditor/Private/LandscapeEditorDetailCustomization_ResizeLandscape.cpp
+++ b/Engine/Source/Editor/LandscapeEditor/Private/LandscapeEditorDetailCustomization_ResizeLandscape.cpp
@@ -79,11 +79,11 @@ void FLandscapeEditorDetailCustomization_ResizeLandscape::CustomizeDetails(IDeta
 		]
 	];
 
-
 	TSharedRef<IPropertyHandle> PropertyHandle_QuadsPerSection = DetailBuilder.GetProperty(GET_MEMBER_NAME_CHECKED(ULandscapeEditorObject, ResizeLandscape_QuadsPerSection));
 	ResizeLandscapeCategory.AddProperty(PropertyHandle_QuadsPerSection)
-	.OverrideResetToDefault(TAttribute<bool>::Create(TAttribute<bool>::FGetter::CreateStatic(&FLandscapeEditorDetailCustomization_ResizeLandscape::IsSectionSizeResetToDefaultVisible)),
-		FSimpleDelegate::CreateStatic(&FLandscapeEditorDetailCustomization_ResizeLandscape::OnSectionSizeResetToDefault))
+	.OverrideResetToDefault(FResetToDefaultOverride::Create(
+		FIsResetToDefaultVisible::CreateStatic(&FLandscapeEditorDetailCustomization_ResizeLandscape::IsSectionSizeResetToDefaultVisible),
+		FResetToDefaultHandler::CreateStatic(&FLandscapeEditorDetailCustomization_ResizeLandscape::OnSectionSizeResetToDefault)))
 	.CustomWidget()
 	.NameContent()
 	[
@@ -119,8 +119,9 @@ void FLandscapeEditorDetailCustomization_ResizeLandscape::CustomizeDetails(IDeta
 
 	TSharedRef<IPropertyHandle> PropertyHandle_SectionsPerComponent = DetailBuilder.GetProperty(GET_MEMBER_NAME_CHECKED(ULandscapeEditorObject, ResizeLandscape_SectionsPerComponent));
 	ResizeLandscapeCategory.AddProperty(PropertyHandle_SectionsPerComponent)
-	.OverrideResetToDefault(TAttribute<bool>::Create(TAttribute<bool>::FGetter::CreateStatic(&FLandscapeEditorDetailCustomization_ResizeLandscape::IsSectionsPerComponentResetToDefaultVisible)),
-		FSimpleDelegate::CreateStatic(&FLandscapeEditorDetailCustomization_ResizeLandscape::OnSectionsPerComponentResetToDefault))
+	.OverrideResetToDefault(FResetToDefaultOverride::Create(
+		FIsResetToDefaultVisible::CreateStatic(&FLandscapeEditorDetailCustomization_ResizeLandscape::IsSectionsPerComponentResetToDefaultVisible),
+		FResetToDefaultHandler::CreateStatic(&FLandscapeEditorDetailCustomization_ResizeLandscape::OnSectionsPerComponentResetToDefault)))
 	.CustomWidget()
 	.NameContent()
 	[
@@ -172,7 +173,7 @@ void FLandscapeEditorDetailCustomization_ResizeLandscape::CustomizeDetails(IDeta
 	TSharedRef<IPropertyHandle> PropertyHandle_ComponentCount_X = PropertyHandle_ComponentCount->GetChildHandle("X").ToSharedRef();
 	TSharedRef<IPropertyHandle> PropertyHandle_ComponentCount_Y = PropertyHandle_ComponentCount->GetChildHandle("Y").ToSharedRef();
 	ResizeLandscapeCategory.AddProperty(PropertyHandle_ComponentCount)
-	.OverrideResetToDefault(false, FSimpleDelegate())
+	.OverrideResetToDefault(FResetToDefaultOverride::Hide())
 	.CustomWidget()
 	.NameContent()
 	[
@@ -340,7 +341,7 @@ FText FLandscapeEditorDetailCustomization_ResizeLandscape::GetSectionSize(TShare
 	return FText::Format(LOCTEXT("NxNQuads", "{0}x{0} Quads"), FText::AsNumber(QuadsPerSection));
 }
 
-bool FLandscapeEditorDetailCustomization_ResizeLandscape::IsSectionSizeResetToDefaultVisible()
+bool FLandscapeEditorDetailCustomization_ResizeLandscape::IsSectionSizeResetToDefaultVisible(TSharedRef<IPropertyHandle> PropertyHandle)
 {
 	FEdModeLandscape* LandscapeEdMode = GetEditorMode();
 	if (LandscapeEdMode != NULL)
@@ -351,7 +352,7 @@ bool FLandscapeEditorDetailCustomization_ResizeLandscape::IsSectionSizeResetToDe
 	return false;
 }
 
-void FLandscapeEditorDetailCustomization_ResizeLandscape::OnSectionSizeResetToDefault()
+void FLandscapeEditorDetailCustomization_ResizeLandscape::OnSectionSizeResetToDefault(TSharedRef<IPropertyHandle> PropertyHandle)
 {
 	FEdModeLandscape* LandscapeEdMode = GetEditorMode();
 	if (LandscapeEdMode != NULL)
@@ -413,7 +414,7 @@ FText FLandscapeEditorDetailCustomization_ResizeLandscape::GetSectionsPerCompone
 	return FText::Format(SectionsPerComponent == 1 ? LOCTEXT("1x1Section", "{Width}\u00D7{Height} Section") : LOCTEXT("NxNSections", "{Width}\u00D7{Height} Sections"), Args);
 }
 
-bool FLandscapeEditorDetailCustomization_ResizeLandscape::IsSectionsPerComponentResetToDefaultVisible()
+bool FLandscapeEditorDetailCustomization_ResizeLandscape::IsSectionsPerComponentResetToDefaultVisible(TSharedRef<IPropertyHandle> PropertyHandle)
 {
 	FEdModeLandscape* LandscapeEdMode = GetEditorMode();
 	if (LandscapeEdMode != NULL)
@@ -424,7 +425,7 @@ bool FLandscapeEditorDetailCustomization_ResizeLandscape::IsSectionsPerComponent
 	return false;
 }
 
-void FLandscapeEditorDetailCustomization_ResizeLandscape::OnSectionsPerComponentResetToDefault()
+void FLandscapeEditorDetailCustomization_ResizeLandscape::OnSectionsPerComponentResetToDefault(TSharedRef<IPropertyHandle> PropertyHandle)
 {
 	FEdModeLandscape* LandscapeEdMode = GetEditorMode();
 	if (LandscapeEdMode != NULL)
diff --git a/Engine/Source/Editor/LandscapeEditor/Private/LandscapeEditorDetailCustomization_ResizeLandscape.h b/Engine/Source/Editor/LandscapeEditor/Private/LandscapeEditorDetailCustomization_ResizeLandscape.h
--- a/Engine/Source/Editor/LandscapeEditor/Private/LandscapeEditorDetailCustomization_ResizeLandscape.h
+++ b/Engine/Source/Editor/LandscapeEditor/Private/LandscapeEditorDetailCustomization_ResizeLandscape.h
@@ -4,7 +4,6 @@
 
 #include "LandscapeEditorDetailCustomization_Base.h"
 
-
 /**
  * Slate widgets customizer for the "Change Landscape Component Size" tool
  */
@@ -23,15 +22,15 @@ public:
 	static TSharedRef<SWidget> GetSectionSizeMenu(TSharedRef<IPropertyHandle> PropertyHandle);
 	static void OnChangeSectionSize(TSharedRef<IPropertyHandle> PropertyHandle, int32 NewSize);
 	static FText GetSectionSize(TSharedRef<IPropertyHandle> PropertyHandle);
-	static bool IsSectionSizeResetToDefaultVisible();
-	static void OnSectionSizeResetToDefault();
+	static bool IsSectionSizeResetToDefaultVisible(TSharedRef<IPropertyHandle> PropertyHandle);
+	static void OnSectionSizeResetToDefault(TSharedRef<IPropertyHandle> PropertyHandle);
 
 	static FText GetOriginalSectionsPerComponent();
 	static TSharedRef<SWidget> GetSectionsPerComponentMenu(TSharedRef<IPropertyHandle> PropertyHandle);
 	static void OnChangeSectionsPerComponent(TSharedRef<IPropertyHandle> PropertyHandle, int32 NewSize);
 	static FText GetSectionsPerComponent(TSharedRef<IPropertyHandle> PropertyHandle);
-	static bool IsSectionsPerComponentResetToDefaultVisible();
-	static void OnSectionsPerComponentResetToDefault();
+	static bool IsSectionsPerComponentResetToDefaultVisible(TSharedRef<IPropertyHandle> PropertyHandle);
+	static void OnSectionsPerComponentResetToDefault(TSharedRef<IPropertyHandle> PropertyHandle);
 
 	static FText GetOriginalComponentCount();
 	static FText GetComponentCount(TSharedRef<IPropertyHandle> PropertyHandle_ComponentCount_X, TSharedRef<IPropertyHandle> PropertyHandle_ComponentCount_Y);
diff --git a/Engine/Source/Editor/MaterialEditor/Private/MaterialEditorInstanceDetailCustomization.cpp b/Engine/Source/Editor/MaterialEditor/Private/MaterialEditorInstanceDetailCustomization.cpp
--- a/Engine/Source/Editor/MaterialEditor/Private/MaterialEditorInstanceDetailCustomization.cpp
+++ b/Engine/Source/Editor/MaterialEditor/Private/MaterialEditorInstanceDetailCustomization.cpp
@@ -187,7 +187,7 @@ void FMaterialInstanceParameterDetails::CreateParameterValueWidget(UDEditorParam
 		.EditCondition( IsParamEnabled, FOnBooleanValueChanged::CreateSP( this, &FMaterialInstanceParameterDetails::OnOverrideParameter, Parameter ) )
 		.Visibility( TAttribute<EVisibility>::Create(TAttribute<EVisibility>::FGetter::CreateSP(this, &FMaterialInstanceParameterDetails::ShouldShowExpression, Parameter)) )
 		// Handle reset to default manually
-		.OverrideResetToDefault( true, FSimpleDelegate::CreateSP( this, &FMaterialInstanceParameterDetails::ResetToDefault, Parameter ) );
+		.OverrideResetToDefault(FResetToDefaultOverride::Create(FResetToDefaultHandler::CreateSP(this, &FMaterialInstanceParameterDetails::ResetToDefault, Parameter)));
 	}
 }
 
@@ -206,7 +206,7 @@ void FMaterialInstanceParameterDetails::CreateMaskParameterValueWidget(UDEditorP
 		IDetailPropertyRow& PropertyRow = DetailGroup.AddPropertyRow( ParameterValueProperty.ToSharedRef() );
 		PropertyRow.EditCondition( IsParamEnabled, FOnBooleanValueChanged::CreateSP( this, &FMaterialInstanceParameterDetails::OnOverrideParameter, Parameter ) );
 		// Handle reset to default manually
-		PropertyRow.OverrideResetToDefault( true, FSimpleDelegate::CreateSP( this, &FMaterialInstanceParameterDetails::ResetToDefault, Parameter ) );
+		PropertyRow.OverrideResetToDefault(FResetToDefaultOverride::Create(FResetToDefaultHandler::CreateSP(this, &FMaterialInstanceParameterDetails::ResetToDefault, Parameter)));
 		PropertyRow.Visibility( TAttribute<EVisibility>::Create(TAttribute<EVisibility>::FGetter::CreateSP(this, &FMaterialInstanceParameterDetails::ShouldShowExpression, Parameter)) );
 
 		const FText ParameterName = FText::FromName(Parameter->ParameterName); 
@@ -362,7 +362,7 @@ FText FMaterialInstanceParameterDetails::GetParameterExpressionDescription(UDEdi
 	return FText::GetEmpty();
 }
 
-void FMaterialInstanceParameterDetails::ResetToDefault( class UDEditorParameterValue* Parameter )
+void FMaterialInstanceParameterDetails::ResetToDefault(TSharedRef<IPropertyHandle> PropertyHandle, class UDEditorParameterValue* Parameter)
 {
 	const FScopedTransaction Transaction( LOCTEXT( "ResetToDefault", "Reset To Default" ) );
 	Parameter->Modify();
diff --git a/Engine/Source/Editor/MaterialEditor/Private/MaterialEditorInstanceDetailCustomization.h b/Engine/Source/Editor/MaterialEditor/Private/MaterialEditorInstanceDetailCustomization.h
--- a/Engine/Source/Editor/MaterialEditor/Private/MaterialEditorInstanceDetailCustomization.h
+++ b/Engine/Source/Editor/MaterialEditor/Private/MaterialEditorInstanceDetailCustomization.h
@@ -56,7 +56,7 @@ private:
 	bool OnShouldSetAsset(const FAssetData& InAssetData) const;
 
 	/** Reset to default implementation.  Resets Parameter to default */
-	void ResetToDefault( class UDEditorParameterValue* Parameter );
+	void ResetToDefault(TSharedRef<IPropertyHandle> PropertyHandle, class UDEditorParameterValue* Parameter);
 
 	/** Returns true if the refraction options should be displayed */
 	EVisibility ShouldShowMaterialRefractionSettings() const;
diff --git a/Engine/Source/Editor/Persona/Private/AnimNotifyDetails.cpp b/Engine/Source/Editor/Persona/Private/AnimNotifyDetails.cpp
--- a/Engine/Source/Editor/Persona/Private/AnimNotifyDetails.cpp
+++ b/Engine/Source/Editor/Persona/Private/AnimNotifyDetails.cpp
@@ -236,7 +236,7 @@ void FAnimNotifyDetails::ClearInstancedSelectionDropDown(IDetailCategoryBuilder&
 	IDetailPropertyRow& PropRow = CategoryBuilder.AddProperty(PropHandle);
 	
 	PropRow
-	.OverrideResetToDefault(false, FSimpleDelegate())
+	.OverrideResetToDefault(FResetToDefaultOverride::Hide())
 	.CustomWidget(bShowChildren)
 	.NameContent()
 	[
diff --git a/Engine/Source/Editor/PropertyEditor/Private/CustomChildBuilder.cpp b/Engine/Source/Editor/PropertyEditor/Private/CustomChildBuilder.cpp
--- a/Engine/Source/Editor/PropertyEditor/Private/CustomChildBuilder.cpp
+++ b/Engine/Source/Editor/PropertyEditor/Private/CustomChildBuilder.cpp
@@ -45,6 +45,11 @@ IDetailPropertyRow& FCustomChildrenBuilder::AddChildProperty( TSharedRef<IProper
 	FDetailLayoutCustomization NewCustomization;
 	NewCustomization.PropertyRow = MakeShareable( new FDetailPropertyRow( StaticCastSharedRef<FPropertyHandleBase>( PropertyHandle )->GetPropertyNode(), ParentCategory.Pin().ToSharedRef() ) );
 
+	if (CustomResetChildToDefault.IsSet())
+	{
+		NewCustomization.PropertyRow->OverrideResetToDefault(CustomResetChildToDefault.GetValue());
+	}
+
 	ChildCustomizations.Add( NewCustomization );
 
 	return *NewCustomization.PropertyRow;
@@ -110,3 +115,8 @@ TSharedRef<SWidget> FCustomChildrenBuilder::GenerateStructValueWidget( TSharedRe
 	}
 }
 
+FCustomChildrenBuilder& FCustomChildrenBuilder::OverrideResetChildrenToDefault(const FResetToDefaultOverride& ResetToDefault)
+{
+	CustomResetChildToDefault = ResetToDefault;
+	return *this;
+}
diff --git a/Engine/Source/Editor/PropertyEditor/Private/CustomChildBuilder.h b/Engine/Source/Editor/PropertyEditor/Private/CustomChildBuilder.h
--- a/Engine/Source/Editor/PropertyEditor/Private/CustomChildBuilder.h
+++ b/Engine/Source/Editor/PropertyEditor/Private/CustomChildBuilder.h
@@ -5,7 +5,7 @@
 class FCustomChildrenBuilder : public IDetailChildrenBuilder
 {
 public:
-	FCustomChildrenBuilder( TSharedRef<FDetailCategoryImpl> InParentCategory )
+	FCustomChildrenBuilder(TSharedRef<FDetailCategoryImpl> InParentCategory)
 		: ParentCategory( InParentCategory )
 	{}
 	
@@ -16,7 +16,14 @@ public:
 	virtual TSharedRef<SWidget> GenerateStructValueWidget( TSharedRef<IPropertyHandle> StructPropertyHandle ) override;
 
 	const TArray< FDetailLayoutCustomization >& GetChildCustomizations() const { return ChildCustomizations; }
+
+	/** Set the user customized reset to default for the children of this builder */
+	FCustomChildrenBuilder& OverrideResetChildrenToDefault(const FResetToDefaultOverride& ResetToDefault);
+
 private:
 	TArray< FDetailLayoutCustomization > ChildCustomizations;
 	TWeakPtr<FDetailCategoryImpl> ParentCategory;
+
+	/** User customized reset to default on children */
+	TOptional<FResetToDefaultOverride> CustomResetChildToDefault;
 };
diff --git a/Engine/Source/Editor/PropertyEditor/Private/DetailPropertyRow.cpp b/Engine/Source/Editor/PropertyEditor/Private/DetailPropertyRow.cpp
--- a/Engine/Source/Editor/PropertyEditor/Private/DetailPropertyRow.cpp
+++ b/Engine/Source/Editor/PropertyEditor/Private/DetailPropertyRow.cpp
@@ -104,13 +104,9 @@ IDetailPropertyRow& FDetailPropertyRow::Visibility( TAttribute<EVisibility> Visi
 	return *this;
 }
 
-IDetailPropertyRow& FDetailPropertyRow::OverrideResetToDefault( TAttribute<bool> IsResetToDefaultVisible, FSimpleDelegate OnResetToDefaultClicked )
+IDetailPropertyRow& FDetailPropertyRow::OverrideResetToDefault(const FResetToDefaultOverride& ResetToDefault)
 {
-	CustomResetToDefault = MakeShareable( new FCustomResetToDefault );
-
-	CustomResetToDefault->IsResetToDefaultVisible = IsResetToDefaultVisible;
-	CustomResetToDefault->OnResetToDefaultClicked = OnResetToDefaultClicked;
-
+	CustomResetToDefault = ResetToDefault;
 	return *this;
 }
 
@@ -217,6 +213,13 @@ void FDetailPropertyRow::OnItemNodeInitialized( TSharedRef<FDetailCategoryImpl>
 	if( bShowCustomPropertyChildren && CustomTypeInterface.IsValid() )
 	{
 		PropertyTypeLayoutBuilder = MakeShareable(new FCustomChildrenBuilder(InParentCategory));
+
+		/** Does this row pass its custom reset behavior to its children? */
+		if (CustomResetToDefault.IsSet() && CustomResetToDefault->PropagatesToChildren())
+		{
+			PropertyTypeLayoutBuilder->OverrideResetChildrenToDefault(CustomResetToDefault.GetValue());
+		}
+
 		CustomTypeInterface->CustomizeChildren(PropertyHandle.ToSharedRef(), *PropertyTypeLayoutBuilder, *this);
 	}
 }
@@ -463,7 +466,7 @@ void FDetailPropertyRow::MakeValueWidget( FDetailWidgetRow& Row, const TSharedPt
 		[
 			SNew( SResetToDefaultPropertyEditor, PropertyEditor.ToSharedRef() )
 			.IsEnabled( IsEnabledAttrib )
-			.CustomResetToDefault( CustomResetToDefault.IsValid() ? *CustomResetToDefault : FCustomResetToDefault() )
+			.CustomResetToDefault(CustomResetToDefault.Get(FResetToDefaultOverride()))
 		];
 	}
 
diff --git a/Engine/Source/Editor/PropertyEditor/Private/DetailPropertyRow.h b/Engine/Source/Editor/PropertyEditor/Private/DetailPropertyRow.h
--- a/Engine/Source/Editor/PropertyEditor/Private/DetailPropertyRow.h
+++ b/Engine/Source/Editor/PropertyEditor/Private/DetailPropertyRow.h
@@ -15,7 +15,7 @@ public:
 	virtual IDetailPropertyRow& IsEnabled(TAttribute<bool> InIsEnabled) override;
 	virtual IDetailPropertyRow& ShouldAutoExpand(bool bForceExpansion) override;
 	virtual IDetailPropertyRow& Visibility( TAttribute<EVisibility> Visibility ) override;
-	virtual IDetailPropertyRow& OverrideResetToDefault( TAttribute<bool> IsResetToDefaultVisible, FSimpleDelegate OnResetToDefaultClicked ) override;
+	virtual IDetailPropertyRow& OverrideResetToDefault(const FResetToDefaultOverride& ResetToDefault) override;
 	virtual FDetailWidgetRow& CustomWidget( bool bShowChildren = false ) override;
 	virtual void GetDefaultWidgets( TSharedPtr<SWidget>& OutNameWidget, TSharedPtr<SWidget>& OutValueWidget ) override;
 	virtual void GetDefaultWidgets( TSharedPtr<SWidget>& OutNameWidget, TSharedPtr<SWidget>& OutValueWidget, FDetailWidgetRow& Row ) override;
@@ -127,7 +127,7 @@ private:
 	/** User customized edit condition */
 	TSharedPtr<struct FCustomEditCondition> CustomEditCondition;
 	/** User customized reset to default */
-	TSharedPtr<struct FCustomResetToDefault> CustomResetToDefault;
+	TOptional<FResetToDefaultOverride> CustomResetToDefault;
 	/** The category this row resides in */
 	TWeakPtr<FDetailCategoryImpl> ParentCategory;
 	/** Root of the property node if this node comes from an external tree */
diff --git a/Engine/Source/Editor/PropertyEditor/Private/Presentation/PropertyEditor/PropertyEditor.cpp b/Engine/Source/Editor/PropertyEditor/Private/Presentation/PropertyEditor/PropertyEditor.cpp
--- a/Engine/Source/Editor/PropertyEditor/Private/Presentation/PropertyEditor/PropertyEditor.cpp
+++ b/Engine/Source/Editor/PropertyEditor/Private/Presentation/PropertyEditor/PropertyEditor.cpp
@@ -349,10 +349,10 @@ void FPropertyEditor::ResetToDefault()
 	PropertyUtilities->EnqueueDeferredAction( FSimpleDelegate::CreateSP( this, &FPropertyEditor::OnResetToDefault ) );
 }
 
-void FPropertyEditor::CustomResetToDefault( FSimpleDelegate OnCustomResetToDefaultDelegate )
+void FPropertyEditor::CustomResetToDefault(const FResetToDefaultOverride& OnCustomResetToDefault)
 {
 	// This action must be deferred until next tick so that we avoid accessing invalid data before we have a chance to tick
-	PropertyUtilities->EnqueueDeferredAction( FSimpleDelegate::CreateSP( this, &FPropertyEditor::OnCustomResetToDefault, OnCustomResetToDefaultDelegate ) );
+	PropertyUtilities->EnqueueDeferredAction(FSimpleDelegate::CreateLambda([this, OnCustomResetToDefault](){ this->OnCustomResetToDefault(OnCustomResetToDefault); }));
 }
 
 void FPropertyEditor::OnGetClassesForAssetPicker( TArray<const UClass*>& OutClasses )
@@ -405,13 +405,13 @@ void FPropertyEditor::OnResetToDefault()
 	PropertyNode->ResetToDefault( PropertyUtilities->GetNotifyHook() );
 }
 
-void FPropertyEditor::OnCustomResetToDefault( FSimpleDelegate OnCustomResetToDefaultDelegate )
+void FPropertyEditor::OnCustomResetToDefault(const FResetToDefaultOverride& OnCustomResetToDefault)
 {
-	if(OnCustomResetToDefaultDelegate.IsBound())
+	if (OnCustomResetToDefault.OnResetToDefaultClicked().IsBound())
 	{
 		PropertyNode->NotifyPreChange( PropertyNode->GetProperty(), PropertyUtilities->GetNotifyHook() );
 
-		OnCustomResetToDefaultDelegate.Execute();
+		OnCustomResetToDefault.OnResetToDefaultClicked().Execute(GetPropertyHandle());
 
 		// Call PostEditchange on all the objects
 		FPropertyChangedEvent ChangeEvent( PropertyNode->GetProperty() );
diff --git a/Engine/Source/Editor/PropertyEditor/Private/Presentation/PropertyEditor/PropertyEditor.h b/Engine/Source/Editor/PropertyEditor/Private/Presentation/PropertyEditor/PropertyEditor.h
--- a/Engine/Source/Editor/PropertyEditor/Private/Presentation/PropertyEditor/PropertyEditor.h
+++ b/Engine/Source/Editor/PropertyEditor/Private/Presentation/PropertyEditor/PropertyEditor.h
@@ -5,6 +5,7 @@
 
 class FPropertyNode;
 class IPropertyUtilities;
+class FResetToDefaultOverride;
 
 class FPropertyEditor : public TSharedFromThis< FPropertyEditor >	
 {
@@ -87,7 +88,7 @@ public:
 	void EmptyArray();
 	void Find();
 	void ResetToDefault();
-	void CustomResetToDefault( FSimpleDelegate OnCustomResetToDefaultDelegate );
+	void CustomResetToDefault(const FResetToDefaultOverride& OnCustomResetToDefault);
 	void OnGetClassesForAssetPicker( TArray<const UClass*>& OutClasses );
 	void OnAssetSelected( const FAssetData& AssetData );
 	void OnActorSelected( AActor* InActor );
@@ -124,7 +125,7 @@ private:
 	void OnBrowseTo();
 	void OnEmptyArray();
 	void OnResetToDefault();
-	void OnCustomResetToDefault( FSimpleDelegate OnCustomResetToDefaultDelegate );
+	void OnCustomResetToDefault(const FResetToDefaultOverride& OnCustomResetToDefault);
 
 	/**
 	 * Returns true if the value of the conditional property matches the value required.  Indicates whether editing or otherwise interacting with this item's
diff --git a/Engine/Source/Editor/PropertyEditor/Private/UserInterface/PropertyEditor/SResetToDefaultPropertyEditor.cpp b/Engine/Source/Editor/PropertyEditor/Private/UserInterface/PropertyEditor/SResetToDefaultPropertyEditor.cpp
--- a/Engine/Source/Editor/PropertyEditor/Private/UserInterface/PropertyEditor/SResetToDefaultPropertyEditor.cpp
+++ b/Engine/Source/Editor/PropertyEditor/Private/UserInterface/PropertyEditor/SResetToDefaultPropertyEditor.cpp
@@ -2,6 +2,7 @@
 #include "PropertyEditorPrivatePCH.h"
 #include "SResetToDefaultPropertyEditor.h"
 #include "PropertyEditor.h"
+#include "IDetailPropertyRow.h"
 
 #define LOCTEXT_NAMESPACE "ResetToDefaultPropertyEditor"
 
@@ -10,9 +11,9 @@ void SResetToDefaultPropertyEditor::Construct( const FArguments& InArgs, const T
 	PropertyEditor = InPropertyEditor;
 	NonVisibleState = InArgs._NonVisibleState;
 	bValueDiffersFromDefault = false;
-	CustomResetToDefault = InArgs._CustomResetToDefault;
+	OptionalCustomResetToDefault = InArgs._CustomResetToDefault;
 
-	if( CustomResetToDefault.IsResetToDefaultVisible.IsSet() )
+	if (OptionalCustomResetToDefault.IsSet())
 	{
 		ChildSlot
 		[
@@ -52,9 +53,9 @@ void SResetToDefaultPropertyEditor::Construct( const FArguments& InArgs, const T
 
 void SResetToDefaultPropertyEditor::Tick( const FGeometry& AllottedGeometry, const double InCurrentTime, const float InDeltaTime )
 {
-	if( CustomResetToDefault.IsResetToDefaultVisible.IsSet() )
+	if (OptionalCustomResetToDefault.IsSet())
 	{
-		bValueDiffersFromDefault = CustomResetToDefault.IsResetToDefaultVisible.Get();
+		bValueDiffersFromDefault = OptionalCustomResetToDefault.GetValue().IsResetToDefaultVisible(PropertyEditor->GetPropertyHandle());
 	}
 	else
 	{
@@ -85,7 +86,7 @@ FReply SResetToDefaultPropertyEditor::OnDefaultResetClicked()
 
 FReply SResetToDefaultPropertyEditor::OnCustomResetClicked()
 {
-	PropertyEditor->CustomResetToDefault( CustomResetToDefault.OnResetToDefaultClicked );
+	PropertyEditor->CustomResetToDefault(OptionalCustomResetToDefault.GetValue());
 
 	return FReply::Handled();
 }
diff --git a/Engine/Source/Editor/PropertyEditor/Private/UserInterface/PropertyEditor/SResetToDefaultPropertyEditor.h b/Engine/Source/Editor/PropertyEditor/Private/UserInterface/PropertyEditor/SResetToDefaultPropertyEditor.h
--- a/Engine/Source/Editor/PropertyEditor/Private/UserInterface/PropertyEditor/SResetToDefaultPropertyEditor.h
+++ b/Engine/Source/Editor/PropertyEditor/Private/UserInterface/PropertyEditor/SResetToDefaultPropertyEditor.h
@@ -1,12 +1,10 @@
 // Copyright 1998-2015 Epic Games, Inc. All Rights Reserved.
 #pragma once
 
-struct FCustomResetToDefault
-{
-	TAttribute<bool> IsResetToDefaultVisible;
-	FSimpleDelegate OnResetToDefaultClicked;
-};
+// Forward decl
+class FResetToDefaultOverride;
 
+/** Widget showing the reset to default value button */
 class SResetToDefaultPropertyEditor : public SCompoundWidget
 {
 
@@ -16,7 +14,7 @@ public:
 		: _NonVisibleState( EVisibility::Hidden )
 		{}
 		SLATE_ARGUMENT( EVisibility, NonVisibleState )
-		SLATE_ARGUMENT( FCustomResetToDefault, CustomResetToDefault )
+		SLATE_ARGUMENT(TOptional<FResetToDefaultOverride>, CustomResetToDefault)
 	SLATE_END_ARGS()
 
 	void Construct( const FArguments& InArgs, const TSharedRef< class FPropertyEditor>& InPropertyEditor );
@@ -33,7 +31,7 @@ private:
 	FReply OnCustomResetClicked();
 
 private:
-	FCustomResetToDefault CustomResetToDefault;
+	TOptional<FResetToDefaultOverride> OptionalCustomResetToDefault;
 
 	TSharedPtr< class FPropertyEditor > PropertyEditor;
 
diff --git a/Engine/Source/Editor/PropertyEditor/Public/IDetailPropertyRow.h b/Engine/Source/Editor/PropertyEditor/Public/IDetailPropertyRow.h
--- a/Engine/Source/Editor/PropertyEditor/Public/IDetailPropertyRow.h
+++ b/Engine/Source/Editor/PropertyEditor/Public/IDetailPropertyRow.h
@@ -2,9 +2,96 @@
 
 #pragma once
 
+#include "PropertyHandle.h"
+
 // Forward declarations
 class FDetailWidgetRow;
 
+DECLARE_DELEGATE_RetVal_OneParam(bool, FIsResetToDefaultVisible, TSharedRef<IPropertyHandle> /* PropertyHandle */);
+DECLARE_DELEGATE_OneParam(FResetToDefaultHandler, TSharedRef<IPropertyHandle> /* PropertyHandle*/);
+
+/**
+ * Structure describing the delegates needed to override the behavior of reset to default in detail properties
+ */
+class FResetToDefaultOverride
+{
+public:
+	/** Creates a FResetToDefaultOverride in which the the reset to default is always visible */
+	static FResetToDefaultOverride Create(FResetToDefaultHandler InResetToDefaultClicked, const bool InPropagateToChildren = false)
+	{
+		FResetToDefaultOverride ResetToDefault;
+		ResetToDefault.bForceShow = true;
+		ResetToDefault.OnResetToDefaultClickedDelegate = InResetToDefaultClicked;
+		ResetToDefault.bPropagateToChildren = InPropagateToChildren;
+		ResetToDefault.bForceHide = false;
+		return ResetToDefault;
+	}
+
+	/** Creates a FResetToDefaultOverride from visibility and click handler callback delegates */
+	static FResetToDefaultOverride Create(FIsResetToDefaultVisible InIsResetToDefaultVisible, FResetToDefaultHandler InResetToDefaultClicked, const bool InPropagateToChildren = false)
+	{
+		FResetToDefaultOverride ResetToDefault;
+		ResetToDefault.bForceShow = false;
+		ResetToDefault.IsResetToDefaultVisibleDelegate = InIsResetToDefaultVisible;
+		ResetToDefault.OnResetToDefaultClickedDelegate = InResetToDefaultClicked;
+		ResetToDefault.bPropagateToChildren = InPropagateToChildren;
+		ResetToDefault.bForceHide = false;
+		return ResetToDefault;
+	}
+
+	/** Creates a FResetToDefaultOverride in which reset to default is never visible */
+	static FResetToDefaultOverride Hide(const bool InPropagateToChildren = false)
+	{
+		FResetToDefaultOverride HideResetToDefault;
+		HideResetToDefault.bForceShow = false;
+		HideResetToDefault.bForceHide = true;
+		HideResetToDefault.bPropagateToChildren = InPropagateToChildren;
+		return HideResetToDefault;
+	}
+
+	/** Called by the UI to show/hide the reset widgets */
+	bool IsResetToDefaultVisible(TSharedRef<IPropertyHandle> Property) const
+	{
+		if (bForceShow)
+		{
+			return true;
+		}
+		if (!bForceHide && IsResetToDefaultVisibleDelegate.IsBound())
+		{
+			return IsResetToDefaultVisibleDelegate.Execute(Property);
+		}
+		return false;
+	}
+
+	/** Called by the property editor to actually reset the property to default */
+	FResetToDefaultHandler OnResetToDefaultClicked() const
+	{
+		return OnResetToDefaultClickedDelegate;
+	}
+
+	/** Called by properties to determine whether this override should set on their children */
+	bool PropagatesToChildren() const
+	{
+		return bPropagateToChildren;
+	}
+
+private:
+	/** Callback to indicate whether or not reset to default is visible */
+	FIsResetToDefaultVisible IsResetToDefaultVisibleDelegate;
+
+	/** Delegate called when reset to default is clicked */
+	FResetToDefaultHandler OnResetToDefaultClickedDelegate;
+
+	/** Should properties pass this on to their children? */
+	bool bPropagateToChildren;
+
+	/** Ignore the visibility delegate and always show the reset to default widgets? */
+	bool bForceShow;
+
+	/** Ignore the visibility delegate and never show the reset to default widgets? */
+	bool bForceHide;
+};
+
 /**
  * A single row for a property in a details panel                                                              
  */
@@ -69,10 +156,9 @@ public:
 	/**
 	 * Overrides the behavior of reset to default
 	 *
-	 * @param IsResetToDefaultVisible	Attribute to indicate whether or not reset to default is visible
-	 * @param OnResetToDefaultClicked	Delegate called when reset to default is clicked
+	 * @param ResetToDefault	Contains the delegates needed to override the behavior of reset to default
 	 */
-	virtual IDetailPropertyRow& OverrideResetToDefault( TAttribute<bool> IsResetToDefaultVisible, FSimpleDelegate OnResetToDefaultClicked ) = 0;
+	virtual IDetailPropertyRow& OverrideResetToDefault(const FResetToDefaultOverride& ResetToDefault) = 0;
 
 	/**
 	 * Returns the name and value widget of this property row.  You can use this widget to apply further customization to existing widgets (by using this  with CustomWidget)
diff --git a/Engine/Source/Editor/UMGEditor/Private/Customizations/UMGDetailCustomizations.cpp b/Engine/Source/Editor/UMGEditor/Private/Customizations/UMGDetailCustomizations.cpp
--- a/Engine/Source/Editor/UMGEditor/Private/Customizations/UMGDetailCustomizations.cpp
+++ b/Engine/Source/Editor/UMGEditor/Private/Customizations/UMGDetailCustomizations.cpp
@@ -94,7 +94,7 @@ void FBlueprintWidgetCustomization::CreateEventCustomization( IDetailLayoutBuild
 	IDetailCategoryBuilder& PropertyCategory = DetailLayout.EditCategory(FObjectEditorUtils::GetCategoryFName(Property), FText::GetEmpty(), ECategoryPriority::Uncommon);
 
 	IDetailPropertyRow& PropertyRow = PropertyCategory.AddProperty(DelegatePropertyHandle);
-	PropertyRow.OverrideResetToDefault(true, FSimpleDelegate::CreateSP(this, &FBlueprintWidgetCustomization::ResetToDefault_RemoveBinding, DelegatePropertyHandle));
+	PropertyRow.OverrideResetToDefault(FResetToDefaultOverride::Create(FResetToDefaultHandler::CreateSP(this, &FBlueprintWidgetCustomization::ResetToDefault_RemoveBinding)));
 
 	FString LabelStr = Property->GetName();
 	LabelStr.RemoveFromEnd(TEXT("Event"));
diff --git a/Engine/Source/Runtime/SlateCore/Public/Styling/SlateBrush.h b/Engine/Source/Runtime/SlateCore/Public/Styling/SlateBrush.h
--- a/Engine/Source/Runtime/SlateCore/Public/Styling/SlateBrush.h
+++ b/Engine/Source/Runtime/SlateCore/Public/Styling/SlateBrush.h
@@ -96,6 +96,10 @@ namespace ESlateBrushImageType
 	};
 }
 
+namespace SlateBrushDefs
+{
+	static const float DefaultImageSize = 32.0f;
+}
 
 /**
  * An brush which contains information about how to draw a Slate element
@@ -144,7 +148,7 @@ public:
 	 * Default constructor.
 	 */
 	FSlateBrush()
-		: ImageSize(32.0f, 32.0f)
+		: ImageSize(SlateBrushDefs::DefaultImageSize, SlateBrushDefs::DefaultImageSize)
 		, DrawAs(ESlateBrushDrawType::Image)
 		, Margin(0.0f)
 		, Tint_DEPRECATED(FLinearColor::White)
