Editor: Exposed 'Convert to Static Mesh' on volumes as well as brushes
[INTEGRATE] Change 2648443 by [REDACTED_PERSON][REDACTED_HANDLE] on 2015/08/07 16:07:42

diff --git a/Engine/Source/Editor/DetailCustomizations/Private/BrushDetails.cpp b/Engine/Source/Editor/DetailCustomizations/Private/BrushDetails.cpp
--- a/Engine/Source/Editor/DetailCustomizations/Private/BrushDetails.cpp
+++ b/Engine/Source/Editor/DetailCustomizations/Private/BrushDetails.cpp
@@ -36,26 +36,24 @@ void FBrushDetails::CustomizeDetails( IDetailLayoutBuilder& DetailLayout )
 	// See if we have a volume. If we do - we hide the BSP stuff (solidity, order)
 	bool bHaveAVolume = false;
 	TArray< TWeakObjectPtr<UObject> > SelectedObjects = DetailLayout.GetDetailsView().GetSelectedObjects();
-	for(int32 ObjIdx=0; ObjIdx<SelectedObjects.Num(); ObjIdx++)
+	for (int32 ObjIdx = 0; ObjIdx < SelectedObjects.Num(); ObjIdx++)
 	{
-		AVolume* Volume = Cast<AVolume>( SelectedObjects[ObjIdx].Get() );
-
-		if (Volume != NULL)
-		{
-			bHaveAVolume = true;
-		}
-		else
+		if (ABrush* Brush = Cast<ABrush>(SelectedObjects[ObjIdx].Get()))
 		{
-			// Since the brush is not a volume, it is valid for the SelectedBSPBrushes list.
-			ABrush* Brush = Cast<ABrush>( SelectedObjects[ObjIdx].Get() );
-			if(!FActorEditorUtils::IsABuilderBrush(Brush))
+			if (AVolume* Volume = Cast<AVolume>(Brush))
+			{
+				bHaveAVolume = true;
+			}
+
+			if (!FActorEditorUtils::IsABuilderBrush(Brush))
 			{
 				// Store the selected actors for use later. Its fine to do this when CustomizeDetails is called because if the selected actors changes, CustomizeDetails will be called again on a new instance
 				// and our current resource would be destroyed.
-				SelectedBSPBrushes.Add( Brush );
+				SelectedBrushes.Add(Brush);
 			}
 		}
 	}
+
 	FMenuBuilder PolygonsMenuBuilder( true, CommandBindings );
 	{
 		PolygonsMenuBuilder.BeginSection("BrushDetailsPolygons");
@@ -297,12 +295,12 @@ void FBrushDetails::CustomizeDetails( IDetailLayoutBuilder& DetailLayout )
 		]
 	];
 
-	if(SelectedBSPBrushes.Num())
+	if (SelectedBrushes.Num() > 0)
 	{
 		BrushHorizontalBox->AddSlot()
 		[
 			SNew( SButton )
-			.ToolTipText( LOCTEXT("CreateStaticMeshActor_Tooltip", "Creates a static mesh from selected brushes and replaces the affected brushes in the scene with the new static mesh") )
+			.ToolTipText( LOCTEXT("CreateStaticMeshActor_Tooltip", "Creates a static mesh from selected brushes or volumes and replaces them in the scene with the new static mesh") )
 			.OnClicked( this, &FBrushDetails::OnCreateStaticMesh )
 			.HAlign( HAlign_Center )
 			[
@@ -317,22 +315,10 @@ END_SLATE_FUNCTION_BUILD_OPTIMIZATION
 
 FReply FBrushDetails::OnCreateStaticMesh()
 {
-	// Raw list of brushes to be sent to be converted.
-	TArray<AActor*> RawSelectedBrushes;
-	
-	TWeakObjectPtr<ABrush> Brush;
-	for(int32 BrushIdx = 0; BrushIdx < SelectedBSPBrushes.Num(); ++BrushIdx )
-	{
-		Brush = SelectedBSPBrushes[BrushIdx];
-
-		// Make sure the brush is still valid.
-		if(Brush.IsValid())
-		{
-			RawSelectedBrushes.Add(Brush.Get());
-		}
-	}
+	TArray<AActor*> ValidSelectedBrushes;
+	CopyFromWeakArray(ValidSelectedBrushes, SelectedBrushes);
 
-	GEditor->ConvertActors(RawSelectedBrushes, AStaticMeshActor::StaticClass(), TSet<FString>(), true);
+	GEditor->ConvertActors(ValidSelectedBrushes, AStaticMeshActor::StaticClass(), TSet<FString>(), true);
 
 	return FReply::Handled();
 }
diff --git a/Engine/Source/Editor/DetailCustomizations/Private/BrushDetails.h b/Engine/Source/Editor/DetailCustomizations/Private/BrushDetails.h
--- a/Engine/Source/Editor/DetailCustomizations/Private/BrushDetails.h
+++ b/Engine/Source/Editor/DetailCustomizations/Private/BrushDetails.h
@@ -18,8 +18,8 @@ private:
 	FReply OnCreateStaticMesh();
 
 private:
-	/** Holds a list of BSP brushes, used for converting to static meshes */
-	TArray< TWeakObjectPtr<ABrush> > SelectedBSPBrushes;
+	/** Holds a list of BSP brushes or volumes, used for converting to static meshes */
+	TArray< TWeakObjectPtr<ABrush> > SelectedBrushes;
 
 	/** Container widget for the geometry mode tools */
 	TSharedPtr< SHorizontalBox > GeometryToolsContainer;
diff --git a/Engine/Source/Editor/UnrealEd/Classes/Editor/EditorEngine.h b/Engine/Source/Editor/UnrealEd/Classes/Editor/EditorEngine.h
--- a/Engine/Source/Editor/UnrealEd/Classes/Editor/EditorEngine.h
+++ b/Engine/Source/Editor/UnrealEd/Classes/Editor/EditorEngine.h
@@ -882,8 +882,9 @@ public:
 	 *
 	 * @param Model					The model to be rebuilt.
 	 * @param bSelectedBrushesOnly	Use all brushes in the current level or just the selected ones?.
+	 * @param bTreatMovableBrushesAsStatic	Treat moveable brushes as static?.
 	 */
-	void RebuildModelFromBrushes(UModel* Model, bool bSelectedBrushesOnly);
+	void RebuildModelFromBrushes(UModel* Model, bool bSelectedBrushesOnly, bool bTreatMovableBrushesAsStatic = false);
 
 	/**
 	 * Rebuilds levels containing currently selected brushes and should be invoked after a brush has been modified
diff --git a/Engine/Source/Editor/UnrealEd/Private/EditorEngine.cpp b/Engine/Source/Editor/UnrealEd/Private/EditorEngine.cpp
--- a/Engine/Source/Editor/UnrealEd/Private/EditorEngine.cpp
+++ b/Engine/Source/Editor/UnrealEd/Private/EditorEngine.cpp
@@ -4981,7 +4981,7 @@ AActor* UEditorEngine::ConvertBrushesToStaticMesh(const FString& InStaticMeshPac
 		InBrushesToConvert[BrushesIdx]->TeleportTo(Location - InPivotLocation, Rotation, false, true);
 	}
 
-	GEditor->RebuildModelFromBrushes(ConversionTempModel, true );
+	GEditor->RebuildModelFromBrushes(ConversionTempModel, true, true );
 	GEditor->bspBuildFPolys(ConversionTempModel, true, 0);
 
 	if (0 < ConversionTempModel->Polys->Element.Num())
diff --git a/Engine/Source/Editor/UnrealEd/Private/EditorServer.cpp b/Engine/Source/Editor/UnrealEd/Private/EditorServer.cpp
--- a/Engine/Source/Editor/UnrealEd/Private/EditorServer.cpp
+++ b/Engine/Source/Editor/UnrealEd/Private/EditorServer.cpp
@@ -1573,7 +1573,7 @@ void UEditorEngine::RebuildLevel(ULevel& Level)
 	GLevelEditorModeTools().MapChangeNotify();
 }
 
-void UEditorEngine::RebuildModelFromBrushes(UModel* Model, bool bSelectedBrushesOnly)
+void UEditorEngine::RebuildModelFromBrushes(UModel* Model, bool bSelectedBrushesOnly, bool bTreatMovableBrushesAsStatic)
 {
 	TUniquePtr<FBspPointsGrid> BspPoints = MakeUnique<FBspPointsGrid>(50.0f, THRESH_POINTS_ARE_SAME);
 	TUniquePtr<FBspPointsGrid> BspVectors = MakeUnique<FBspPointsGrid>(1/16.0f, FMath::Max(THRESH_NORMALS_ARE_SAME, THRESH_VECTORS_ARE_NEAR));
@@ -1612,7 +1612,7 @@ void UEditorEngine::RebuildModelFromBrushes(UModel* Model, bool bSelectedBrushes
 	for( auto It(Level->Actors.CreateConstIterator()); It; ++It )
 	{	
 		ABrush* Brush = Cast<ABrush>(*It);
-		if ((Brush && Brush->IsStaticBrush() && !FActorEditorUtils::IsABuilderBrush(Brush)) &&
+		if ((Brush && (Brush->IsStaticBrush() || bTreatMovableBrushesAsStatic) && !FActorEditorUtils::IsABuilderBrush(Brush)) &&
 			(!bSelectedBrushesOnly || Brush->IsSelected()) &&
 			(!(Brush->PolyFlags & PF_Semisolid) || (Brush->BrushType != Brush_Add) || (Brush->PolyFlags & PF_Portal)))
 		{
@@ -1639,19 +1639,22 @@ void UEditorEngine::RebuildModelFromBrushes(UModel* Model, bool bSelectedBrushes
 		}
 	}
 
-	// Rebuild dynamic brush BSP's.
-	for( auto It(Level->Actors.CreateConstIterator()); It; ++It )
+	// Rebuild dynamic brush BSP's (if they weren't handled earlier)
+	if (!bTreatMovableBrushesAsStatic)
 	{
-		ABrush* DynamicBrush = Cast<ABrush>(*It);
-		if (DynamicBrush && DynamicBrush->Brush && !DynamicBrush->IsStaticBrush() &&
-			(!bSelectedBrushesOnly || DynamicBrush->IsSelected()))
+		for (auto It(Level->Actors.CreateConstIterator()); It; ++It)
 		{
-			BspPoints = MakeUnique<FBspPointsGrid>(50.0f, THRESH_POINTS_ARE_SAME);
-			BspVectors = MakeUnique<FBspPointsGrid>(1 / 16.0f, FMath::Max(THRESH_NORMALS_ARE_SAME, THRESH_VECTORS_ARE_NEAR));
-			FBspPointsGrid::GBspPoints = BspPoints.Get();
-			FBspPointsGrid::GBspVectors = BspVectors.Get();
+			ABrush* DynamicBrush = Cast<ABrush>(*It);
+			if (DynamicBrush && DynamicBrush->Brush && !DynamicBrush->IsStaticBrush() &&
+				(!bSelectedBrushesOnly || DynamicBrush->IsSelected()))
+			{
+				BspPoints = MakeUnique<FBspPointsGrid>(50.0f, THRESH_POINTS_ARE_SAME);
+				BspVectors = MakeUnique<FBspPointsGrid>(1 / 16.0f, FMath::Max(THRESH_NORMALS_ARE_SAME, THRESH_VECTORS_ARE_NEAR));
+				FBspPointsGrid::GBspPoints = BspPoints.Get();
+				FBspPointsGrid::GBspVectors = BspVectors.Get();
 
-			FBSPOps::csgPrepMovingBrush(DynamicBrush);
+				FBSPOps::csgPrepMovingBrush(DynamicBrush);
+			}
 		}
 	}
 
