#ANIM: Allow each animation to compress and change setting in the detail panel

Merging

//depot/UE4-Orion/...

to //depot/UE4-Orion/...

diff --git a/Engine/Source/Editor/DetailCustomizations/Private/AnimSequenceDetails.cpp b/Engine/Source/Editor/DetailCustomizations/Private/AnimSequenceDetails.cpp
--- a/Engine/Source/Editor/DetailCustomizations/Private/AnimSequenceDetails.cpp
+++ b/Engine/Source/Editor/DetailCustomizations/Private/AnimSequenceDetails.cpp
@@ -6,6 +6,7 @@
 #include "AnimPreviewInstance.h"
 #include "Runtime/Engine/Public/Slate/SceneViewport.h"
 #include "Editor/KismetWidgets/Public/SScrubControlPanel.h"
+#include "AnimationEditorUtils.h"
 
 #define LOCTEXT_NAMESPACE	"AnimSequenceDetails"
 
@@ -20,7 +21,9 @@ TSharedRef<IDetailCustomization> FAnimSequenceDetails::MakeInstance()
 BEGIN_SLATE_FUNCTION_BUILD_OPTIMIZATION
 void FAnimSequenceDetails::CustomizeDetails(IDetailLayoutBuilder& DetailBuilder)
 {
-	// retarget source handler
+	/////////////////////////////////////////////////////////////////////////////////
+	// retarget source handler in Animation
+	/////////////////////////////////////////////////////////////////////////////////
 	IDetailCategoryBuilder& AnimationCategory = DetailBuilder.EditCategory("Animation");
 	RetargetSourceNameHandler = DetailBuilder.GetProperty(GET_MEMBER_NAME_CHECKED(UAnimSequence, RetargetSource));
 
@@ -36,20 +39,27 @@ void FAnimSequenceDetails::CustomizeDetails(IDetailLayoutBuilder& DetailBuilder)
 	
 	// Check if we use only one skeleton
 	USkeleton* Skeleton = NULL;
+	SelectedAnimSequences.Reset();
 	TArray< TWeakObjectPtr<UObject> > SelectedObjectsList = DetailBuilder.GetDetailsView().GetSelectedObjects();
 	for (auto SelectionIt = SelectedObjectsList.CreateIterator(); SelectionIt; ++SelectionIt)
 	{
 		if (UAnimSequence* TestAnimSequence = Cast<UAnimSequence>(SelectionIt->Get()))
 		{
-			// we should only have one selected anim sequence
-			if(Skeleton && Skeleton != TestAnimSequence->GetSkeleton())
-			{
-				// Multiple different skeletons
-				Skeleton = NULL;
-				break;
-			}
-			Skeleton = TestAnimSequence->GetSkeleton();
+			SelectedAnimSequences.Add(TestAnimSequence);
+		}
+	}
+
+	// do it in separate loop since before it only cared AnimSequence 
+	for (auto& It : SelectedAnimSequences)
+	{
+		// we should only have one selected anim sequence
+		if(Skeleton && Skeleton != It->GetSkeleton())
+		{
+			// Multiple different skeletons
+			Skeleton = NULL;
+			break;
 		}
+		Skeleton = It->GetSkeleton();
 	}
 
 	// set target skeleton. It can be null
@@ -99,6 +109,9 @@ void FAnimSequenceDetails::CustomizeDetails(IDetailLayoutBuilder& DetailBuilder)
 
 	DetailBuilder.HideProperty(RetargetSourceNameHandler);
 
+	/////////////////////////////////////////////////////////////////////////////
+	// Additive settings category
+	/////////////////////////////////////////////////////////////////////////////
 	// now customize to combo box
 	IDetailCategoryBuilder& AdditiveSettingsCategory = DetailBuilder.EditCategory("AdditiveSettings");
 
@@ -139,6 +152,38 @@ void FAnimSequenceDetails::CustomizeDetails(IDetailLayoutBuilder& DetailBuilder)
 	];
 
 	CreateOverridenProperty(DetailBuilder, AdditiveSettingsCategory, RefFrameIndexHandle, TAttribute<EVisibility>( this, &FAnimSequenceDetails::ShouldShowRefFrameIndex ));
+
+	/////////////////////////////////////////////////////////////////
+	//  compression category!
+	///////////////////////////////////////////////////////////////
+	// add Apply button Compression
+	IDetailCategoryBuilder& CompressionCategory = DetailBuilder.EditCategory("Compression");
+
+	TArray<TSharedRef<IPropertyHandle>> CompressionProperties;
+
+	CompressionCategory.GetDefaultProperties( CompressionProperties );
+
+	for (auto& Property : CompressionProperties)
+	{
+		CompressionCategory.AddProperty(Property);
+	}
+
+	FDetailWidgetRow& CustomRow = CompressionCategory.AddCustomRow(LOCTEXT("ApplyCompressionLabel", "Apply"))
+		.WholeRowContent()
+		[
+			SNew(SVerticalBox)
+			+SVerticalBox::Slot()
+			.AutoHeight()
+			.Padding(2.0f, 0.0f)
+			.VAlign(VAlign_Center)
+			.HAlign(HAlign_Left)
+			[
+				SNew(SButton)
+				.Text(LOCTEXT("ApplyCompressionButton_Label", "Apply"))
+				.ToolTipText(LOCTEXT("ApplyCompressionButton_Tooltip", "Click Apply to apply the changes you've made to Compression Settings"))
+				.OnClicked(this, &FAnimSequenceDetails::OnApplyCompression)
+			]
+		];
 }
 END_SLATE_FUNCTION_BUILD_OPTIMIZATION
 
@@ -286,7 +331,34 @@ TSharedPtr<FString> FAnimSequenceDetails::GetRetargetSourceString(FName Retarget
 	return RetargetSourceComboList[0];
 }
 
+FReply FAnimSequenceDetails::OnApplyCompression()
+{
+	TArray<UAnimSequence*> AnimSequences;
+	UAnimCompress* CompressAlgorithm = NULL;
+
+	for(auto SelectionIt = SelectedAnimSequences.CreateIterator(); SelectionIt; ++SelectionIt)
+	{
+		if(UAnimSequence* TestAnimSequence = Cast<UAnimSequence>(SelectionIt->Get()))
+		{
+			AnimSequences.Add(TestAnimSequence);
 
+			if (!CompressAlgorithm)
+			{
+				// it should be all same by now
+				// better way is to get this from property handler
+				// but by this time, they should have all same value
+				CompressAlgorithm = TestAnimSequence->CompressionScheme;
+			}
+		}
+	}
+
+	if (CompressAlgorithm)
+	{
+		AnimationEditorUtils::ApplyCompressionAlgorithm(AnimSequences, CompressAlgorithm);
+	}
+
+	return FReply::Handled();
+}
 ////////////////////////////////////////////////
 // based on SAnimationSegmentViewport
 SAnimationRefPoseViewport::~SAnimationRefPoseViewport()
diff --git a/Engine/Source/Editor/DetailCustomizations/Private/AnimSequenceDetails.h b/Engine/Source/Editor/DetailCustomizations/Private/AnimSequenceDetails.h
--- a/Engine/Source/Editor/DetailCustomizations/Private/AnimSequenceDetails.h
+++ b/Engine/Source/Editor/DetailCustomizations/Private/AnimSequenceDetails.h
@@ -27,6 +27,7 @@ private:
 	
 	TWeakObjectPtr<USkeleton> TargetSkeleton;
 
+	// additive setting handler
 	TSharedPtr<IPropertyHandle> AdditiveAnimTypeHandle;
 	TSharedPtr<IPropertyHandle> RefPoseTypeHandle;
 	TSharedPtr<IPropertyHandle> RefPoseSeqHandle;
@@ -49,6 +50,12 @@ private:
 	FDelegateHandle OnDelegateRetargetSourceChangedDelegateHandle;
 	void RegisterRetargetSourceChanged();
 	void DelegateRetargetSourceChanged();
+
+	// button handler for Apply Compression
+	// cache all anim sequences that are selected
+	// but I need to know them before compress
+	TArray< TWeakObjectPtr<UAnimSequence> > SelectedAnimSequences;
+	FReply OnApplyCompression();
 };
 
 ///////////////////////////////////////////////////
diff --git a/Engine/Source/Editor/UnrealEd/Private/AnimationCompressionPanel.cpp b/Engine/Source/Editor/UnrealEd/Private/AnimationCompressionPanel.cpp
--- a/Engine/Source/Editor/UnrealEd/Private/AnimationCompressionPanel.cpp
+++ b/Engine/Source/Editor/UnrealEd/Private/AnimationCompressionPanel.cpp
@@ -3,6 +3,7 @@
 #include "UnrealEd.h"
 #include "Animation/AnimCompress.h"
 #include "AnimationCompressionPanel.h"
+#include "AnimationEditorUtils.h"
 
 #define LOCTEXT_NAMESPACE "AnimationCompression"
 
@@ -128,24 +129,16 @@ void SAnimationCompressionPanel::ApplyAlgorithm(class UAnimCompress* Algorithm)
 {
 	if ( Algorithm )
 	{
-		const bool bProceed = EAppReturnType::Yes == FMessageDialog::Open( EAppMsgType::YesNo,
-			FText::Format( NSLOCTEXT("UnrealEd", "AboutToCompressAnimations_F", "About to compress {0} animations.  Proceed?"), FText::AsNumber(AnimSequences.Num()) ) );
-		if ( bProceed )
-		{
-			TArray<UAnimSequence*> AnimSequencePtrs;
-			AnimSequencePtrs.Reserve(AnimSequences.Num());
-
-			for(int32 Index = 0; Index < AnimSequences.Num(); ++Index)
-			{
-				AnimSequencePtrs.Add(AnimSequences[Index].Get());
-			}
+		TArray<UAnimSequence*> AnimSequencePtrs;
+		AnimSequencePtrs.Reserve(AnimSequences.Num());
 
-			GWarn->BeginSlowTask( LOCTEXT("AnimCompressing", "Compressing"), true);
-			
-			Algorithm->Reduce(AnimSequencePtrs, true);
-
-			GWarn->EndSlowTask( );
+		for(int32 Index = 0; Index < AnimSequences.Num(); ++Index)
+		{
+			AnimSequencePtrs.Add(AnimSequences[Index].Get());
+		}
 
+		if (AnimationEditorUtils::ApplyCompressionAlgorithm(AnimSequencePtrs, Algorithm))
+		{
 			ParentWindow->RequestDestroyWindow();
 		}
 	}
diff --git a/Engine/Source/Editor/UnrealEd/Private/AnimationEditorUtils.cpp b/Engine/Source/Editor/UnrealEd/Private/AnimationEditorUtils.cpp
--- a/Engine/Source/Editor/UnrealEd/Private/AnimationEditorUtils.cpp
+++ b/Engine/Source/Editor/UnrealEd/Private/AnimationEditorUtils.cpp
@@ -413,6 +413,27 @@ namespace AnimationEditorUtils
 		}
 		MenuBuilder.EndSection();
 	}
+
+	bool ApplyCompressionAlgorithm(TArray<UAnimSequence*>& AnimSequencePtrs, class UAnimCompress* Algorithm)
+	{
+		if(Algorithm)
+		{
+			const bool bProceed = (AnimSequencePtrs.Num() > 1)? EAppReturnType::Yes == FMessageDialog::Open(EAppMsgType::YesNo,
+				FText::Format(NSLOCTEXT("UnrealEd", "AboutToCompressAnimations_F", "About to compress {0} animations.  Proceed?"), FText::AsNumber(AnimSequencePtrs.Num()))) : true;
+			if(bProceed)
+			{
+				GWarn->BeginSlowTask(LOCTEXT("AnimCompressing", "Compressing"), true);
+
+				Algorithm->Reduce(AnimSequencePtrs, true);
+
+				GWarn->EndSlowTask();
+
+				return true;
+			}
+		}
+
+		return false;
+	}
 }
 
 #undef LOCTEXT_NAMESPACE
\ No newline at end of file
diff --git a/Engine/Source/Editor/UnrealEd/Public/AnimationEditorUtils.h b/Engine/Source/Editor/UnrealEd/Public/AnimationEditorUtils.h
--- a/Engine/Source/Editor/UnrealEd/Public/AnimationEditorUtils.h
+++ b/Engine/Source/Editor/UnrealEd/Public/AnimationEditorUtils.h
@@ -59,6 +59,7 @@ namespace AnimationEditorUtils
 	UNREALED_API void CreateNewAnimBlueprint(TArray<TWeakObjectPtr<USkeleton>> Skeletons, FAnimAssetCreated AssetCreated);
 	UNREALED_API void FillCreateAssetMenu(FMenuBuilder& MenuBuilder, TArray<TWeakObjectPtr<USkeleton>> Skeletons, FAnimAssetCreated AssetCreated, bool bInContentBrowser=true);
 	UNREALED_API void CreateUniqueAssetName(const FString& InBasePackageName, const FString& InSuffix, FString& OutPackageName, FString& OutAssetName);
+	UNREALED_API bool ApplyCompressionAlgorithm(TArray<UAnimSequence*>& AnimSequencePtrs, class UAnimCompress* Algorithm);
 
 	template <typename TFactory, typename T>
 	void ExecuteNewAnimAsset(TArray<TWeakObjectPtr<USkeleton>> Objects, const FString InSuffix, FAnimAssetCreated AssetCreated, bool bInContentBrowser)
diff --git a/Engine/Source/Runtime/Engine/Classes/Animation/AnimCompress.h b/Engine/Source/Runtime/Engine/Classes/Animation/AnimCompress.h
--- a/Engine/Source/Runtime/Engine/Classes/Animation/AnimCompress.h
+++ b/Engine/Source/Runtime/Engine/Classes/Animation/AnimCompress.h
@@ -10,7 +10,7 @@
 #include "AnimSequence.h"
 #include "AnimCompress.generated.h"
 
-UCLASS(abstract, hidecategories=Object, MinimalAPI)
+UCLASS(abstract, hidecategories=Object, MinimalAPI, EditInlineNew)
 class UAnimCompress : public UObject
 {
 	GENERATED_UCLASS_BODY()
diff --git a/Engine/Source/Runtime/Engine/Classes/Animation/AnimSequence.h b/Engine/Source/Runtime/Engine/Classes/Animation/AnimSequence.h
--- a/Engine/Source/Runtime/Engine/Classes/Animation/AnimSequence.h
+++ b/Engine/Source/Runtime/Engine/Classes/Animation/AnimSequence.h
@@ -386,7 +386,7 @@ class ENGINE_API UAnimSequence : public UAnimSequenceBase
 	 * The compression scheme that was most recently used to compress this animation.
 	 * May be NULL.
 	 */
-	UPROPERTY(Instanced, Category=Compression, VisibleAnywhere)
+	UPROPERTY(Instanced, Category=Compression, EditAnywhere)
 	class UAnimCompress* CompressionScheme;
 #endif // WITH_EDITORONLY_DATA
 
diff --git a/Engine/Source/Runtime/Engine/Private/Animation/AnimationUtils.cpp b/Engine/Source/Runtime/Engine/Private/Animation/AnimationUtils.cpp
--- a/Engine/Source/Runtime/Engine/Private/Animation/AnimationUtils.cpp
+++ b/Engine/Source/Runtime/Engine/Private/Animation/AnimationUtils.cpp
@@ -436,7 +436,7 @@ bool FAnimationUtils::GetForcedRecompressionSetting()
 {																																							\
 	/* try the alternative compressor	*/																													\
 	(compressionalgorithm)->Reduce( AnimSeq, bOutput );																				\
-	const SIZE_T NewSize = AnimSeq->GetResourceSize(EResourceSizeMode::Exclusive);																												\
+	const SIZE_T NewSize = AnimSeq->GetResourceSize(EResourceSizeMode::Inclusive);																												\
 																																							\
 	/* compute the savings and compression error*/																											\
 	const SIZE_T MemorySavingsFromOriginal = OriginalSize - NewSize;																									\
@@ -489,7 +489,7 @@ bool FAnimationUtils::GetForcedRecompressionSetting()
 		AnimSeq->ScaleCodec = SavedScaleCodec;																												\
 		AnimationFormat_SetInterfaceLinks(*AnimSeq);																										\
 																																							\
-		const SIZE_T RestoredSize = AnimSeq->GetResourceSize(EResourceSizeMode::Exclusive);																										\
+		const SIZE_T RestoredSize = AnimSeq->GetResourceSize(EResourceSizeMode::Inclusive);																										\
 		check(RestoredSize == CurrentSize);																													\
 	}																																						\
 	else																																					\
@@ -685,7 +685,7 @@ void FAnimationUtils::CompressAnimSequenceExplicit(
 		bool const bTryAlternateCompressor = MasterTolerance > 0.0f;
 
 		// Get the current size
-		int32 OriginalSize = AnimSeq->GetResourceSize(EResourceSizeMode::Exclusive);
+		int32 OriginalSize = AnimSeq->GetResourceSize(EResourceSizeMode::Inclusive);
 		TotalSizeBefore += OriginalSize;
 
 		// Estimate total uncompressed
@@ -721,14 +721,14 @@ void FAnimationUtils::CompressAnimSequenceExplicit(
 			}
 
 			OriginalCompressionAlgorithm->Reduce( AnimSeq, bOutput );
-			AfterOriginalRecompression = AnimSeq->GetResourceSize(EResourceSizeMode::Exclusive);
+			AfterOriginalRecompression = AnimSeq->GetResourceSize(EResourceSizeMode::Inclusive);
 
 			// figure out our current compression error
 			FAnimationUtils::ComputeCompressionError(AnimSeq, BoneData, OriginalErrorStats);
 		}
 		else
 		{
-			AfterOriginalRecompression = AnimSeq->GetResourceSize(EResourceSizeMode::Exclusive);
+			AfterOriginalRecompression = AnimSeq->GetResourceSize(EResourceSizeMode::Inclusive);
 			OriginalErrorStats = TrueOriginalErrorStats;
 		}
  
@@ -764,7 +764,7 @@ void FAnimationUtils::CompressAnimSequenceExplicit(
 				++TotalRecompressions;
 
 				// Prepare to compress
-				int32 CurrentSize = AnimSeq->GetResourceSize(EResourceSizeMode::Exclusive);
+				int32 CurrentSize = AnimSeq->GetResourceSize(EResourceSizeMode::Inclusive);
 				int32* WinningCompressorCounter = NULL;
 				float* WinningCompressorErrorSum = NULL;
 				int32* WinningCompressorMarginalSavingsSum = NULL;
@@ -1254,7 +1254,7 @@ void FAnimationUtils::CompressAnimSequenceExplicit(
 				}
 
 				// Make sure we got that right.
-				check(CurrentSize == AnimSeq->GetResourceSize(EResourceSizeMode::Exclusive));
+				check(CurrentSize == AnimSeq->GetResourceSize(EResourceSizeMode::Inclusive));
 				TotalSizeNow += CurrentSize;
 
 				PctSaving = TotalSizeBefore > 0 ? 100.f - (100.f * float(TotalSizeNow) / float(TotalSizeBefore)) : 0.f;
@@ -1336,7 +1336,7 @@ void FAnimationUtils::CompressAnimSequenceExplicit(
 		// Do not recompress - Still take into account size for stats.
 		else
 		{
-			TotalSizeNow += AnimSeq->GetResourceSize(EResourceSizeMode::Exclusive);
+			TotalSizeNow += AnimSeq->GetResourceSize(EResourceSizeMode::Inclusive);
 		}
 	}
 	else
