Fixing UTF-8 bullet points. Changing to asterisks because asterisks show up correctly in the editor.

diff --git a/Engine/Source/Editor/DetailCustomizations/Private/LightComponentDetails.cpp b/Engine/Source/Editor/DetailCustomizations/Private/LightComponentDetails.cpp
--- a/Engine/Source/Editor/DetailCustomizations/Private/LightComponentDetails.cpp
+++ b/Engine/Source/Editor/DetailCustomizations/Private/LightComponentDetails.cpp
@@ -16,7 +16,7 @@ void FLightComponentDetails::CustomizeDetails( IDetailLayoutBuilder& DetailBuild
 	// Mobility property is on the scene component base class not the light component and that is why we have to use USceneComponent::StaticClass
 	TSharedRef<IPropertyHandle> MobilityHandle = DetailBuilder.GetProperty(GET_MEMBER_NAME_CHECKED(ULightComponent, Mobility), USceneComponent::StaticClass());
 	// Set a mobility tooltip specific to lights
-	MobilityHandle->SetToolTipText(LOCTEXT("LightMobilityTooltip", "Mobility for lights controls what the light is allowed to do at runtime and therefore what rendering methods are used.\n● A movable light uses fully dynamic lighting and anything can change in game, however it has a large performance cost, typically proportional to the light's influence size.\n● A stationary light will only have its shadowing and bounced lighting from static geometry baked by Lightmass, all other lighting will be dynamic.  It can change color and intensity in game. \n● A static light is fully baked into lightmaps and therefore has no performance cost, but also can't change in game."));
+	MobilityHandle->SetToolTipText(LOCTEXT("LightMobilityTooltip", "Mobility for lights controls what the light is allowed to do at runtime and therefore what rendering methods are used.\n* A movable light uses fully dynamic lighting and anything can change in game, however it has a large performance cost, typically proportional to the light's influence size.\n* A stationary light will only have its shadowing and bounced lighting from static geometry baked by Lightmass, all other lighting will be dynamic.  It can change color and intensity in game. \n* A static light is fully baked into lightmaps and therefore has no performance cost, but also can't change in game."));
 
 	IDetailCategoryBuilder& LightCategory = DetailBuilder.EditCategory( "Light", FText::GetEmpty(), ECategoryPriority::TypeSpecific );
 
diff --git a/Engine/Source/Editor/DetailCustomizations/Private/MobilityCustomization.cpp b/Engine/Source/Editor/DetailCustomizations/Private/MobilityCustomization.cpp
--- a/Engine/Source/Editor/DetailCustomizations/Private/MobilityCustomization.cpp
+++ b/Engine/Source/Editor/DetailCustomizations/Private/MobilityCustomization.cpp
@@ -40,7 +40,7 @@ void FMobilityCustomization::CreateMobilityCustomization(IDetailCategoryBuilder&
 			.Style(FEditorStyle::Get(), "Property.ToggleButton.Start")
 			.IsChecked(this, &FMobilityCustomization::IsMobilityActive, EComponentMobility::Static)
 			.OnCheckStateChanged(this, &FMobilityCustomization::OnMobilityChanged, EComponentMobility::Static)
-			.ToolTipText(LOCTEXT("Mobility_Static_Tooltip", "A static object can't be changed in game.\n● Allows Baked Lighting\n● Fastest Rendering"))
+			.ToolTipText(LOCTEXT("Mobility_Static_Tooltip", "A static object can't be changed in game.\n* Allows Baked Lighting\n* Fastest Rendering"))
 			[
 				SNew(SHorizontalBox)
 
@@ -79,7 +79,7 @@ void FMobilityCustomization::CreateMobilityCustomization(IDetailCategoryBuilder&
 			.IsChecked(this, &FMobilityCustomization::IsMobilityActive, EComponentMobility::Stationary)
 			.Style(FEditorStyle::Get(), ( ColumnIndex == 0 ) ? "Property.ToggleButton.Start" : "Property.ToggleButton.Middle")
 			.OnCheckStateChanged(this, &FMobilityCustomization::OnMobilityChanged, EComponentMobility::Stationary)
-			.ToolTipText(LOCTEXT("Mobility_Stationary_Tooltip", "A stationary light will only have its shadowing and bounced lighting from static geometry baked by Lightmass, all other lighting will be dynamic.  It can change color and intensity in game.\n● Can't Move\n● Allows Partial Baked Lighting\n● Dynamic Shadows"))
+			.ToolTipText(LOCTEXT("Mobility_Stationary_Tooltip", "A stationary light will only have its shadowing and bounced lighting from static geometry baked by Lightmass, all other lighting will be dynamic.  It can change color and intensity in game.\n* Can't Move\n* Allows Partial Baked Lighting\n* Dynamic Shadows"))
 			[
 				SNew(SHorizontalBox)
 
@@ -116,7 +116,7 @@ void FMobilityCustomization::CreateMobilityCustomization(IDetailCategoryBuilder&
 		.IsChecked(this, &FMobilityCustomization::IsMobilityActive, EComponentMobility::Movable)
 		.Style(FEditorStyle::Get(), ( ColumnIndex == 0 ) ? "Property.ToggleButton" : "Property.ToggleButton.End")
 		.OnCheckStateChanged(this, &FMobilityCustomization::OnMobilityChanged, EComponentMobility::Movable)
-		.ToolTipText(LOCTEXT("Mobility_Movable_Tooltip", "Movable objects can be moved and changed in game.\n● Totally Dynamic\n● Allows Dynamic Shadows\n● Slowest Rendering"))
+		.ToolTipText(LOCTEXT("Mobility_Movable_Tooltip", "Movable objects can be moved and changed in game.\n* Totally Dynamic\n* Allows Dynamic Shadows\n* Slowest Rendering"))
 		[
 			SNew(SHorizontalBox)
 
diff --git a/Engine/Source/Editor/DetailCustomizations/Private/PrimitiveComponentDetails.cpp b/Engine/Source/Editor/DetailCustomizations/Private/PrimitiveComponentDetails.cpp
--- a/Engine/Source/Editor/DetailCustomizations/Private/PrimitiveComponentDetails.cpp
+++ b/Engine/Source/Editor/DetailCustomizations/Private/PrimitiveComponentDetails.cpp
@@ -71,7 +71,7 @@ void FPrimitiveComponentDetails::CustomizeDetails( IDetailLayoutBuilder& DetailB
 	}
 
 	TSharedRef<IPropertyHandle> MobilityHandle = DetailBuilder.GetProperty(GET_MEMBER_NAME_CHECKED(UPrimitiveComponent, Mobility), USceneComponent::StaticClass());
-	MobilityHandle->SetToolTipText(LOCTEXT("PrimitiveMobilityTooltip", "Mobility for primitive components controls how they can be modified in game and therefore how they interact with lighting and physics.\n● A movable primitive component can be changed in game, but requires dynamic lighting and shadowing from lights which have a large performance cost.\n● A static primitive component can't be changed in game, but can have its lighting baked, which allows rendering to be very efficient."));
+	MobilityHandle->SetToolTipText(LOCTEXT("PrimitiveMobilityTooltip", "Mobility for primitive components controls how they can be modified in game and therefore how they interact with lighting and physics.\n* A movable primitive component can be changed in game, but requires dynamic lighting and shadowing from lights which have a large performance cost.\n* A static primitive component can't be changed in game, but can have its lighting baked, which allows rendering to be very efficient."));
 
 
 	if(!HideCategories.Contains("Physics"))
diff --git a/Engine/Source/Runtime/Engine/Classes/Engine/Font.h b/Engine/Source/Runtime/Engine/Classes/Engine/Font.h
--- a/Engine/Source/Runtime/Engine/Classes/Engine/Font.h
+++ b/Engine/Source/Runtime/Engine/Classes/Engine/Font.h
@@ -89,8 +89,8 @@ struct TStructOpsTypeTraits<FFontCharacter> : public TStructOpsTypeTraitsBase
  * A font object, for use by Slate, UMG, and Canvas.
  *
  * A font can either be:
- *   ● Runtime cached - The font contains a series of TTF files that combine to form a composite font. The glyphs are cached on demand when required at runtime.
- *   ● Offline cached - The font contains a series of textures containing pre-baked cached glyphs and their associated texture coordinates.
+ *   * Runtime cached - The font contains a series of TTF files that combine to form a composite font. The glyphs are cached on demand when required at runtime.
+ *   * Offline cached - The font contains a series of textures containing pre-baked cached glyphs and their associated texture coordinates.
  */
 UCLASS(hidecategories=Object, autoexpandcategories=Font, MinimalAPI, BlueprintType)
 class UFont : public UObject, public IFontProviderInterface
diff --git a/Engine/Source/Runtime/Engine/Classes/Kismet/GameplayStatics.h b/Engine/Source/Runtime/Engine/Classes/Kismet/GameplayStatics.h
--- a/Engine/Source/Runtime/Engine/Classes/Kismet/GameplayStatics.h
+++ b/Engine/Source/Runtime/Engine/Classes/Kismet/GameplayStatics.h
@@ -306,8 +306,8 @@ class ENGINE_API UGameplayStatics : public UBlueprintFunctionLibrary
 	/**
 	 * Plays a sound directly with no attenuation, perfect for UI sounds.
 	 *
-	 * ● Fire and Forget.
-	 * ● Not Replicated.
+	 * * Fire and Forget.
+	 * * Not Replicated.
 	 * @param Sound - Sound to play.
 	 * @param VolumeMultiplier - Multiplied with the volume to make the sound louder or softer.
 	 * @param PitchMultiplier - Multiplies the pitch.
@@ -319,7 +319,7 @@ class ENGINE_API UGameplayStatics : public UBlueprintFunctionLibrary
 	/**
 	 * Spawns a sound with no attenuation, perfect for UI sounds.
 	 *
-	 * ● Not Replicated.
+	 * * Not Replicated.
 	 * @param Sound - Sound to play.
 	 * @param VolumeMultiplier - Multiplied with the volume to make the sound louder or softer.
 	 * @param PitchMultiplier - Multiplies the pitch.
@@ -392,8 +392,8 @@ class ENGINE_API UGameplayStatics : public UBlueprintFunctionLibrary
 	/**
 	 * Plays a dialogue directly with no attenuation, perfect for UI.
 	 *
-	 * ● Fire and Forget.
-	 * ● Not Replicated.
+	 * * Fire and Forget.
+	 * * Not Replicated.
 	 * @param Dialogue - dialogue to play
 	 * @param Context - context the dialogue is to play in
 	 * @param VolumeMultiplier - Multiplied with the volume to make the sound louder or softer.
@@ -406,7 +406,7 @@ class ENGINE_API UGameplayStatics : public UBlueprintFunctionLibrary
 	/**
 	 * Spawns a dialogue with no attenuation, perfect for UI.
 	 *
-	 * ● Not Replicated.
+	 * * Not Replicated.
 	 * @param Dialogue - dialogue to play
 	 * @param Context - context the dialogue is to play in
 	 * @param VolumeMultiplier - Multiplied with the volume to make the sound louder or softer.
diff --git a/Engine/Source/Runtime/UMG/Public/Components/Border.h b/Engine/Source/Runtime/UMG/Public/Components/Border.h
--- a/Engine/Source/Runtime/UMG/Public/Components/Border.h
+++ b/Engine/Source/Runtime/UMG/Public/Components/Border.h
@@ -10,8 +10,8 @@ class USlateBrushAsset;
  * A border is a container widget that can contain one child widget, providing an opportunity 
  * to surround it with a background image and adjustable padding.
  *
- * ● Single Child
- * ● Image
+ * * Single Child
+ * * Image
  */
 UCLASS()
 class UMG_API UBorder : public UContentWidget
diff --git a/Engine/Source/Runtime/UMG/Public/Components/Button.h b/Engine/Source/Runtime/UMG/Public/Components/Button.h
--- a/Engine/Source/Runtime/UMG/Public/Components/Button.h
+++ b/Engine/Source/Runtime/UMG/Public/Components/Button.h
@@ -15,8 +15,8 @@ DECLARE_DYNAMIC_MULTICAST_DELEGATE(FOnButtonReleasedEvent);
  * can place any other widget inside a button to make a more complex and 
  * interesting click-able element in your UI.
  *
- * ● Single Child
- * ● Clickable
+ * * Single Child
+ * * Clickable
  */
 UCLASS()
 class UMG_API UButton : public UContentWidget
diff --git a/Engine/Source/Runtime/UMG/Public/Components/CanvasPanel.h b/Engine/Source/Runtime/UMG/Public/Components/CanvasPanel.h
--- a/Engine/Source/Runtime/UMG/Public/Components/CanvasPanel.h
+++ b/Engine/Source/Runtime/UMG/Public/Components/CanvasPanel.h
@@ -12,9 +12,9 @@ class UCanvasPanelSlot;
  * for manual layout, but bad when you want to procedurally just generate widgets and place them in a 
  * container (unless you want absolute layout).
  *
- * ● Many Children
- * ● Absolute Layout
- * ● Anchors
+ * * Many Children
+ * * Absolute Layout
+ * * Anchors
  */
 UCLASS()
 class UMG_API UCanvasPanel : public UPanelWidget
diff --git a/Engine/Source/Runtime/UMG/Public/Components/CheckBox.h b/Engine/Source/Runtime/UMG/Public/Components/CheckBox.h
--- a/Engine/Source/Runtime/UMG/Public/Components/CheckBox.h
+++ b/Engine/Source/Runtime/UMG/Public/Components/CheckBox.h
@@ -13,8 +13,8 @@ DECLARE_DYNAMIC_MULTICAST_DELEGATE_OneParam( FOnCheckBoxComponentStateChanged, b
  * 'indeterminable.  You can use the checkbox for a classic checkbox, or as a toggle button,
  * or as radio buttons.
  * 
- * ● Single Child
- * ● Toggle
+ * * Single Child
+ * * Toggle
  */
 UCLASS()
 class UMG_API UCheckBox : public UContentWidget
diff --git a/Engine/Source/Runtime/UMG/Public/Components/CircularThrobber.h b/Engine/Source/Runtime/UMG/Public/Components/CircularThrobber.h
--- a/Engine/Source/Runtime/UMG/Public/Components/CircularThrobber.h
+++ b/Engine/Source/Runtime/UMG/Public/Components/CircularThrobber.h
@@ -9,8 +9,8 @@ class USlateBrushAsset;
 /**
  * A throbber widget that orients images in a spinning circle.
  * 
- * ● No Children
- * ● Spinner Progress
+ * * No Children
+ * * Spinner Progress
  */
 UCLASS()
 class UMG_API UCircularThrobber : public UWidget
diff --git a/Engine/Source/Runtime/UMG/Public/Components/EditableTextBox.h b/Engine/Source/Runtime/UMG/Public/Components/EditableTextBox.h
--- a/Engine/Source/Runtime/UMG/Public/Components/EditableTextBox.h
+++ b/Engine/Source/Runtime/UMG/Public/Components/EditableTextBox.h
@@ -7,8 +7,8 @@
 /**
  * Allows the user to type in custom text.  Only permits a single line of text to be entered.
  * 
- * ● No Children
- * ● Text Entry
+ * * No Children
+ * * Text Entry
  */
 UCLASS(meta=(DisplayName="Text Box"))
 class UMG_API UEditableTextBox : public UWidget
diff --git a/Engine/Source/Runtime/UMG/Public/Components/GridPanel.h b/Engine/Source/Runtime/UMG/Public/Components/GridPanel.h
--- a/Engine/Source/Runtime/UMG/Public/Components/GridPanel.h
+++ b/Engine/Source/Runtime/UMG/Public/Components/GridPanel.h
@@ -7,7 +7,7 @@
 /**
  * A panel that evenly divides up available space between all of its children.
  * 
- * ● Many Children
+ * * Many Children
  */
 UCLASS()
 class UMG_API UGridPanel : public UPanelWidget
diff --git a/Engine/Source/Runtime/UMG/Public/Components/HorizontalBox.h b/Engine/Source/Runtime/UMG/Public/Components/HorizontalBox.h
--- a/Engine/Source/Runtime/UMG/Public/Components/HorizontalBox.h
+++ b/Engine/Source/Runtime/UMG/Public/Components/HorizontalBox.h
@@ -9,8 +9,8 @@ class UHorizontalBoxSlot;
 /**
  * Allows widgets to be laid out in a flow horizontally.
  *
- * ● Many Children
- * ● Flow Horizontal
+ * * Many Children
+ * * Flow Horizontal
  */
 UCLASS()
 class UMG_API UHorizontalBox : public UPanelWidget
diff --git a/Engine/Source/Runtime/UMG/Public/Components/Image.h b/Engine/Source/Runtime/UMG/Public/Components/Image.h
--- a/Engine/Source/Runtime/UMG/Public/Components/Image.h
+++ b/Engine/Source/Runtime/UMG/Public/Components/Image.h
@@ -9,7 +9,7 @@ class USlateBrushAsset;
 /**
  * The image widget allows you to display a Slate Brush, or texture or material in the UI.
  *
- * ● No Children
+ * * No Children
  */
 UCLASS()
 class UMG_API UImage : public UWidget
diff --git a/Engine/Source/Runtime/UMG/Public/Components/InvalidationBox.h b/Engine/Source/Runtime/UMG/Public/Components/InvalidationBox.h
--- a/Engine/Source/Runtime/UMG/Public/Components/InvalidationBox.h
+++ b/Engine/Source/Runtime/UMG/Public/Components/InvalidationBox.h
@@ -6,8 +6,8 @@
 
 /**
  * Invalidate
- * ● Single Child
- * ● Caching / Performance
+ * * Single Child
+ * * Caching / Performance
  */
 UCLASS()
 class UMG_API UInvalidationBox : public UContentWidget
diff --git a/Engine/Source/Runtime/UMG/Public/Components/MenuAnchor.h b/Engine/Source/Runtime/UMG/Public/Components/MenuAnchor.h
--- a/Engine/Source/Runtime/UMG/Public/Components/MenuAnchor.h
+++ b/Engine/Source/Runtime/UMG/Public/Components/MenuAnchor.h
@@ -9,8 +9,8 @@ DECLARE_DYNAMIC_MULTICAST_DELEGATE_OneParam(FOnMenuOpenChangedEvent, bool, bIsOp
 /**
  * The Menu Anchor allows you to specify an location that a popup menu should be anchored to, 
  * and should be summoned from.
- * ● Single Child
- * ● Popup
+ * * Single Child
+ * * Popup
  */
 UCLASS()
 class UMG_API UMenuAnchor : public UContentWidget
diff --git a/Engine/Source/Runtime/UMG/Public/Components/ProgressBar.h b/Engine/Source/Runtime/UMG/Public/Components/ProgressBar.h
--- a/Engine/Source/Runtime/UMG/Public/Components/ProgressBar.h
+++ b/Engine/Source/Runtime/UMG/Public/Components/ProgressBar.h
@@ -9,7 +9,7 @@
 /**
  * The progress bar widget is a simple bar that fills up that can be restyled to fit any number of uses.
  *
- * ● No Children
+ * * No Children
  */
 UCLASS()
 class UMG_API UProgressBar : public UWidget
diff --git a/Engine/Source/Runtime/UMG/Public/Components/RichTextBlock.h b/Engine/Source/Runtime/UMG/Public/Components/RichTextBlock.h
--- a/Engine/Source/Runtime/UMG/Public/Components/RichTextBlock.h
+++ b/Engine/Source/Runtime/UMG/Public/Components/RichTextBlock.h
@@ -7,8 +7,8 @@
 /**
  * The rich text block
  *
- * ● Fancy Text
- * ● No Children
+ * * Fancy Text
+ * * No Children
  */
 UCLASS(Experimental)
 class UMG_API URichTextBlock : public UWidget
diff --git a/Engine/Source/Runtime/UMG/Public/Components/ScaleBox.h b/Engine/Source/Runtime/UMG/Public/Components/ScaleBox.h
--- a/Engine/Source/Runtime/UMG/Public/Components/ScaleBox.h
+++ b/Engine/Source/Runtime/UMG/Public/Components/ScaleBox.h
@@ -9,8 +9,8 @@
  * you needed to have a background image scale to fill an area but not become distorted with different aspect ratios, or if you need
  * to auto fit some text to an area, this is the control for you.
  *
- * ● Single Child
- * ● Aspect Ratio
+ * * Single Child
+ * * Aspect Ratio
  */
 UCLASS()
 class UMG_API UScaleBox : public UContentWidget
diff --git a/Engine/Source/Runtime/UMG/Public/Components/SizeBox.h b/Engine/Source/Runtime/UMG/Public/Components/SizeBox.h
--- a/Engine/Source/Runtime/UMG/Public/Components/SizeBox.h
+++ b/Engine/Source/Runtime/UMG/Public/Components/SizeBox.h
@@ -10,8 +10,8 @@ class USlateBrushAsset;
  * A widget that allows you to specify the size it reports to have and desire.  Not all widgets report a desired size
  * that you actually desire.  Wrapping them in a SizeBox lets you have the Size Box force them to be a particular size.
  *
- * ● Single Child
- * ● Fixed Size
+ * * Single Child
+ * * Fixed Size
  */
 UCLASS()
 class UMG_API USizeBox : public UContentWidget
diff --git a/Engine/Source/Runtime/UMG/Public/Components/Slider.h b/Engine/Source/Runtime/UMG/Public/Components/Slider.h
--- a/Engine/Source/Runtime/UMG/Public/Components/Slider.h
+++ b/Engine/Source/Runtime/UMG/Public/Components/Slider.h
@@ -11,7 +11,7 @@ DECLARE_DYNAMIC_MULTICAST_DELEGATE_OneParam(FOnFloatValueChangedEvent, float, Va
 /**
  * A simple widget that shows a sliding bar with a handle that allows you to control the value between 0..1.
  *
- * ● No Children
+ * * No Children
  */
 UCLASS()
 class UMG_API USlider : public UWidget
diff --git a/Engine/Source/Runtime/UMG/Public/Components/Spacer.h b/Engine/Source/Runtime/UMG/Public/Components/Spacer.h
--- a/Engine/Source/Runtime/UMG/Public/Components/Spacer.h
+++ b/Engine/Source/Runtime/UMG/Public/Components/Spacer.h
@@ -7,7 +7,7 @@
 /**
  * A spacer widget; it does not have a visual representation, and just provides padding between other widgets.
  *
- * ● No Children
+ * * No Children
  */
 UCLASS()
 class UMG_API USpacer : public UWidget
diff --git a/Engine/Source/Runtime/UMG/Public/Components/TextBlock.h b/Engine/Source/Runtime/UMG/Public/Components/TextBlock.h
--- a/Engine/Source/Runtime/UMG/Public/Components/TextBlock.h
+++ b/Engine/Source/Runtime/UMG/Public/Components/TextBlock.h
@@ -7,8 +7,8 @@
 /**
  * A simple static text widget.
  *
- * ● No Children
- * ● Text
+ * * No Children
+ * * Text
  */
 UCLASS(meta=(DisplayName="Text"))
 class UMG_API UTextBlock : public UWidget
diff --git a/Engine/Source/Runtime/UMG/Public/Components/VerticalBox.h b/Engine/Source/Runtime/UMG/Public/Components/VerticalBox.h
--- a/Engine/Source/Runtime/UMG/Public/Components/VerticalBox.h
+++ b/Engine/Source/Runtime/UMG/Public/Components/VerticalBox.h
@@ -10,8 +10,8 @@ class UVerticalBoxSlot;
  * A vertical box widget is a layout panel allowing child widgets to be automatically laid out
  * vertically.
  *
- * ● Many Children
- * ● Flows Vertical
+ * * Many Children
+ * * Flows Vertical
  */
 UCLASS()
 class UMG_API UVerticalBox : public UPanelWidget
diff --git a/Engine/Source/Runtime/UMG/Public/Components/WrapBox.h b/Engine/Source/Runtime/UMG/Public/Components/WrapBox.h
--- a/Engine/Source/Runtime/UMG/Public/Components/WrapBox.h
+++ b/Engine/Source/Runtime/UMG/Public/Components/WrapBox.h
@@ -7,9 +7,9 @@
 /**
  * Arranges widgets left-to-right.  When the widgets exceed the Width it will place widgets on the next line.
  * 
- * ● Many Children
- * ● Flows
- * ● Wraps
+ * * Many Children
+ * * Flows
+ * * Wraps
  */
 UCLASS()
 class UMG_API UWrapBox : public UPanelWidget
