Fixed issue where clicking a Text property field corresponding to multiple properties would cause the "Multiple Values" text to be committed to those properties.

diff --git a/Engine/Source/Editor/DetailCustomizations/Private/TextCustomization.cpp b/Engine/Source/Editor/DetailCustomizations/Private/TextCustomization.cpp
--- a/Engine/Source/Editor/DetailCustomizations/Private/TextCustomization.cpp
+++ b/Engine/Source/Editor/DetailCustomizations/Private/TextCustomization.cpp
@@ -37,8 +37,12 @@ namespace
 		TOptional< float > PreviousHeight;
 
 		bool bIsMultiLine;
+
+		static FText MultipleValuesText;
 	};
 
+	FText STextPropertyWidget::MultipleValuesText(NSLOCTEXT("PropertyEditor", "MultipleValues", "Multiple Values"));
+
 	void STextPropertyWidget::Construct(const FArguments& InArgs, const TSharedRef<IPropertyHandle>& InPropertyHandle, const TSharedPtr<IPropertyUtilities>& InPropertyUtilities)
 	{
 		PropertyHandle = InPropertyHandle;
@@ -60,7 +64,7 @@ namespace
 			}
 			else if(RawData.Num() > 1)
 			{
-				TextValue = NSLOCTEXT("PropertyEditor", "MultipleValues", "Multiple Values");
+				TextValue = MultipleValuesText;
 			}
 			else
 			{
@@ -75,7 +79,8 @@ namespace
 			TArray<void*> RawData;
 			PropertyHandle->AccessRawData(RawData);
 
-			if (RawData.Num() > 0)
+			// Don't commit the Multiple Values text if there are multiple properties being set
+			if (RawData.Num() > 0 && (RawData.Num() == 1 || NewText.ToString() != MultipleValuesText.ToString()))
 			{
 				PropertyHandle->NotifyPreChange();
 
