Attempt to fix a non-unity CIS.

diff --git a/Engine/Source/Editor/DetailCustomizations/Private/AnimTrailNodeDetails.cpp b/Engine/Source/Editor/DetailCustomizations/Private/AnimTrailNodeDetails.cpp
--- a/Engine/Source/Editor/DetailCustomizations/Private/AnimTrailNodeDetails.cpp
+++ b/Engine/Source/Editor/DetailCustomizations/Private/AnimTrailNodeDetails.cpp
@@ -4,6 +4,7 @@
 
 #include "AnimTrailNodeDetails.h"
 #include "AnimGraphNode_Trail.h"
+#include "SCurveEditor.h"
 
 #define LOCTEXT_NAMESPACE "FAnimTrailNodeDetails"
 
