Fixed issue where pivot was not moved when entering a new brush translation in Geometry Mode.

diff --git a/Engine/Source/Editor/DetailCustomizations/Private/ComponentTransformDetails.cpp b/Engine/Source/Editor/DetailCustomizations/Private/ComponentTransformDetails.cpp
--- a/Engine/Source/Editor/DetailCustomizations/Private/ComponentTransformDetails.cpp
+++ b/Engine/Source/Editor/DetailCustomizations/Private/ComponentTransformDetails.cpp
@@ -1099,7 +1099,7 @@ void FComponentTransformDetails::OnSetLocation( float NewValue, ETextCommit::Typ
 
 	CacheTransform();
 
-	GUnrealEd->UpdatePivotLocationForSelection();
+	GUnrealEd->UpdatePivotLocationForSelection(true);
 	GUnrealEd->SetPivotMovedIndependently(false);
 	GUnrealEd->RedrawLevelEditingViewports();
 }
