Change so that Materials category always appears in Primitive Component detail panels with Material properties.

diff --git a/Engine/Source/Editor/DetailCustomizations/Private/ComponentMaterialCategory.cpp b/Engine/Source/Editor/DetailCustomizations/Private/ComponentMaterialCategory.cpp
--- a/Engine/Source/Editor/DetailCustomizations/Private/ComponentMaterialCategory.cpp
+++ b/Engine/Source/Editor/DetailCustomizations/Private/ComponentMaterialCategory.cpp
@@ -184,14 +184,13 @@ void FComponentMaterialCategory::Create( IDetailLayoutBuilder& DetailBuilder )
 void FComponentMaterialCategory::OnGetMaterialsForView( IMaterialListBuilder& MaterialList )
 {
 	const bool bAllowNullEntries = true;
+	bool bAnyMaterialsToDisplay = false;
 
 	// Iterate over every material on the actors
 	for( FMaterialIterator It( SelectedComponents ); It; ++It )
 	{	
 		int32 MaterialIndex = It.GetMaterialIndex();
 
-		MaterialCategory->SetCategoryVisibility( true );
-
 		UActorComponent* CurrentComponent = It.GetComponent();
 
 		if( CurrentComponent )
@@ -210,9 +209,12 @@ void FComponentMaterialCategory::OnGetMaterialsForView( IMaterialListBuilder& Ma
 			if( bAllowNullEntries || Material )
 			{
 				MaterialList.AddMaterial( MaterialIndex, Material, bCanBeReplaced );
+				bAnyMaterialsToDisplay = true;
 			}
 		}
 	}
+
+	MaterialCategory->SetCategoryVisibility(bAnyMaterialsToDisplay);
 }
 
 void FComponentMaterialCategory::OnMaterialChanged( UMaterialInterface* NewMaterial, UMaterialInterface* PrevMaterial, int32 SlotIndex, bool bReplaceAll )
