Improved asset import data customization

It now fills the width of the details panel, and has a button to clear the import path.

diff --git a/Engine/Source/Editor/DetailCustomizations/Private/AssetImportDataCustomization.cpp b/Engine/Source/Editor/DetailCustomizations/Private/AssetImportDataCustomization.cpp
--- a/Engine/Source/Editor/DetailCustomizations/Private/AssetImportDataCustomization.cpp
+++ b/Engine/Source/Editor/DetailCustomizations/Private/AssetImportDataCustomization.cpp
@@ -40,6 +40,7 @@ void FAssetImportDataCustomization::CustomizeChildren( TSharedRef<IPropertyHandl
 			.Font(Font)
 		]
 		.ValueContent()
+		.HAlign(HAlign_Fill)
 		.MaxDesiredWidth(TOptional<float>())
 		[
 			SNew(SVerticalBox)
@@ -70,6 +71,21 @@ void FAssetImportDataCustomization::CustomizeChildren( TSharedRef<IPropertyHandl
 						.Font(Font)
 					]
 				]
+
+				+SHorizontalBox::Slot()
+				.AutoWidth()
+				[
+					SNew(SButton)
+					.VAlign(VAlign_Center)
+					.HAlign(HAlign_Center)
+					.ButtonStyle(FEditorStyle::Get(), "HoverHintOnly")
+					.OnClicked(this, &FAssetImportDataCustomization::OnClearPathClicked, Index)
+					.ToolTipText(LOCTEXT("ClearPath_Tooltip", "Clear this source file information from the asset"))
+					[
+						SNew(SImage)
+						.Image(FEditorStyle::GetBrush("Cross"))
+					]
+				]
 			]
 
 			+ SVerticalBox::Slot()
@@ -156,5 +172,17 @@ FReply FAssetImportDataCustomization::OnChangePathClicked(int32 Index) const
 	return FReply::Handled();
 }
 
+FReply FAssetImportDataCustomization::OnClearPathClicked(int32 Index) const
+{
+	UAssetImportData* ImportData = GetOuterClass();
+	if (ImportData && ImportData->SourceData.SourceFiles.IsValidIndex(Index))
+	{
+		ImportData->SourceData.SourceFiles[Index] = FAssetImportInfo::FSourceFile(FString());
+		ImportData->MarkPackageDirty();
+	}
+
+	return FReply::Handled();
+}
+
 
 #undef LOCTEXT_NAMESPACE
diff --git a/Engine/Source/Editor/DetailCustomizations/Private/AssetImportDataCustomization.h b/Engine/Source/Editor/DetailCustomizations/Private/AssetImportDataCustomization.h
--- a/Engine/Source/Editor/DetailCustomizations/Private/AssetImportDataCustomization.h
+++ b/Engine/Source/Editor/DetailCustomizations/Private/AssetImportDataCustomization.h
@@ -17,8 +17,8 @@ private:
 	/** Handle the user wanting to change a specific source path */
 	FReply OnChangePathClicked(int32 Index) const;
 
-	/** Update the filename strings */
-	void UpdateCachedData();
+	/** Handle the user requesting that the specified index be cleared */
+	FReply OnClearPathClicked(int32 Index) const;
 
 	/** Access the struct we are editing */
 	FAssetImportInfo* GetEditStruct() const;
