#ANIM: Trail node

 - at the end, no change on the data because this guarantees frame independent.
- Added help and renamed variable to speed, so that it knows it is speed.

Merging

//depot/UE4-Orion/Engine/...

to //depot/UE4-Orion/Engine/...

diff --git a/Engine/Config/BaseEngine.ini b/Engine/Config/BaseEngine.ini
--- a/Engine/Config/BaseEngine.ini
+++ b/Engine/Config/BaseEngine.ini
@@ -392,6 +392,7 @@ bUseStreamingPause=false
 +ActiveClassRedirects=(OldClassName="AnimNode_TwoBoneIK",NewClassName="/Script/AnimGraphRuntime.AnimNode_TwoBoneIK")
 +ActiveClassRedirects=(OldClassName="AnimNode_WheelHandler",NewClassName="/Script/AnimGraphRuntime.AnimNode_WheelHandler")
 
++TaggedPropertyRedirects=(ClassName="AnimNode_Trail",OldPropertyName="TrailRelaxationCurve",NewPropertyName="TrailRelaxationSpeed")
 
 +ActiveClassRedirects=(OldClassName="AnimGraphNode_BlendSpace",NewClassName="/Script/AnimGraph.AnimGraphNode_BlendSpacePlayer")
 +ActiveStructRedirects=(OldStructName="AnimNode_BlendSpace",NewStructName="AnimNode_BlendSpacePlayer")
diff --git a/Engine/Documentation/Source/Shared/Types/FAnimNode_Trail/FAnimNode_Trail.INT.udn b/Engine/Documentation/Source/Shared/Types/FAnimNode_Trail/FAnimNode_Trail.INT.udn
new file mode 100644
--- /dev/null
+++ b/Engine/Documentation/Source/Shared/Types/FAnimNode_Trail/FAnimNode_Trail.INT.udn
@@ -0,0 +1,20 @@
+Availability:NoPublish
+Title:FAnimNode_Trail RTTs
+Crumbs:
+
+[EXCERPT:TrailRelaxationSpeed]
+
+### Overview
+
+Define curve of trail relaxation speed. Trail Relaxation Speed is the speed of trail to relax to animated pose. It is not [0, 1] because it needs to be frame independent. 
+
+![](Trail.png)
+
+The x axis is ranged from 0 to 1, and it is mapped to chain of joints. The 0 value is mapped to root joint, and 1 will be mapped to bottom joint. 
+If you do have chain of 10, it will evaluate interval of 9 and allocate from 0->9 joints. 
+
+The y axis is ranged from any positive and it is speed of relaxation. What this means is how fast it goes back to animated position. 
+It is absolute value to prevent frame dependency. The default value is from [10, 5], as the top part goes back faster than bottom part. 
+
+[/EXCERPT:TrailRelaxationSpeed]
+
diff --git a/Engine/Source/Editor/DetailCustomizations/Private/AnimTrailNodeDetails.cpp b/Engine/Source/Editor/DetailCustomizations/Private/AnimTrailNodeDetails.cpp
--- a/Engine/Source/Editor/DetailCustomizations/Private/AnimTrailNodeDetails.cpp
+++ b/Engine/Source/Editor/DetailCustomizations/Private/AnimTrailNodeDetails.cpp
@@ -36,9 +36,11 @@ void FAnimTrailNodeDetails::CustomizeDetails( IDetailLayoutBuilder& DetailBuilde
 
 	//Trail Relax curve
 	IDetailCategoryBuilder& TrailCategory = DetailBuilder.EditCategory("Trail");
-	TSharedRef<IPropertyHandle> TrailRelaxCurveHandle = DetailBuilder.GetProperty("Node.TrailRelaxationCurve");
+	TSharedRef<IPropertyHandle> TrailRelaxCurveHandle = DetailBuilder.GetProperty("Node.TrailRelaxationSpeed");
 	TSharedPtr<class SCurveEditor> TrailRelaxCurveWidget;
 
+	DetailBuilder.HideProperty(TrailRelaxCurveHandle);
+
 	TrailCategory.AddProperty(TrailRelaxCurveHandle).CustomWidget()
 	.NameContent()
 		[
@@ -66,7 +68,7 @@ void FAnimTrailNodeDetails::CustomizeDetails( IDetailLayoutBuilder& DetailBuilde
 TArray<FRichCurveEditInfoConst> FAnimTrailNodeDetails::FTrailRelaxCurveEditor::GetCurves() const
 {
 	TArray<FRichCurveEditInfoConst> Curves;
-	Curves.Add(&GraphNodeOwner->Node.TrailRelaxationCurve.EditorCurveData);
+	Curves.Add(&GraphNodeOwner->Node.TrailRelaxationSpeed.EditorCurveData);
 
 	return Curves;
 }
@@ -74,7 +76,7 @@ TArray<FRichCurveEditInfoConst> FAnimTrailNodeDetails::FTrailRelaxCurveEditor::G
 TArray<FRichCurveEditInfo> FAnimTrailNodeDetails::FTrailRelaxCurveEditor::GetCurves()
 {
 	TArray<FRichCurveEditInfo> Curves;
-	Curves.Add(&GraphNodeOwner->Node.TrailRelaxationCurve.EditorCurveData);
+	Curves.Add(&GraphNodeOwner->Node.TrailRelaxationSpeed.EditorCurveData);
 
 	return Curves;
 }
@@ -97,7 +99,7 @@ void FAnimTrailNodeDetails::FTrailRelaxCurveEditor::MakeTransactional()
 
 bool FAnimTrailNodeDetails::FTrailRelaxCurveEditor::IsValidCurve(FRichCurveEditInfo CurveInfo)
 {
-	return CurveInfo.CurveToEdit == &GraphNodeOwner->Node.TrailRelaxationCurve.EditorCurveData;
+	return CurveInfo.CurveToEdit == &GraphNodeOwner->Node.TrailRelaxationSpeed.EditorCurveData;
 }
 
 FAnimTrailNodeDetails::FTrailRelaxCurveEditor::FTrailRelaxCurveEditor(UAnimGraphNode_Trail * InGraphNode)
diff --git a/Engine/Source/Runtime/AnimGraphRuntime/Private/BoneControllers/AnimNode_Trail.cpp b/Engine/Source/Runtime/AnimGraphRuntime/Private/BoneControllers/AnimNode_Trail.cpp
--- a/Engine/Source/Runtime/AnimGraphRuntime/Private/BoneControllers/AnimNode_Trail.cpp
+++ b/Engine/Source/Runtime/AnimGraphRuntime/Private/BoneControllers/AnimNode_Trail.cpp
@@ -17,9 +17,9 @@ FAnimNode_Trail::FAnimNode_Trail()
 	, bActorSpaceFakeVel(false)
 	, bHadValidStrength(false)
 {
-	FRichCurve* TrailRelaxRichCurve = TrailRelaxationCurve.GetRichCurve();
+	FRichCurve* TrailRelaxRichCurve = TrailRelaxationSpeed.GetRichCurve();
 	TrailRelaxRichCurve->AddKey(0.f, 10.f);
-	TrailRelaxRichCurve->AddKey(1.f, 10.f);
+	TrailRelaxRichCurve->AddKey(1.f, 5.f);
 }
 
 void FAnimNode_Trail::Update(const FAnimationUpdateContext& Context)
@@ -143,7 +143,7 @@ void FAnimNode_Trail::EvaluateBoneTransforms(USkeletalMeshComponent* SkelComp, F
 		FVector ChildTarget = ParentPos + TargetDelta;
 
 		// Find vector from child to target
-		FVector Error = ChildTarget - ChildPos;
+		FVector Error = (ChildTarget - ChildPos);
 
 		// Calculate how much to push the child towards its target
 		float Correction = FMath::Clamp<float>(ThisTimstep * PerJointTrailData[i].TrailRelaxationSpeedPerSecond, 0.f, 1.f);
@@ -237,7 +237,7 @@ void FAnimNode_Trail::PostLoad()
 {
 	if (TrailRelaxation_DEPRECATED != 10.f)
 	{
-		FRichCurve* TrailRelaxRichCurve = TrailRelaxationCurve.GetRichCurve();
+		FRichCurve* TrailRelaxRichCurve = TrailRelaxationSpeed.GetRichCurve();
 		TrailRelaxRichCurve->Reset();
 		TrailRelaxRichCurve->AddKey(0.f, TrailRelaxation_DEPRECATED);
 		TrailRelaxRichCurve->AddKey(1.f, TrailRelaxation_DEPRECATED);
@@ -261,7 +261,7 @@ void FAnimNode_Trail::Initialize(const FAnimationInitializeContext& Context)
 		TrailBoneLocations.AddZeroed(ChainLength);
 
 		float Interval = (ChainLength > 1)? (1.f/(ChainLength-1)) : 0.f;
-		const FRichCurve* TrailRelaxRichCurve = TrailRelaxationCurve.GetRichCurveConst();
+		const FRichCurve* TrailRelaxRichCurve = TrailRelaxationSpeed.GetRichCurveConst();
 		ensure(TrailRelaxRichCurve);
 		for(int32 Idx=0; Idx<ChainLength; ++Idx)
 		{
diff --git a/Engine/Source/Runtime/AnimGraphRuntime/Public/BoneControllers/AnimNode_Trail.h b/Engine/Source/Runtime/AnimGraphRuntime/Public/BoneControllers/AnimNode_Trail.h
--- a/Engine/Source/Runtime/AnimGraphRuntime/Public/BoneControllers/AnimNode_Trail.h
+++ b/Engine/Source/Runtime/AnimGraphRuntime/Public/BoneControllers/AnimNode_Trail.h
@@ -46,9 +46,9 @@ struct ANIMGRAPHRUNTIME_API FAnimNode_Trail : public FAnimNode_SkeletalControlBa
 	UPROPERTY()
 	float	TrailRelaxation_DEPRECATED;
 
-	/** How quickly we 'relax' the bones to their animated positions. Time 0 will map to start joint and time 1 will map to the last chain bone. */
-	UPROPERTY(EditAnywhere, BlueprintReadWrite, Category=Trail)
-	FRuntimeFloatCurve TrailRelaxationCurve;
+	/** How quickly we 'relax' the bones to their animated positions. Time 0 will map to top root joint, time 1 will map to the bottom joint. */
+	UPROPERTY(EditAnywhere, BlueprintReadWrite, Category=Trail, meta=(CustomizeProperty))
+	FRuntimeFloatCurve TrailRelaxationSpeed;
 
 	/** If bLimitStretch is true, this indicates how long a bone can stretch beyond its length in the ref-pose. */
 	UPROPERTY(EditAnywhere, BlueprintReadWrite, Category=Trail)
