Scalabilty groups cvar added to device profile editor view.

diff --git a/Engine/Source/Editor/DetailCustomizations/Private/DeviceProfileDetails.cpp b/Engine/Source/Editor/DetailCustomizations/Private/DeviceProfileDetails.cpp
--- a/Engine/Source/Editor/DetailCustomizations/Private/DeviceProfileDetails.cpp
+++ b/Engine/Source/Editor/DetailCustomizations/Private/DeviceProfileDetails.cpp
@@ -67,6 +67,7 @@ namespace DeviceProfileCVarFormatHelper
 		CVG_Console,
 		CVG_Compatibility,
 		CVG_UserInterface,
+		CVG_ScalabilityGroups,
 
 		Max_CVarCategories,
 	};
@@ -110,6 +111,10 @@ namespace DeviceProfileCVarFormatHelper
 		{
 			CategoryText = LOCTEXT("UICVarGroupTitle", "User Interface");
 		}
+		else if (CatEnum == ECVarGroup::CVG_ScalabilityGroups)
+		{
+			CategoryText = LOCTEXT("ScalabilityGroupCVarGroupTitle", "Scalability Group");
+		}
 
 		return CategoryText;
 	}
@@ -151,6 +156,10 @@ namespace DeviceProfileCVarFormatHelper
 		{
 			CategoryText = LOCTEXT("UICVarGroupTitle", "User Interface");
 		}
+		else if (LowerPrefix == TEXT("sg") || LowerPrefix == TEXT("sg."))
+		{
+			CategoryText = LOCTEXT("ScalabilityGroupCVarGroupTitle", "Scalability Group");
+		}
 		else
 		{
 			CategoryText = LOCTEXT("UncategorizedCVarGroupTitle", "Uncategorized");
@@ -196,6 +205,10 @@ namespace DeviceProfileCVarFormatHelper
 		{
 			CVarPrefix = TEXT("ui");
 		}
+		else if (LowerCategory == TEXT("scalability group"))
+		{
+			CVarPrefix = TEXT("sg");
+		}
 
 		return CVarPrefix;
 	}
