Fallback to default size handler when the property resource is not a UTexture2D.
[UE-20553] UMG check box is not set to default size by default

diff --git a/Engine/Source/Editor/DetailCustomizations/Private/SlateBrushCustomization.cpp b/Engine/Source/Editor/DetailCustomizations/Private/SlateBrushCustomization.cpp
--- a/Engine/Source/Editor/DetailCustomizations/Private/SlateBrushCustomization.cpp
+++ b/Engine/Source/Editor/DetailCustomizations/Private/SlateBrushCustomization.cpp
@@ -1338,55 +1338,71 @@ EVisibility FSlateBrushStructCustomization::GetMarginPropertyVisibility() const
 
 bool FSlateBrushStructCustomization::IsImageSizeResetToDefaultVisible(TSharedRef<IPropertyHandle> PropertyHandle) const
 {
-	// get texture size from ResourceObjectProperty and compare to image size prop value
-	const FVector2D SizeDefault = GetDefaultImageSize();
+	UObject* ResourceObject;
+	if (FPropertyAccess::Success == ResourceObjectProperty->GetValue(ResourceObject) && ResourceObject && ResourceObject->IsA<UTexture2D>())
+	{
+		// get texture size from ResourceObjectProperty and compare to image size prop value
+		const FVector2D SizeDefault = GetDefaultImageSize();
 
-	FVector2D Size;
-	ImageSizeProperty->GetValue(Size);
+		FVector2D Size;
+		ImageSizeProperty->GetValue(Size);
 
-	if (PropertyHandle->GetProperty() == ImageSizeProperty->GetProperty())
-	{
-		// reseting the whole vector
-		return SizeDefault != Size;
-	}
-	else if (PropertyHandle->GetProperty() == ImageSizeProperty->GetChildHandle(0)->GetProperty())	// X
-	{
-		// reseting the vector.X
-		return SizeDefault.X != Size.X;
-	}
-	else if (PropertyHandle->GetProperty() == ImageSizeProperty->GetChildHandle(1)->GetProperty())	// Y
-	{
-		// reseting the vector.Y
-		return SizeDefault.Y != Size.Y;
+		if (PropertyHandle->GetProperty() == ImageSizeProperty->GetProperty())
+		{
+			// reseting the whole vector
+			return SizeDefault != Size;
+		}
+		else if (PropertyHandle->GetProperty() == ImageSizeProperty->GetChildHandle(0)->GetProperty())	// X
+		{
+			// reseting the vector.X
+			return SizeDefault.X != Size.X;
+		}
+		else if (PropertyHandle->GetProperty() == ImageSizeProperty->GetChildHandle(1)->GetProperty())	// Y
+		{
+			// reseting the vector.Y
+			return SizeDefault.Y != Size.Y;
+		}
+
+		ensureMsgf(false, TEXT("Property handle mismatch in brush size FVector2D struct"));
+		return false;
 	}
 
-	ensureMsgf(false, TEXT("Property handle mismatch in brush size FVector2D struct"));
-	return false;
+	// Fall back to default handler
+	return PropertyHandle->DiffersFromDefault();
 }
 
 void FSlateBrushStructCustomization::OnImageSizeResetToDefault(TSharedRef<IPropertyHandle> PropertyHandle) const
 {
-	// set image size prop value to the texture size in ResourceObjectProperty
-	const FVector2D SizeDefault = GetDefaultImageSize();
-
-	if (PropertyHandle->GetProperty() == ImageSizeProperty->GetProperty())
-	{
-		// reseting the whole vector
-		PropertyHandle->SetValue(SizeDefault);
-	}
-	else if (PropertyHandle->GetProperty() == ImageSizeProperty->GetChildHandle(0)->GetProperty())	// X
-	{
-		// reseting the vector.X
-		PropertyHandle->SetValue(SizeDefault.X);
-	}
-	else if (PropertyHandle->GetProperty() == ImageSizeProperty->GetChildHandle(1)->GetProperty())	// Y
+	UObject* ResourceObject;
+	if (FPropertyAccess::Success == ResourceObjectProperty->GetValue(ResourceObject) && ResourceObject && ResourceObject->IsA<UTexture2D>())
 	{
-		// reseting the vector.Y
-		PropertyHandle->SetValue(SizeDefault.Y);
-	}
+		// Set image size prop value to the texture size in ResourceObjectProperty
+		const FVector2D SizeDefault = GetDefaultImageSize();
+
+		if (PropertyHandle->GetProperty() == ImageSizeProperty->GetProperty())
+		{
+			// reseting the whole vector
+			PropertyHandle->SetValue(SizeDefault);
+		}
+		else if (PropertyHandle->GetProperty() == ImageSizeProperty->GetChildHandle(0)->GetProperty())	// X
+		{
+			// reseting the vector.X
+			PropertyHandle->SetValue(SizeDefault.X);
+		}
+		else if (PropertyHandle->GetProperty() == ImageSizeProperty->GetChildHandle(1)->GetProperty())	// Y
+		{
+			// reseting the vector.Y
+			PropertyHandle->SetValue(SizeDefault.Y);
+		}
+		else
+		{
+			ensureMsgf(false, TEXT("Property handle mismatch in brush size FVector2D struct"));
+		}
+	}	
 	else
 	{
-		ensureMsgf(false, TEXT("Property handle mismatch in brush size FVector2D struct"));
+		// Fall back to default handler. 
+		PropertyHandle->ResetToDefault();
 	}
 }
 
