Added detail customization for actor actor animation playback settings

diff --git a/Engine/Source/Editor/DetailCustomizations/Private/ActorAnimationPlaybackSettingsCustomization.cpp b/Engine/Source/Editor/DetailCustomizations/Private/ActorAnimationPlaybackSettingsCustomization.cpp
new file mode 100644
--- /dev/null
+++ b/Engine/Source/Editor/DetailCustomizations/Private/ActorAnimationPlaybackSettingsCustomization.cpp
@@ -0,0 +1,130 @@
+﻿// Copyright 1998-2015 Epic Games, Inc. All Rights Reserved.
+
+#include "DetailCustomizationsPrivatePCH.h"
+#include "ActorAnimationPlaybackSettingsCustomization.h"
+#include "ActorAnimationPlayer.h"
+
+#define LOCTEXT_NAMESPACE "ActorAnimationPlaybackSettingsCustomization"
+
+TSharedRef<IPropertyTypeCustomization> FActorAnimationPlaybackSettingsCustomization::MakeInstance()
+{
+	return MakeShareable( new FActorAnimationPlaybackSettingsCustomization );
+}
+
+void FActorAnimationPlaybackSettingsCustomization::CustomizeHeader( TSharedRef<IPropertyHandle> InPropertyHandle, FDetailWidgetRow& HeaderRow, IPropertyTypeCustomizationUtils& CustomizationUtils )
+{
+
+}
+
+void FActorAnimationPlaybackSettingsCustomization::CustomizeChildren( TSharedRef<IPropertyHandle> InPropertyHandle, IDetailChildrenBuilder& ChildBuilder, IPropertyTypeCustomizationUtils& CustomizationUtils )
+{
+	StructPropertyHandle = InPropertyHandle;
+
+	ChildBuilder.AddChildProperty(StructPropertyHandle->GetChildHandle(GET_MEMBER_NAME_CHECKED(FActorAnimationPlaybackSettings, PlayRate)).ToSharedRef());
+
+
+	LoopCountProperty = StructPropertyHandle->GetChildHandle(GET_MEMBER_NAME_CHECKED(FActorAnimationPlaybackSettings, LoopCount));
+
+	// Set up the initial environment
+	{
+		LoopModes.Add(MakeShareable( new FLoopMode{ LOCTEXT("DontLoop", "Don't Loop"), 0 } ));
+		LoopModes.Add(MakeShareable( new FLoopMode{ LOCTEXT("Indefinitely", "Loop Indefinitely"), -1 }) );
+		LoopModes.Add(MakeShareable( new FLoopMode{ LOCTEXT("Exactly", "Loop Exactly..."), 1 }) );
+
+		int32 CurrentValue = -1;
+		LoopCountProperty->GetValue(CurrentValue);
+		if (auto* FoundMode = LoopModes.FindByPredicate([&](const TSharedPtr<FLoopMode>& In){ return In->Value == CurrentValue; }) )
+		{
+			CurrentMode = *FoundMode;
+		}
+		else
+		{
+			CurrentMode = LoopModes.Last();
+		}
+	}
+
+	ChildBuilder.AddChildContent(LOCTEXT("LoopTitle", "Loop"))
+	.NameContent()
+	[
+		LoopCountProperty->CreatePropertyNameWidget()
+	]
+	.ValueContent()
+	.HAlign(HAlign_Fill)
+	.MaxDesiredWidth(200)
+	[
+		SNew(SHorizontalBox)
+
+		+ SHorizontalBox::Slot()
+		.VAlign(VAlign_Center)
+		.AutoWidth()
+		[
+			SNew(SComboBox<TSharedPtr<FLoopMode>>)
+			.OptionsSource(&LoopModes)
+			.OnSelectionChanged_Lambda([&](TSharedPtr<FLoopMode> Mode, ESelectInfo::Type){
+				CurrentMode = Mode;
+				UpdateProperty();
+			})
+			.OnGenerateWidget_Lambda([&](TSharedPtr<FLoopMode> InMode){
+				return SNew(STextBlock)
+				.Font(IDetailLayoutBuilder::GetDetailFont())
+				.Text(InMode->DisplayName);
+			})
+			.InitiallySelectedItem(CurrentMode)
+			[
+				SAssignNew(CurrentText, STextBlock)
+				.Font(IDetailLayoutBuilder::GetDetailFont())
+				.Text(CurrentMode->DisplayName)
+			]
+		]
+
+		+ SHorizontalBox::Slot()
+		.Padding(FMargin(4.f,0))
+		.VAlign(VAlign_Center)
+		[
+			SAssignNew(LoopEntry, SHorizontalBox)
+			.Visibility(CurrentMode == LoopModes.Last() ? EVisibility::Visible : EVisibility::Collapsed)
+
+			+ SHorizontalBox::Slot()
+			.VAlign(VAlign_Center)
+			.Padding(FMargin(0, 0, 4.f, 0))
+			[
+				LoopCountProperty->CreatePropertyValueWidget()
+			]
+
+			+ SHorizontalBox::Slot()
+			.VAlign(VAlign_Center)
+			.AutoWidth()
+			[
+				SNew(STextBlock)
+				.Font(IDetailLayoutBuilder::GetDetailFont())
+				.Text_Raw(this, &FActorAnimationPlaybackSettingsCustomization::GetCustomSuffix)
+			]
+		]
+	];
+}
+
+FText FActorAnimationPlaybackSettingsCustomization::GetCustomSuffix() const
+{
+	int32 NumLoops = 0;
+	LoopCountProperty->GetValue(NumLoops);
+	return NumLoops == 1 ? LOCTEXT("Time", "time") : LOCTEXT("Times", "times");
+}
+
+void FActorAnimationPlaybackSettingsCustomization::UpdateProperty()
+{
+	if (CurrentMode == LoopModes.Last())
+	{
+		// Show custom loop entry
+		LoopEntry->SetVisibility(EVisibility::Visible);
+	}
+	else
+	{
+		// Hide custom loop entry
+		LoopEntry->SetVisibility(EVisibility::Collapsed);
+	}
+
+	LoopCountProperty->SetValue(CurrentMode->Value);
+	CurrentText->SetText(CurrentMode->DisplayName);
+}
+
+#undef LOCTEXT_NAMESPACE
diff --git a/Engine/Source/Editor/DetailCustomizations/Private/ActorAnimationPlaybackSettingsCustomization.h b/Engine/Source/Editor/DetailCustomizations/Private/ActorAnimationPlaybackSettingsCustomization.h
new file mode 100644
--- /dev/null
+++ b/Engine/Source/Editor/DetailCustomizations/Private/ActorAnimationPlaybackSettingsCustomization.h
@@ -0,0 +1,43 @@
+// Copyright 1998-2015 Epic Games, Inc. All Rights Reserved.
+
+#pragma once
+
+class FActorAnimationPlaybackSettingsCustomization : public IPropertyTypeCustomization
+{
+public:
+	/** Makes a new instance of this detail layout class for a specific detail view requesting it */
+	static TSharedRef<IPropertyTypeCustomization> MakeInstance();
+
+	/** IDetailCustomization interface */
+	virtual void CustomizeHeader( TSharedRef<IPropertyHandle> PropertyHandle, FDetailWidgetRow& HeaderRow, IPropertyTypeCustomizationUtils& CustomizationUtils ) override;
+	virtual void CustomizeChildren( TSharedRef<IPropertyHandle> PropertyHandle, IDetailChildrenBuilder& ChildBuilder, IPropertyTypeCustomizationUtils& CustomizationUtils ) override;
+
+private:
+
+	/** Respond to a selection change event from the combo box */
+	void UpdateProperty();
+	
+private:
+
+	/** Get the suffix to display after the custom entry box */
+	FText GetCustomSuffix() const;
+
+	/** Array of loop modes */
+	struct FLoopMode
+	{
+		FText DisplayName;
+		int32 Value;
+	};
+	TArray<TSharedPtr<FLoopMode>> LoopModes;
+
+	/** The loop mode we're currently displaying */
+	TSharedPtr<FLoopMode> CurrentMode;
+	/** The text of the current selection */
+	TSharedPtr<STextBlock> CurrentText;
+	/** The loop number entry to be hidden and shown based on combo box selection */
+	TSharedPtr<SWidget> LoopEntry;
+
+	/** Property handles of the properties we're editing */
+	TSharedPtr<IPropertyHandle> StructPropertyHandle;
+	TSharedPtr<IPropertyHandle> LoopCountProperty;
+};
diff --git a/Engine/Source/Editor/DetailCustomizations/Private/DetailCustomizations.cpp b/Engine/Source/Editor/DetailCustomizations/Private/DetailCustomizations.cpp
--- a/Engine/Source/Editor/DetailCustomizations/Private/DetailCustomizations.cpp
+++ b/Engine/Source/Editor/DetailCustomizations/Private/DetailCustomizations.cpp
@@ -104,6 +104,7 @@
 #include "ConfigEditorPropertyDetails.h"
 #include "AssetImportDataCustomization.h"
 #include "CaptureResolutionCustomization.h"
+#include "ActorAnimationPlaybackSettingsCustomization.h"
 #include "TextCustomization.h"
 #include "AnimTrailNodeDetails.h"
 
@@ -212,6 +213,7 @@ void FDetailCustomizationsModule::RegisterPropertyTypeCustomizations()
 	RegisterCustomPropertyTypeLayout("HierarchicalSimplification", FOnGetPropertyTypeCustomizationInstance::CreateStatic(&FHierarchicalSimplificationCustomizations::MakeInstance));
 	RegisterCustomPropertyTypeLayout("PostProcessSettings", FOnGetPropertyTypeCustomizationInstance::CreateStatic(&FPostProcessSettingsCustomization::MakeInstance));
 	RegisterCustomPropertyTypeLayout("AssetImportInfo", FOnGetPropertyTypeCustomizationInstance::CreateStatic(&FAssetImportDataCustomization::MakeInstance));
+	RegisterCustomPropertyTypeLayout("ActorAnimationPlaybackSettings", FOnGetPropertyTypeCustomizationInstance::CreateStatic(&FActorAnimationPlaybackSettingsCustomization::MakeInstance));
 	RegisterCustomPropertyTypeLayout("CaptureResolution", FOnGetPropertyTypeCustomizationInstance::CreateStatic(&FCaptureResolutionCustomization::MakeInstance));
 	RegisterCustomPropertyTypeLayout("WeightedBlendable", FOnGetPropertyTypeCustomizationInstance::CreateStatic(&FWeightedBlendableCustomization::MakeInstance));
 }
diff --git a/Engine/Source/Runtime/ActorAnimation/Public/ActorAnimationPlayer.h b/Engine/Source/Runtime/ActorAnimation/Public/ActorAnimationPlayer.h
--- a/Engine/Source/Runtime/ActorAnimation/Public/ActorAnimationPlayer.h
+++ b/Engine/Source/Runtime/ActorAnimation/Public/ActorAnimationPlayer.h
@@ -23,11 +23,11 @@ struct FActorAnimationPlaybackSettings
 	GENERATED_BODY()
 
 	/** Number of times to loop playback. -1 for infinite, else the number of times to loop before stopping */
-	UPROPERTY(EditAnywhere, BlueprintReadWrite, Category="Playback", meta=(UIMin=-1))
+	UPROPERTY(EditAnywhere, BlueprintReadWrite, Category="Playback", meta=(UIMin=1, DisplayName="Loop"))
 	int32 LoopCount;
 
 	/** The rate at which to playback the animation */
-	UPROPERTY(EditAnywhere, BlueprintReadWrite, Category="Playback")
+	UPROPERTY(EditAnywhere, BlueprintReadWrite, Category="Playback", meta=(Units=Multiplier))
 	float PlayRate;
 };
 
