Fixing compile error for UFE

diff --git a/Engine/Source/Developer/Windows/WindowsTargetPlatform/Classes/WindowsTargetSettings.h b/Engine/Source/Developer/Windows/WindowsTargetPlatform/Classes/WindowsTargetSettings.h
--- a/Engine/Source/Developer/Windows/WindowsTargetPlatform/Classes/WindowsTargetSettings.h
+++ b/Engine/Source/Developer/Windows/WindowsTargetPlatform/Classes/WindowsTargetSettings.h
@@ -6,7 +6,6 @@
 
 #pragma once
 
-#include "AudioDevice.h"
 #include "WindowsTargetSettings.generated.h"
 
 UENUM()
diff --git a/Engine/Source/Editor/DetailCustomizations/Private/WindowsTargetSettingsDetails.cpp b/Engine/Source/Editor/DetailCustomizations/Private/WindowsTargetSettingsDetails.cpp
--- a/Engine/Source/Editor/DetailCustomizations/Private/WindowsTargetSettingsDetails.cpp
+++ b/Engine/Source/Editor/DetailCustomizations/Private/WindowsTargetSettingsDetails.cpp
@@ -7,7 +7,10 @@
 #include "IExternalImagePickerModule.h"
 #include "GameProjectGenerationModule.h"
 #include "ISourceControlModule.h"
+
+#if WITH_ENGINE
 #include "AudioDevice.h"
+#endif 
 
 #define LOCTEXT_NAMESPACE "WindowsTargetSettingsDetails"
 
@@ -358,6 +361,7 @@ bool FWindowsTargetSettingsDetails::IsValidAudioDeviceName(const FString& InDevi
 {
 	bool bIsValid = false;
 
+#if WITH_ENGINE
 	FAudioDevice* AudioDevice = GEngine->GetMainAudioDevice();
 	if (AudioDevice)
 	{
@@ -373,6 +377,7 @@ bool FWindowsTargetSettingsDetails::IsValidAudioDeviceName(const FString& InDevi
 			}
 		}
 	}
+#endif
 
 	return bIsValid;
 }
@@ -380,6 +385,8 @@ bool FWindowsTargetSettingsDetails::IsValidAudioDeviceName(const FString& InDevi
 TSharedRef<SWidget> FWindowsTargetSettingsDetails::MakeAudioDeviceMenu(const TSharedPtr<IPropertyHandle>& PropertyHandle)
 {
 	FMenuBuilder MenuBuilder(true, nullptr);
+
+#if WITH_ENGINE
 	FAudioDevice* AudioDevice = GEngine->GetMainAudioDevice();
 	if (AudioDevice)
 	{
@@ -402,6 +409,7 @@ TSharedRef<SWidget> FWindowsTargetSettingsDetails::MakeAudioDeviceMenu(const TSh
 		}
 		MenuBuilder.EndSection();
 	}
+#endif
 
 	return MenuBuilder.MakeWidget();
 }
