Copying //UE4/Dev-Sequencer to Dev-Main (//UE4/Dev-Main)

diff --git a/Engine/Source/Editor/DetailCustomizations/Private/CameraDetails.cpp b/Engine/Source/Editor/DetailCustomizations/Private/CameraDetails.cpp
--- a/Engine/Source/Editor/DetailCustomizations/Private/CameraDetails.cpp
+++ b/Engine/Source/Editor/DetailCustomizations/Private/CameraDetails.cpp
@@ -2,7 +2,10 @@
 
 #include "DetailCustomizationsPrivatePCH.h"
 #include "CameraDetails.h"
+#include "EditorCategoryUtils.h"
 #include "SNumericEntryBox.h"
+#include "EditorCategoryUtils.h"
+
 
 #define LOCTEXT_NAMESPACE "CameraDetails"
 
@@ -11,17 +14,18 @@ const float FCameraDetails::MaxAspectRatio = 100.0f;
 const float FCameraDetails::LowestCommonAspectRatio = 1.0f;
 const float FCameraDetails::HighestCommonAspectRatio = 2.5f;
 
+
 TSharedRef<IDetailCustomization> FCameraDetails::MakeInstance()
 {
 	return MakeShareable( new FCameraDetails );
 }
 
+
 BEGIN_SLATE_FUNCTION_BUILD_OPTIMIZATION
 void FCameraDetails::CustomizeDetails( IDetailLayoutBuilder& DetailLayout )
 {
 	FSlateFontInfo FontStyle = FEditorStyle::GetFontStyle(TEXT("PropertyWindow.NormalFont"));
 
-
 	LastParsedAspectRatioValue = -1.0f;
 
 	TSharedPtr<IPropertyHandle> bConstrainAspectRatioProperty = DetailLayout.GetProperty(GET_MEMBER_NAME_CHECKED(UCameraComponent, bConstrainAspectRatio));
@@ -29,58 +33,80 @@ void FCameraDetails::CustomizeDetails( IDetailLayoutBuilder& DetailLayout )
 	AspectRatioProperty = DetailLayout.GetProperty(GET_MEMBER_NAME_CHECKED(UCameraComponent, AspectRatio));
 	check(AspectRatioProperty.IsValid());
 
-	FSimpleDelegate OnAspectRatioChangedDelegate = FSimpleDelegate::CreateSP( this, &FCameraDetails::OnAspectRatioChanged );
-	AspectRatioProperty->SetOnPropertyValueChanged( OnAspectRatioChangedDelegate );
+	if (AspectRatioProperty.IsValid() && AspectRatioProperty->IsValidHandle())
+	{
+		FSimpleDelegate OnAspectRatioChangedDelegate = FSimpleDelegate::CreateSP(this, &FCameraDetails::OnAspectRatioChanged);
+		AspectRatioProperty->SetOnPropertyValueChanged(OnAspectRatioChangedDelegate);
+	}
+
+	// see if CameraSettings category should be hidden
+	bool bCameraSettingsHidden = false;
+	{
+		TArray< TWeakObjectPtr<UObject> > Objects;
+		DetailLayout.GetObjectsBeingCustomized(Objects);
 
-	IDetailCategoryBuilder& CameraCategory = DetailLayout.EditCategory( "CameraSettings", FText::GetEmpty(), ECategoryPriority::Important );
+		for (TWeakObjectPtr<UObject> ObjWP : Objects)
+		{
+			UObject* const Obj = ObjWP.Get();
+			if (Obj && (FEditorCategoryUtils::IsCategoryHiddenFromClass(Obj->GetClass(), TEXT("CameraSettings"))))
+			{
+				bCameraSettingsHidden = true;
+				break;
+			}
+		}
+	}
+
+	if (bCameraSettingsHidden == false)
+	{
+		IDetailCategoryBuilder& CameraCategory = DetailLayout.EditCategory( "CameraSettings", FText::GetEmpty(), ECategoryPriority::Important );
 	
-	// Organize the properties
-	CameraCategory.AddProperty(ProjectionModeProperty);
-
-	// Perspective-specific properties
-	IDetailPropertyRow& FieldOfViewRow = CameraCategory.AddProperty(DetailLayout.GetProperty(GET_MEMBER_NAME_CHECKED(UCameraComponent, FieldOfView)));
-	FieldOfViewRow.Visibility( TAttribute<EVisibility>::Create(TAttribute<EVisibility>::FGetter::CreateSP( this, &FCameraDetails::ProjectionModeMatches, ProjectionModeProperty, ECameraProjectionMode::Perspective ) ) );
-
-	// Orthographic-specific properties
-	TAttribute<EVisibility> OrthographicVisibility = TAttribute<EVisibility>::Create(TAttribute<EVisibility>::FGetter::CreateSP(this, &FCameraDetails::ProjectionModeMatches, ProjectionModeProperty, ECameraProjectionMode::Orthographic));
-
-	IDetailPropertyRow& OrthoWidthRow = CameraCategory.AddProperty(DetailLayout.GetProperty(GET_MEMBER_NAME_CHECKED(UCameraComponent, OrthoWidth)));
-	OrthoWidthRow.Visibility(OrthographicVisibility);
-
-	IDetailPropertyRow& OrthoNearClipPlaneRow = CameraCategory.AddProperty(DetailLayout.GetProperty(GET_MEMBER_NAME_CHECKED(UCameraComponent, OrthoNearClipPlane)));
-	OrthoNearClipPlaneRow.Visibility(OrthographicVisibility);
-
-	IDetailPropertyRow& OrthoFarClipPlaneRow = CameraCategory.AddProperty(DetailLayout.GetProperty(GET_MEMBER_NAME_CHECKED(UCameraComponent, OrthoFarClipPlane)));
-	OrthoFarClipPlaneRow.Visibility(OrthographicVisibility);
-
-	// Aspect ratio
-	CameraCategory.AddProperty( bConstrainAspectRatioProperty );
-	IDetailPropertyRow& AspectRatioRow = CameraCategory.AddProperty(AspectRatioProperty);
-
-	// Provide the special aspect ratio row
-	AspectRatioRow.CustomWidget()
-		.NameContent()
-		[
-			AspectRatioProperty->CreatePropertyNameWidget()
-		]
-		.ValueContent()
-		[
-			SNew(SHorizontalBox)
-			+SHorizontalBox::Slot()
-			.Padding(0.0f, 2.0f, 5.0f, 2.0f)
+		// Organize the properties
+		CameraCategory.AddProperty(ProjectionModeProperty);
+
+		// Perspective-specific properties
+		IDetailPropertyRow& FieldOfViewRow = CameraCategory.AddProperty(DetailLayout.GetProperty(GET_MEMBER_NAME_CHECKED(UCameraComponent, FieldOfView)));
+		FieldOfViewRow.Visibility( TAttribute<EVisibility>::Create(TAttribute<EVisibility>::FGetter::CreateSP( this, &FCameraDetails::ProjectionModeMatches, ProjectionModeProperty, ECameraProjectionMode::Perspective ) ) );
+
+		// Orthographic-specific properties
+		TAttribute<EVisibility> OrthographicVisibility = TAttribute<EVisibility>::Create(TAttribute<EVisibility>::FGetter::CreateSP(this, &FCameraDetails::ProjectionModeMatches, ProjectionModeProperty, ECameraProjectionMode::Orthographic));
+
+		IDetailPropertyRow& OrthoWidthRow = CameraCategory.AddProperty(DetailLayout.GetProperty(GET_MEMBER_NAME_CHECKED(UCameraComponent, OrthoWidth)));
+		OrthoWidthRow.Visibility(OrthographicVisibility);
+
+		IDetailPropertyRow& OrthoNearClipPlaneRow = CameraCategory.AddProperty(DetailLayout.GetProperty(GET_MEMBER_NAME_CHECKED(UCameraComponent, OrthoNearClipPlane)));
+		OrthoNearClipPlaneRow.Visibility(OrthographicVisibility);
+
+		IDetailPropertyRow& OrthoFarClipPlaneRow = CameraCategory.AddProperty(DetailLayout.GetProperty(GET_MEMBER_NAME_CHECKED(UCameraComponent, OrthoFarClipPlane)));
+		OrthoFarClipPlaneRow.Visibility(OrthographicVisibility);
+
+		// Aspect ratio
+		CameraCategory.AddProperty( bConstrainAspectRatioProperty );
+		IDetailPropertyRow& AspectRatioRow = CameraCategory.AddProperty(AspectRatioProperty);
+
+		// Provide the special aspect ratio row
+		AspectRatioRow.CustomWidget()
+			.NameContent()
 			[
-				SNew(SNumericEntryBox<float>)
-				.AllowSpin(true)
-				.Value(this, &FCameraDetails::GetAspectRatio)
-				.Font(FontStyle)
-				.MinValue(MinAspectRatio)
-				.MaxValue(MaxAspectRatio)
-				.MinSliderValue(LowestCommonAspectRatio)
-				.MaxSliderValue(HighestCommonAspectRatio)
-				.OnValueChanged(this, &FCameraDetails::OnAspectRatioSpinnerChanged)
-				.ToolTipText(LOCTEXT("AspectFloatTooltip", "Aspect Ratio (Width/Height)"))
+				AspectRatioProperty->CreatePropertyNameWidget()
 			]
-			+SHorizontalBox::Slot()
+			.ValueContent()
+			[
+				SNew(SHorizontalBox)
+				+SHorizontalBox::Slot()
+				.Padding(0.0f, 2.0f, 5.0f, 2.0f)
+				[
+					SNew(SNumericEntryBox<float>)
+					.AllowSpin(true)
+					.Value(this, &FCameraDetails::GetAspectRatio)
+					.Font(FontStyle)
+					.MinValue(MinAspectRatio)
+					.MaxValue(MaxAspectRatio)
+					.MinSliderValue(LowestCommonAspectRatio)
+					.MaxSliderValue(HighestCommonAspectRatio)
+					.OnValueChanged(this, &FCameraDetails::OnAspectRatioSpinnerChanged)
+					.ToolTipText(LOCTEXT("AspectFloatTooltip", "Aspect Ratio (Width/Height)"))
+				]
+				+SHorizontalBox::Slot()
 				[
 					SNew(SComboButton)
 					.OnGetMenuContent( this, &FCameraDetails::OnGetComboContent )
@@ -97,10 +123,11 @@ void FCameraDetails::CustomizeDetails( IDetailLayoutBuilder& DetailLayout )
 						.OnTextCommitted(this, &FCameraDetails::OnCommitAspectRatioText)
 					]
 				]
-		];
-
-	CameraCategory.AddProperty(DetailLayout.GetProperty(GET_MEMBER_NAME_CHECKED(UCameraComponent, bUsePawnControlRotation)));
-	CameraCategory.AddProperty(DetailLayout.GetProperty(GET_MEMBER_NAME_CHECKED(UCameraComponent, PostProcessBlendWeight)));
+			];
+		
+		CameraCategory.AddProperty(DetailLayout.GetProperty(GET_MEMBER_NAME_CHECKED(UCameraComponent, bUsePawnControlRotation)));
+		CameraCategory.AddProperty(DetailLayout.GetProperty(GET_MEMBER_NAME_CHECKED(UCameraComponent, PostProcessBlendWeight)));
+	}
 
 	UpdateAspectTextFromProperty();
 }
@@ -133,13 +160,16 @@ void FCameraDetails::OnAspectRatioSpinnerChanged(float InValue)
 void FCameraDetails::UpdateAspectTextFromProperty()
 {
 	// Called whenever the actual aspect ratio property changes - clears the text box if the value no longer matches the current text
-	TOptional<float> Value = GetAspectRatio();
-	if (!Value.IsSet() || Value.GetValue() < LastParsedAspectRatioValue - DELTA || Value.GetValue() > LastParsedAspectRatioValue + DELTA)
+	if (AspectTextBox.IsValid())
 	{
-		LastParsedAspectRatioValue = -1.0f;
-		if (!AspectTextBox->GetText().IsEmpty())
+		TOptional<float> Value = GetAspectRatio();
+		if (!Value.IsSet() || Value.GetValue() < LastParsedAspectRatioValue - DELTA || Value.GetValue() > LastParsedAspectRatioValue + DELTA)
 		{
-			AspectTextBox->SetText(FText::GetEmpty());
+			LastParsedAspectRatioValue = -1.0f;
+			if (!AspectTextBox->GetText().IsEmpty())
+			{
+				AspectTextBox->SetText(FText::GetEmpty());
+			}
 		}
 	}
 }
diff --git a/Engine/Source/Editor/DetailCustomizations/Private/DetailCustomizations.cpp b/Engine/Source/Editor/DetailCustomizations/Private/DetailCustomizations.cpp
--- a/Engine/Source/Editor/DetailCustomizations/Private/DetailCustomizations.cpp
+++ b/Engine/Source/Editor/DetailCustomizations/Private/DetailCustomizations.cpp
@@ -109,6 +109,8 @@
 
 #include "MaterialProxySettingsCustomizations.h"
 #include "ImportantToggleSettingCustomization.h"
+#include "CameraFilmbackSettingsCustomization.h"
+#include "CameraLensSettingsCustomization.h"
 
 IMPLEMENT_MODULE( FDetailCustomizationsModule, DetailCustomizations );
 
@@ -219,6 +221,8 @@ void FDetailCustomizationsModule::RegisterPropertyTypeCustomizations()
 	RegisterCustomPropertyTypeLayout("CompositionGraphCapturePasses", FOnGetPropertyTypeCustomizationInstance::CreateStatic(&FRenderPassesCustomization::MakeInstance));
 	RegisterCustomPropertyTypeLayout("WeightedBlendable", FOnGetPropertyTypeCustomizationInstance::CreateStatic(&FWeightedBlendableCustomization::MakeInstance));
 	RegisterCustomPropertyTypeLayout("MaterialProxySettings", FOnGetPropertyTypeCustomizationInstance::CreateStatic(&FMaterialProxySettingsCustomizations::MakeInstance));
+	RegisterCustomPropertyTypeLayout("CameraFilmbackSettings", FOnGetPropertyTypeCustomizationInstance::CreateStatic(&FCameraFilmbackSettingsCustomization::MakeInstance));
+	RegisterCustomPropertyTypeLayout("CameraLensSettings", FOnGetPropertyTypeCustomizationInstance::CreateStatic(&FCameraLensSettingsCustomization::MakeInstance));
 }
 
 
diff --git a/Engine/Source/Editor/DetailCustomizations/Private/MathStructCustomizations.cpp b/Engine/Source/Editor/DetailCustomizations/Private/MathStructCustomizations.cpp
--- a/Engine/Source/Editor/DetailCustomizations/Private/MathStructCustomizations.cpp
+++ b/Engine/Source/Editor/DetailCustomizations/Private/MathStructCustomizations.cpp
@@ -13,28 +13,30 @@
 
 TSharedRef<IPropertyTypeCustomization> FMathStructCustomization::MakeInstance() 
 {
-	return MakeShareable( new FMathStructCustomization );
+	return MakeShareable(new FMathStructCustomization);
 }
 
-void FMathStructCustomization::CustomizeHeader( TSharedRef<class IPropertyHandle> StructPropertyHandle, class FDetailWidgetRow& HeaderRow, IPropertyTypeCustomizationUtils& StructCustomizationUtils )
-{
-	GetSortedChildren( StructPropertyHandle, SortedChildHandles );
 
-	MakeHeaderRow( StructPropertyHandle, HeaderRow );
+void FMathStructCustomization::CustomizeHeader(TSharedRef<class IPropertyHandle> StructPropertyHandle, class FDetailWidgetRow& HeaderRow, IPropertyTypeCustomizationUtils& StructCustomizationUtils)
+{
+	GetSortedChildren(StructPropertyHandle, SortedChildHandles);
+	MakeHeaderRow(StructPropertyHandle, HeaderRow);
 }
 
-void FMathStructCustomization::CustomizeChildren( TSharedRef<class IPropertyHandle> StructPropertyHandle, class IDetailChildrenBuilder& StructBuilder, IPropertyTypeCustomizationUtils& StructCustomizationUtils )
+
+void FMathStructCustomization::CustomizeChildren(TSharedRef<class IPropertyHandle> StructPropertyHandle, class IDetailChildrenBuilder& StructBuilder, IPropertyTypeCustomizationUtils& StructCustomizationUtils)
 {
-	for( int32 ChildIndex = 0; ChildIndex < SortedChildHandles.Num(); ++ChildIndex )
+	for (int32 ChildIndex = 0; ChildIndex < SortedChildHandles.Num(); ++ChildIndex)
 	{
 		TSharedRef<IPropertyHandle> ChildHandle = SortedChildHandles[ChildIndex];
 
 		// Add the individual properties as children as well so the vector can be expanded for more room
-		StructBuilder.AddChildProperty( ChildHandle );
+		StructBuilder.AddChildProperty(ChildHandle);
 	}
 }
 
-void FMathStructCustomization::MakeHeaderRow( TSharedRef<class IPropertyHandle>& StructPropertyHandle, FDetailWidgetRow& Row )
+
+void FMathStructCustomization::MakeHeaderRow(TSharedRef<class IPropertyHandle>& StructPropertyHandle, FDetailWidgetRow& Row)
 {
 	// We'll set up reset to default ourselves
 	const bool bDisplayResetToDefault = false;
@@ -47,18 +49,18 @@ void FMathStructCustomization::MakeHeaderRow( TSharedRef<class IPropertyHandle>&
 
 	Row.NameContent()
 	[
-		StructPropertyHandle->CreatePropertyNameWidget( DisplayNameOverride, DisplayToolTipOverride, bDisplayResetToDefault )
+		StructPropertyHandle->CreatePropertyNameWidget(DisplayNameOverride, DisplayToolTipOverride, bDisplayResetToDefault)
 	]
 	.ValueContent()
 	// Make enough space for each child handle
-	.MinDesiredWidth(125.0f * SortedChildHandles.Num() )
-	.MaxDesiredWidth(125.0f * SortedChildHandles.Num() )
+	.MinDesiredWidth(125.0f * SortedChildHandles.Num())
+	.MaxDesiredWidth(125.0f * SortedChildHandles.Num())
 	[
-		SAssignNew( HorizontalBox, SHorizontalBox )
-		.IsEnabled( this, &FMathStructCustomization::IsValueEnabled, StructWeakHandlePtr )
+		SAssignNew(HorizontalBox, SHorizontalBox)
+		.IsEnabled(this, &FMathStructCustomization::IsValueEnabled, StructWeakHandlePtr)
 	];
 
-	for( int32 ChildIndex = 0; ChildIndex < SortedChildHandles.Num(); ++ChildIndex )
+	for (int32 ChildIndex = 0; ChildIndex < SortedChildHandles.Num(); ++ChildIndex)
 	{
 		TSharedRef<IPropertyHandle> ChildHandle = SortedChildHandles[ChildIndex];
 
@@ -66,72 +68,77 @@ void FMathStructCustomization::MakeHeaderRow( TSharedRef<class IPropertyHandle>&
 		// Make a widget for each property.  The vector component properties  will be displayed in the header
 
 		HorizontalBox->AddSlot()
-		.Padding( FMargin(0.0f, 2.0f, bLastChild ? 0.0f : 3.0f, 2.0f ) )
+		.Padding(FMargin(0.0f, 2.0f, bLastChild ? 0.0f : 3.0f, 2.0f))
 		[
 			MakeChildWidget(StructPropertyHandle, ChildHandle)
 		];
 	}
 
-	if ( StructPropertyHandle->GetProperty()->HasMetaData("AllowPreserveRatio") )
+	if (StructPropertyHandle->GetProperty()->HasMetaData("AllowPreserveRatio"))
 	{
-		if ( !GConfig->GetBool(TEXT("SelectionDetails"), *(StructPropertyHandle->GetProperty()->GetName() + TEXT("_PreserveScaleRatio")), bPreserveScaleRatio, GEditorPerProjectIni) )
+		if (!GConfig->GetBool(TEXT("SelectionDetails"), *(StructPropertyHandle->GetProperty()->GetName() + TEXT("_PreserveScaleRatio")), bPreserveScaleRatio, GEditorPerProjectIni))
 		{
 			bPreserveScaleRatio = true;
 		}
 
 		HorizontalBox->AddSlot()
 		.AutoWidth()
-		.MaxWidth( 18.0f )
+		.MaxWidth(18.0f)
 		[
 			// Add a checkbox to toggle between preserving the ratio of x,y,z components of scale when a value is entered
-			SNew( SCheckBox )
+			SNew(SCheckBox)
 			.IsChecked(this, &FMathStructCustomization::IsPreserveScaleRatioChecked)
 			.OnCheckStateChanged(this, &FMathStructCustomization::OnPreserveScaleRatioToggled, StructWeakHandlePtr)
-			.Style( FEditorStyle::Get(), "TransparentCheckBox" )
-			.ToolTipText( LOCTEXT("PreserveScaleToolTip", "When locked, scales uniformly based on the current xyz scale values so the object maintains its shape in each direction when scaled" ) )
+			.Style(FEditorStyle::Get(), "TransparentCheckBox")
+			.ToolTipText(LOCTEXT("PreserveScaleToolTip", "When locked, scales uniformly based on the current xyz scale values so the object maintains its shape in each direction when scaled"))
 			[
-				SNew( SImage )
+				SNew(SImage)
 				.Image(this, &FMathStructCustomization::GetPreserveScaleRatioImage)
-				.ColorAndOpacity( FSlateColor::UseForeground() )
+				.ColorAndOpacity(FSlateColor::UseForeground())
 			]
 		];
 	}
 }
 
+
 const FSlateBrush* FMathStructCustomization::GetPreserveScaleRatioImage() const
 {
 	return bPreserveScaleRatio ? FEditorStyle::GetBrush(TEXT("GenericLock")) : FEditorStyle::GetBrush(TEXT("GenericUnlock"));
 }
 
+
 ECheckBoxState FMathStructCustomization::IsPreserveScaleRatioChecked() const
 {
 	return bPreserveScaleRatio ? ECheckBoxState::Checked : ECheckBoxState::Unchecked;
 }
 
+
 void FMathStructCustomization::OnPreserveScaleRatioToggled(ECheckBoxState NewState, TWeakPtr<IPropertyHandle> PropertyHandle)
 {
-	bPreserveScaleRatio = ( NewState == ECheckBoxState::Checked ) ? true : false;
+	bPreserveScaleRatio = (NewState == ECheckBoxState::Checked) ? true : false;
 
-	if ( PropertyHandle.IsValid() && PropertyHandle.Pin()->GetProperty() )
+	if (PropertyHandle.IsValid() && PropertyHandle.Pin()->GetProperty())
 	{
-		FString SettingKey = ( PropertyHandle.Pin()->GetProperty()->GetName() + TEXT("_PreserveScaleRatio") );
+		FString SettingKey = (PropertyHandle.Pin()->GetProperty()->GetName() + TEXT("_PreserveScaleRatio"));
 		GConfig->SetBool(TEXT("SelectionDetails"), *SettingKey, bPreserveScaleRatio, GEditorPerProjectIni);
 	}
 }
 
-void FMathStructCustomization::GetSortedChildren( TSharedRef<IPropertyHandle> StructPropertyHandle, TArray< TSharedRef<IPropertyHandle> >& OutChildren )
+
+void FMathStructCustomization::GetSortedChildren(TSharedRef<IPropertyHandle> StructPropertyHandle, TArray< TSharedRef<IPropertyHandle> >& OutChildren)
 {
 	OutChildren.Empty();
 
 	uint32 NumChildren;
-	StructPropertyHandle->GetNumChildren( NumChildren );
+	StructPropertyHandle->GetNumChildren(NumChildren);
 
-	for( uint32 ChildIndex = 0; ChildIndex < NumChildren; ++ChildIndex )
+	for (uint32 ChildIndex = 0; ChildIndex < NumChildren; ++ChildIndex)
 	{
-		OutChildren.Add( StructPropertyHandle->GetChildHandle( ChildIndex ).ToSharedRef() );
+		OutChildren.Add(StructPropertyHandle->GetChildHandle(ChildIndex).ToSharedRef());
 	}
 }
 
+
 template<typename NumericType>
 void ExtractNumericMetadata(TSharedRef<IPropertyHandle>& PropertyHandle, TOptional<NumericType>& MinValue, TOptional<NumericType>& MaxValue, TOptional<NumericType>& SliderMinValue, TOptional<NumericType>& SliderMaxValue, NumericType& SliderExponent, NumericType& Delta)
 {
@@ -151,12 +158,12 @@ void ExtractNumericMetadata(TSharedRef<IPropertyHandle>& PropertyHandle, TOption
 	NumericType ClampMin = TNumericLimits<NumericType>::Lowest();
 	NumericType ClampMax = TNumericLimits<NumericType>::Max();
 
-	if ( !ClampMinString.IsEmpty() )
+	if (!ClampMinString.IsEmpty())
 	{
 		TTypeFromString<NumericType>::FromString(ClampMin, *ClampMinString);
 	}
 
-	if ( !ClampMaxString.IsEmpty() )
+	if (!ClampMaxString.IsEmpty())
 	{
 		TTypeFromString<NumericType>::FromString(ClampMax, *ClampMaxString);
 	}
@@ -167,18 +174,20 @@ void ExtractNumericMetadata(TSharedRef<IPropertyHandle>& PropertyHandle, TOption
 	TTypeFromString<NumericType>::FromString(UIMax, *UIMaxString);
 
 	SliderExponent = NumericType(1);
-	if ( SliderExponentString.Len() )
+
+	if (SliderExponentString.Len())
 	{
 		TTypeFromString<NumericType>::FromString(SliderExponent, *SliderExponentString);
 	}
 
 	Delta = NumericType(0);
-	if ( DeltaString.Len() )
+
+	if (DeltaString.Len())
 	{
 		TTypeFromString<NumericType>::FromString(Delta, *DeltaString);
 	}
 
-	if ( ClampMin >= ClampMax && ( ClampMinString.Len() || ClampMaxString.Len() ) )
+	if (ClampMin >= ClampMax && (ClampMinString.Len() || ClampMaxString.Len()))
 	{
 		//UE_LOG(LogPropertyNode, Warning, TEXT("Clamp Min (%s) >= Clamp Max (%s) for Ranged Numeric"), *ClampMinString, *ClampMaxString);
 	}
@@ -188,15 +197,16 @@ void ExtractNumericMetadata(TSharedRef<IPropertyHandle>& PropertyHandle, TOption
 
 	MinValue = ClampMinString.Len() ? ClampMin : TOptional<NumericType>();
 	MaxValue = ClampMaxString.Len() ? ClampMax : TOptional<NumericType>();
-	SliderMinValue = ( UIMinString.Len() ) ? ActualUIMin : TOptional<NumericType>();
-	SliderMaxValue = ( UIMaxString.Len() ) ? ActualUIMax : TOptional<NumericType>();
+	SliderMinValue = (UIMinString.Len()) ? ActualUIMin : TOptional<NumericType>();
+	SliderMaxValue = (UIMaxString.Len()) ? ActualUIMax : TOptional<NumericType>();
 
-	if ( ActualUIMin >= ActualUIMax && ( MetaUIMinString.Len() || MetaUIMaxString.Len() ) )
+	if (ActualUIMin >= ActualUIMax && (MetaUIMinString.Len() || MetaUIMaxString.Len()))
 	{
 		//UE_LOG(LogPropertyNode, Warning, TEXT("UI Min (%s) >= UI Max (%s) for Ranged Numeric"), *UIMinString, *UIMaxString);
 	}
 }
 
+
 template<typename NumericType>
 TSharedRef<SWidget> FMathStructCustomization::MakeNumericWidget(
 	TSharedRef<IPropertyHandle>& StructurePropertyHandle,
@@ -208,19 +218,18 @@ TSharedRef<SWidget> FMathStructCustomization::MakeNumericWidget(
 
 	TWeakPtr<IPropertyHandle> WeakHandlePtr = PropertyHandle;
 
-	return 
-		SNew( SNumericEntryBox<NumericType> )
-		.IsEnabled( this, &FMathStructCustomization::IsValueEnabled, WeakHandlePtr )
-		.Value( this, &FMathStructCustomization::OnGetValue, WeakHandlePtr )
-		.Font( IDetailLayoutBuilder::GetDetailFont() )
-		.UndeterminedString( NSLOCTEXT("PropertyEditor", "MultipleValues", "Multiple Values") )
-		.OnValueCommitted( this, &FMathStructCustomization::OnValueCommitted<NumericType>, WeakHandlePtr )
-		.OnValueChanged( this, &FMathStructCustomization::OnValueChanged<NumericType>, WeakHandlePtr )
-		.OnBeginSliderMovement( this, &FMathStructCustomization::OnBeginSliderMovement )
-		.OnEndSliderMovement( this, &FMathStructCustomization::OnEndSliderMovement<NumericType> )
+	return SNew(SNumericEntryBox<NumericType>)
+		.IsEnabled(this, &FMathStructCustomization::IsValueEnabled, WeakHandlePtr)
+		.Value(this, &FMathStructCustomization::OnGetValue, WeakHandlePtr)
+		.Font(IDetailLayoutBuilder::GetDetailFont())
+		.UndeterminedString(NSLOCTEXT("PropertyEditor", "MultipleValues", "Multiple Values"))
+		.OnValueCommitted(this, &FMathStructCustomization::OnValueCommitted<NumericType>, WeakHandlePtr)
+		.OnValueChanged(this, &FMathStructCustomization::OnValueChanged<NumericType>, WeakHandlePtr)
+		.OnBeginSliderMovement(this, &FMathStructCustomization::OnBeginSliderMovement)
+		.OnEndSliderMovement(this, &FMathStructCustomization::OnEndSliderMovement<NumericType>)
 		.LabelVAlign(VAlign_Center)
 		// Only allow spin on handles with one object.  Otherwise it is not clear what value to spin
-		.AllowSpin( PropertyHandle->GetNumOuterObjects() == 1 )
+		.AllowSpin(PropertyHandle->GetNumOuterObjects() == 1)
 		.MinValue(MinValue)
 		.MaxValue(MaxValue)
 		.MinSliderValue(SliderMinValue)
@@ -229,15 +238,16 @@ TSharedRef<SWidget> FMathStructCustomization::MakeNumericWidget(
 		.Delta(Delta)
 		.Label()
 		[
-			SNew( STextBlock )
-			.Font( IDetailLayoutBuilder::GetDetailFont() )
-			.Text( PropertyHandle->GetPropertyDisplayName() )
+			SNew(STextBlock)
+				.Font(IDetailLayoutBuilder::GetDetailFont())
+				.Text(PropertyHandle->GetPropertyDisplayName())
 		];
 }
 
-TSharedRef<SWidget> FMathStructCustomization::MakeChildWidget( 
+
+TSharedRef<SWidget> FMathStructCustomization::MakeChildWidget(
 	TSharedRef<IPropertyHandle>& StructurePropertyHandle,
-	TSharedRef<IPropertyHandle>& PropertyHandle )
+	TSharedRef<IPropertyHandle>& PropertyHandle)
 {
 	const UClass* PropertyClass = PropertyHandle->GetPropertyClass();
 	
@@ -260,71 +270,75 @@ TSharedRef<SWidget> FMathStructCustomization::MakeChildWidget(
 	return SNullWidget::NullWidget;
 }
 
+
 template<typename NumericType>
-TOptional<NumericType> FMathStructCustomization::OnGetValue( TWeakPtr<IPropertyHandle> WeakHandlePtr ) const
+TOptional<NumericType> FMathStructCustomization::OnGetValue(TWeakPtr<IPropertyHandle> WeakHandlePtr) const
 {
 	NumericType NumericVal = 0;
-	if( WeakHandlePtr.Pin()->GetValue( NumericVal ) == FPropertyAccess::Success )
+	if (WeakHandlePtr.Pin()->GetValue(NumericVal) == FPropertyAccess::Success)
 	{
-		return TOptional<NumericType>( NumericVal );
+		return TOptional<NumericType>(NumericVal);
 	}
 
 	// Value couldn't be accessed.  Return an unset value
 	return TOptional<NumericType>();
 }
 
+
 template<typename NumericType>
-void FMathStructCustomization::OnValueCommitted( NumericType NewValue, ETextCommit::Type CommitType, TWeakPtr<IPropertyHandle> WeakHandlePtr )
+void FMathStructCustomization::OnValueCommitted(NumericType NewValue, ETextCommit::Type CommitType, TWeakPtr<IPropertyHandle> WeakHandlePtr)
 {
 	EPropertyValueSetFlags::Type Flags = EPropertyValueSetFlags::DefaultFlags;
 	SetValue(NewValue, Flags, WeakHandlePtr);
-}	
+}
+
 
 template<typename NumericType>
-void FMathStructCustomization::OnValueChanged( NumericType NewValue, TWeakPtr<IPropertyHandle> WeakHandlePtr )
+void FMathStructCustomization::OnValueChanged(NumericType NewValue, TWeakPtr<IPropertyHandle> WeakHandlePtr)
 {
-	if( bIsUsingSlider )
+	if (bIsUsingSlider)
 	{
 		EPropertyValueSetFlags::Type Flags = EPropertyValueSetFlags::InteractiveChange;
 		SetValue(NewValue, Flags, WeakHandlePtr);
 	}
 }
 
+
 template<typename NumericType>
 void FMathStructCustomization::SetValue(NumericType NewValue, EPropertyValueSetFlags::Type Flags, TWeakPtr<IPropertyHandle> WeakHandlePtr)
 {
-	if ( bPreserveScaleRatio )
+	if (bPreserveScaleRatio)
 	{
 		// Get the value for each object for the modified component
 		TArray<FString> OldValues;
-		if ( WeakHandlePtr.Pin()->GetPerObjectValues(OldValues) == FPropertyAccess::Success )
+		if (WeakHandlePtr.Pin()->GetPerObjectValues(OldValues) == FPropertyAccess::Success)
 		{
 			// Loop through each object and scale based on the new ratio for each object individually
-			for ( int32 OutputIndex = 0; OutputIndex < OldValues.Num(); ++OutputIndex )
+			for (int32 OutputIndex = 0; OutputIndex < OldValues.Num(); ++OutputIndex)
 			{
 				NumericType OldValue;
 				TTypeFromString<NumericType>::FromString(OldValue, *OldValues[OutputIndex]);
 
 				// Account for the previous scale being zero.  Just set to the new value in that case?
 				NumericType Ratio = OldValue == 0 ? NewValue : NewValue / OldValue;
-				if ( Ratio == 0 )
+				if (Ratio == 0)
 				{
 					Ratio = NewValue;
 				}
 
 				// Loop through all the child handles (each component of the math struct, like X, Y, Z...etc)
-				for ( int32 ChildIndex = 0; ChildIndex < SortedChildHandles.Num(); ++ChildIndex )
+				for (int32 ChildIndex = 0; ChildIndex < SortedChildHandles.Num(); ++ChildIndex)
 				{
 					// Ignore scaling our selves.
 					TSharedRef<IPropertyHandle> ChildHandle = SortedChildHandles[ChildIndex];
-					if ( ChildHandle != WeakHandlePtr.Pin() )
+					if (ChildHandle != WeakHandlePtr.Pin())
 					{
 						// Get the value for each object.
 						TArray<FString> ObjectChildValues;
-						if ( ChildHandle->GetPerObjectValues(ObjectChildValues) == FPropertyAccess::Success )
+						if (ChildHandle->GetPerObjectValues(ObjectChildValues) == FPropertyAccess::Success)
 						{
 							// Individually scale each object's components by the same ratio.
-							for ( int32 ChildOutputIndex = 0; ChildOutputIndex < ObjectChildValues.Num(); ++ChildOutputIndex )
+							for (int32 ChildOutputIndex = 0; ChildOutputIndex < ObjectChildValues.Num(); ++ChildOutputIndex)
 							{
 								NumericType ChildOldValue;
 								TTypeFromString<NumericType>::FromString(ChildOldValue, *ObjectChildValues[ChildOutputIndex]);
@@ -344,9 +358,10 @@ void FMathStructCustomization::SetValue(NumericType NewValue, EPropertyValueSetF
 	WeakHandlePtr.Pin()->SetValue(NewValue, Flags);
 }
 
+
 bool FMathStructCustomization::IsValueEnabled(TWeakPtr<IPropertyHandle> WeakHandlePtr) const
 {
-	if ( WeakHandlePtr.IsValid() )
+	if (WeakHandlePtr.IsValid())
 	{
 		return !WeakHandlePtr.Pin()->IsEditConst();
 	}
@@ -354,27 +369,31 @@ bool FMathStructCustomization::IsValueEnabled(TWeakPtr<IPropertyHandle> WeakHand
 	return false;
 }
 
+
 void FMathStructCustomization::OnBeginSliderMovement()
 {
 	bIsUsingSlider = true;
 
-	GEditor->BeginTransaction( NSLOCTEXT("FMathStructCustomization", "SetVectorProperty", "Set Vector Property") );
+	GEditor->BeginTransaction(NSLOCTEXT("FMathStructCustomization", "SetVectorProperty", "Set Vector Property"));
 }
 
+
 template<typename NumericType>
-void FMathStructCustomization::OnEndSliderMovement( NumericType NewValue )
+void FMathStructCustomization::OnEndSliderMovement(NumericType NewValue)
 {
 	bIsUsingSlider = false;
 
 	GEditor->EndTransaction();
 }
 
+
 TSharedRef<IPropertyTypeCustomization> FRotatorStructCustomization::MakeInstance() 
 {
-	return MakeShareable( new FRotatorStructCustomization );
+	return MakeShareable(new FRotatorStructCustomization);
 }
 
-void FRotatorStructCustomization::GetSortedChildren( TSharedRef<IPropertyHandle> StructPropertyHandle, TArray< TSharedRef<IPropertyHandle> >& OutChildren )
+
+void FRotatorStructCustomization::GetSortedChildren(TSharedRef<IPropertyHandle> StructPropertyHandle, TArray< TSharedRef<IPropertyHandle> >& OutChildren)
 {
 	static const FName Roll("Roll");
 	static const FName Pitch("Pitch");
@@ -383,47 +402,48 @@ void FRotatorStructCustomization::GetSortedChildren( TSharedRef<IPropertyHandle>
 	TSharedPtr< IPropertyHandle > RotatorChildren[3];
 
 	uint32 NumChildren;
-	StructPropertyHandle->GetNumChildren( NumChildren );
+	StructPropertyHandle->GetNumChildren(NumChildren);
 
-	for( uint32 ChildIndex = 0; ChildIndex < NumChildren; ++ChildIndex )
+	for (uint32 ChildIndex = 0; ChildIndex < NumChildren; ++ChildIndex)
 	{
-		TSharedRef<IPropertyHandle> ChildHandle = StructPropertyHandle->GetChildHandle( ChildIndex ).ToSharedRef();
+		TSharedRef<IPropertyHandle> ChildHandle = StructPropertyHandle->GetChildHandle(ChildIndex).ToSharedRef();
 		const FName PropertyName = ChildHandle->GetProperty()->GetFName();
 
-		if( PropertyName == Roll )
+		if (PropertyName == Roll)
 		{
 			RotatorChildren[0] = ChildHandle;
 		}
-		else if( PropertyName == Pitch )
+		else if (PropertyName == Pitch)
 		{
 			RotatorChildren[1] = ChildHandle;
 		}
 		else
 		{
-			check( PropertyName == Yaw );
+			check(PropertyName == Yaw);
 			RotatorChildren[2] = ChildHandle;
 		}
 	}
 
-	OutChildren.Add( RotatorChildren[0].ToSharedRef() );
-	OutChildren.Add( RotatorChildren[1].ToSharedRef() );
-	OutChildren.Add( RotatorChildren[2].ToSharedRef() );
+	OutChildren.Add(RotatorChildren[0].ToSharedRef());
+	OutChildren.Add(RotatorChildren[1].ToSharedRef());
+	OutChildren.Add(RotatorChildren[2].ToSharedRef());
 }
 
 
 TSharedRef<IPropertyTypeCustomization> FColorStructCustomization::MakeInstance() 
 {
-	return MakeShareable( new FColorStructCustomization );
+	return MakeShareable(new FColorStructCustomization);
 }
 
-void FColorStructCustomization::CustomizeHeader( TSharedRef<class IPropertyHandle> InStructPropertyHandle, class FDetailWidgetRow& InHeaderRow, IPropertyTypeCustomizationUtils& StructCustomizationUtils )
+
+void FColorStructCustomization::CustomizeHeader(TSharedRef<class IPropertyHandle> InStructPropertyHandle, class FDetailWidgetRow& InHeaderRow, IPropertyTypeCustomizationUtils& StructCustomizationUtils)
 {
 	StructPropertyHandle = InStructPropertyHandle;
 
-	bIsLinearColor = CastChecked<UStructProperty>( StructPropertyHandle->GetProperty() )->Struct->GetFName() == NAME_LinearColor;
+	bIsLinearColor = CastChecked<UStructProperty>(StructPropertyHandle->GetProperty())->Struct->GetFName() == NAME_LinearColor;
 	bIgnoreAlpha = StructPropertyHandle->GetProperty()->HasMetaData(TEXT("HideAlphaChannel"));
 	
-	if ( StructPropertyHandle->GetProperty()->HasMetaData(TEXT("sRGB")) )
+	if (StructPropertyHandle->GetProperty()->HasMetaData(TEXT("sRGB")))
 	{
 		sRGBOverride = StructPropertyHandle->GetProperty()->GetBoolMetaData(TEXT("sRGB"));
 	}
@@ -431,52 +451,66 @@ void FColorStructCustomization::CustomizeHeader( TSharedRef<class IPropertyHandl
 	auto PropertyUtils = StructCustomizationUtils.GetPropertyUtilities();
 	bDontUpdateWhileEditing = PropertyUtils.IsValid() ? PropertyUtils->DontUpdateValueWhileEditing() : false;
 
-	FMathStructCustomization::CustomizeHeader( InStructPropertyHandle, InHeaderRow, StructCustomizationUtils );
+	FMathStructCustomization::CustomizeHeader(InStructPropertyHandle, InHeaderRow, StructCustomizationUtils);
 }
 
-void FColorStructCustomization::MakeHeaderRow( TSharedRef<class IPropertyHandle>& InStructPropertyHandle, FDetailWidgetRow& Row )
+
+void FColorStructCustomization::MakeHeaderRow(TSharedRef<class IPropertyHandle>& InStructPropertyHandle, FDetailWidgetRow& Row)
 {
 	// We'll set up reset to default ourselves
 	const bool bDisplayResetToDefault = false;
 	const FText DisplayNameOverride = FText::GetEmpty();
 	const FText DisplayToolTipOverride = FText::GetEmpty();
+	
+	TSharedPtr<SWidget> ColorWidget;
+	float ContentWidth = 250.0f;
+
+	if (InStructPropertyHandle->HasMetaData("InlineColorPicker"))
+	{
+		ColorWidget = CreateInlineColorPicker();
+		ContentWidth = 384.0f;
+	}
+	else
+	{
+		ColorWidget = CreateColorWidget();
+	}
 
 	Row.NameContent()
 	[
-		StructPropertyHandle->CreatePropertyNameWidget( DisplayNameOverride, DisplayToolTipOverride, bDisplayResetToDefault )
+		StructPropertyHandle->CreatePropertyNameWidget(DisplayNameOverride, DisplayToolTipOverride, bDisplayResetToDefault)
 	]
 	.ValueContent()
-	.MinDesiredWidth(250.0f)
-	.MaxDesiredWidth(250.0f)
+	.MinDesiredWidth(ContentWidth)
 	[
-		CreateColorWidget()
+		ColorWidget.ToSharedRef()
 	];
 }
 
+
 TSharedRef<SWidget> FColorStructCustomization::CreateColorWidget()
 {
 	FSlateFontInfo NormalText = IDetailLayoutBuilder::GetDetailFont();
 
-	return SNew( SHorizontalBox )
+	return SNew(SHorizontalBox)
 		+ SHorizontalBox::Slot()
 		.VAlign(VAlign_Center)
 		[
-			SNew( SOverlay )
+			SNew(SOverlay)
 			+SOverlay::Slot()
 			[
 				// Displays the color with alpha unless it is ignored
-				SAssignNew( ColorPickerParentWidget, SColorBlock )
-				.Color( this, &FColorStructCustomization::OnGetColorForColorBlock )
+				SAssignNew(ColorPickerParentWidget, SColorBlock)
+				.Color(this, &FColorStructCustomization::OnGetColorForColorBlock)
 				.ShowBackgroundForAlpha(true)
-				.IgnoreAlpha( bIgnoreAlpha )
-				.OnMouseButtonDown( this, &FColorStructCustomization::OnMouseButtonDownColorBlock )
-				.Size( FVector2D( 35.0f, 12.0f ) )
+				.IgnoreAlpha(bIgnoreAlpha)
+				.OnMouseButtonDown(this, &FColorStructCustomization::OnMouseButtonDownColorBlock)
+				.Size(FVector2D(35.0f, 12.0f))
 			]
 			+SOverlay::Slot()
 			.HAlign(HAlign_Center)
 			.VAlign(VAlign_Center)
 			[
-				SNew( STextBlock )
+				SNew(STextBlock)
 				.Text(NSLOCTEXT("PropertyEditor", "MultipleValues", "Multiple Values"))
 				.Font(NormalText)
 				.ColorAndOpacity(FSlateColor(FLinearColor::Black)) // we know the background is always white, so can safely set this to black
@@ -487,16 +521,17 @@ TSharedRef<SWidget> FColorStructCustomization::CreateColorWidget()
 		.VAlign(VAlign_Center)
 		[
 			// Displays the color without alpha
-			SNew( SColorBlock )
-			.Color( this, &FColorStructCustomization::OnGetColorForColorBlock )
+			SNew(SColorBlock)
+			.Color(this, &FColorStructCustomization::OnGetColorForColorBlock)
 			.ShowBackgroundForAlpha(false)
 			.IgnoreAlpha(true)
-			.OnMouseButtonDown( this, &FColorStructCustomization::OnMouseButtonDownColorBlock )
-			.Size( FVector2D( 35.0f, 12.0f ) )
+			.OnMouseButtonDown(this, &FColorStructCustomization::OnMouseButtonDownColorBlock)
+			.Size(FVector2D(35.0f, 12.0f))
 		];
 }
 
-void FColorStructCustomization::GetSortedChildren( TSharedRef<IPropertyHandle> InStructPropertyHandle, TArray< TSharedRef<IPropertyHandle> >& OutChildren )
+
+void FColorStructCustomization::GetSortedChildren(TSharedRef<IPropertyHandle> InStructPropertyHandle, TArray< TSharedRef<IPropertyHandle> >& OutChildren)
 {
 	static const FName Red("R");
 	static const FName Green("G");
@@ -507,22 +542,22 @@ void FColorStructCustomization::GetSortedChildren( TSharedRef<IPropertyHandle> I
 	TSharedPtr< IPropertyHandle > ColorProperties[4];
 
 	uint32 NumChildren;
-	InStructPropertyHandle->GetNumChildren( NumChildren );
+	InStructPropertyHandle->GetNumChildren(NumChildren);
 
-	for( uint32 ChildIndex = 0; ChildIndex < NumChildren; ++ChildIndex )
+	for (uint32 ChildIndex = 0; ChildIndex < NumChildren; ++ChildIndex)
 	{
-		TSharedRef<IPropertyHandle> ChildHandle = InStructPropertyHandle->GetChildHandle( ChildIndex ).ToSharedRef();
+		TSharedRef<IPropertyHandle> ChildHandle = InStructPropertyHandle->GetChildHandle(ChildIndex).ToSharedRef();
 		const FName PropertyName = ChildHandle->GetProperty()->GetFName();
 
-		if( PropertyName == Red )
+		if (PropertyName == Red)
 		{
 			ColorProperties[0] = ChildHandle;
 		}
-		else if( PropertyName == Green )
+		else if (PropertyName == Green)
 		{
 			ColorProperties[1] = ChildHandle;
 		}
-		else if( PropertyName == Blue )
+		else if (PropertyName == Blue)
 		{
 			ColorProperties[2] = ChildHandle;
 		}
@@ -532,110 +567,113 @@ void FColorStructCustomization::GetSortedChildren( TSharedRef<IPropertyHandle> I
 		}
 	}
 
-	OutChildren.Add( ColorProperties[0].ToSharedRef() );
-	OutChildren.Add( ColorProperties[1].ToSharedRef() );
-	OutChildren.Add( ColorProperties[2].ToSharedRef() );
+	OutChildren.Add(ColorProperties[0].ToSharedRef());
+	OutChildren.Add(ColorProperties[1].ToSharedRef());
+	OutChildren.Add(ColorProperties[2].ToSharedRef());
 
 	// Alpha channel may not be used
-	if( !bIgnoreAlpha && ColorProperties[3].IsValid() )
+	if (!bIgnoreAlpha && ColorProperties[3].IsValid())
 	{
-		OutChildren.Add( ColorProperties[3].ToSharedRef() );
+		OutChildren.Add(ColorProperties[3].ToSharedRef());
 	}
 }
 
-void FColorStructCustomization::CreateColorPicker( bool bUseAlpha, bool bOnlyRefreshOnOk )
+
+void FColorStructCustomization::CreateColorPicker(bool bUseAlpha)
 {
 	int32 NumObjects = StructPropertyHandle->GetNumOuterObjects();
 
 	SavedPreColorPickerColors.Empty();
 	TArray<FString> PerObjectValues;
-	StructPropertyHandle->GetPerObjectValues( PerObjectValues );
+	StructPropertyHandle->GetPerObjectValues(PerObjectValues);
 
-	for( int32 ObjectIndex = 0; ObjectIndex < NumObjects; ++ObjectIndex )
+	for (int32 ObjectIndex = 0; ObjectIndex < NumObjects; ++ObjectIndex)
 	{
-		if( bIsLinearColor )
+		if (bIsLinearColor)
 		{
 			FLinearColor Color;
-			Color.InitFromString( PerObjectValues[ObjectIndex] );
-			SavedPreColorPickerColors.Add( Color );	
+			Color.InitFromString(PerObjectValues[ObjectIndex]);
+			SavedPreColorPickerColors.Add(Color);	
 		}
 		else
 		{
 			FColor Color;
-			Color.InitFromString( PerObjectValues[ObjectIndex] );
-			SavedPreColorPickerColors.Add( Color.ReinterpretAsLinear() );
+			Color.InitFromString(PerObjectValues[ObjectIndex]);
+			SavedPreColorPickerColors.Add(Color.ReinterpretAsLinear());
 		}
 	}
 
 	FLinearColor InitialColor;
 	GetColorAsLinear(InitialColor);
 
-	// This needs to be meta data.  Other colors could benefit from this
-	const bool bRefreshOnlyOnOk = bOnlyRefreshOnOk || StructPropertyHandle->GetProperty()->GetOwnerClass()->IsChildOf(UMaterialExpressionConstant3Vector::StaticClass());
+	const bool bRefreshOnlyOnOk = bDontUpdateWhileEditing || StructPropertyHandle->HasMetaData("DontUpdateWhileEditing");
 
 	FColorPickerArgs PickerArgs;
-	PickerArgs.bUseAlpha = !bIgnoreAlpha;
-	PickerArgs.bOnlyRefreshOnMouseUp = false;
-	PickerArgs.bOnlyRefreshOnOk = bRefreshOnlyOnOk;
-	PickerArgs.sRGBOverride = sRGBOverride;
-	PickerArgs.DisplayGamma = TAttribute<float>::Create( TAttribute<float>::FGetter::CreateUObject(GEngine, &UEngine::GetDisplayGamma) );
-	PickerArgs.OnColorCommitted = FOnLinearColorValueChanged::CreateSP( this, &FColorStructCustomization::OnSetColorFromColorPicker );
-	PickerArgs.OnColorPickerCancelled = FOnColorPickerCancelled::CreateSP( this, &FColorStructCustomization::OnColorPickerCancelled );
-	PickerArgs.OnInteractivePickBegin = FSimpleDelegate::CreateSP( this, &FColorStructCustomization::OnColorPickerInteractiveBegin );
-	PickerArgs.OnInteractivePickEnd = FSimpleDelegate::CreateSP( this, &FColorStructCustomization::OnColorPickerInteractiveEnd );
-	PickerArgs.InitialColorOverride = InitialColor;
-	PickerArgs.ParentWidget = ColorPickerParentWidget;
+	{
+		PickerArgs.bUseAlpha = !bIgnoreAlpha;
+		PickerArgs.bOnlyRefreshOnMouseUp = false;
+		PickerArgs.bOnlyRefreshOnOk = bRefreshOnlyOnOk;
+		PickerArgs.sRGBOverride = sRGBOverride;
+		PickerArgs.DisplayGamma = TAttribute<float>::Create(TAttribute<float>::FGetter::CreateUObject(GEngine, &UEngine::GetDisplayGamma));
+		PickerArgs.OnColorCommitted = FOnLinearColorValueChanged::CreateSP(this, &FColorStructCustomization::OnSetColorFromColorPicker);
+		PickerArgs.OnColorPickerCancelled = FOnColorPickerCancelled::CreateSP(this, &FColorStructCustomization::OnColorPickerCancelled);
+		PickerArgs.OnInteractivePickBegin = FSimpleDelegate::CreateSP(this, &FColorStructCustomization::OnColorPickerInteractiveBegin);
+		PickerArgs.OnInteractivePickEnd = FSimpleDelegate::CreateSP(this, &FColorStructCustomization::OnColorPickerInteractiveEnd);
+		PickerArgs.InitialColorOverride = InitialColor;
+		PickerArgs.ParentWidget = ColorPickerParentWidget;
+	}
 
 	OpenColorPicker(PickerArgs);
 }
 
+
 TSharedRef<SColorPicker> FColorStructCustomization::CreateInlineColorPicker()
 {
 	int32 NumObjects = StructPropertyHandle->GetNumOuterObjects();
 
 	SavedPreColorPickerColors.Empty();
 	TArray<FString> PerObjectValues;
-	StructPropertyHandle->GetPerObjectValues( PerObjectValues );
+	StructPropertyHandle->GetPerObjectValues(PerObjectValues);
 
-	for( int32 ObjectIndex = 0; ObjectIndex < NumObjects; ++ObjectIndex )
+	for (int32 ObjectIndex = 0; ObjectIndex < NumObjects; ++ObjectIndex)
 	{
-		if( bIsLinearColor )
+		if (bIsLinearColor)
 		{
 			FLinearColor Color;
-			Color.InitFromString( PerObjectValues[ObjectIndex] );
-			SavedPreColorPickerColors.Add( Color );	
+			Color.InitFromString(PerObjectValues[ObjectIndex]);
+			SavedPreColorPickerColors.Add(Color);	
 		}
 		else
 		{
 			FColor Color;
-			Color.InitFromString( PerObjectValues[ObjectIndex] );
-			SavedPreColorPickerColors.Add( Color.ReinterpretAsLinear() );
+			Color.InitFromString(PerObjectValues[ObjectIndex]);
+			SavedPreColorPickerColors.Add(Color.ReinterpretAsLinear());
 		}
 	}
 
 	FLinearColor InitialColor;
 	GetColorAsLinear(InitialColor);
 
-	// This needs to be meta data.  Other colors could benefit from this
-	const bool bRefreshOnlyOnOk = bDontUpdateWhileEditing || StructPropertyHandle->GetProperty()->GetOwnerClass()->IsChildOf(UMaterialExpressionConstant3Vector::StaticClass());
+	const bool bRefreshOnlyOnOk = /*bDontUpdateWhileEditing ||*/ StructPropertyHandle->HasMetaData("DontUpdateWhileEditing");
 	
 	return SNew(SColorPicker)
-		.Visibility(this, &FColorStructCustomization::GetInlineColorPickerVisibility)
 		.DisplayInlineVersion(true)
 		.OnlyRefreshOnMouseUp(false)
 		.OnlyRefreshOnOk(bRefreshOnlyOnOk)
-		.DisplayGamma(TAttribute<float>::Create( TAttribute<float>::FGetter::CreateUObject(GEngine, &UEngine::GetDisplayGamma) ))
-		.OnColorCommitted(FOnLinearColorValueChanged::CreateSP( this, &FColorStructCustomization::OnSetColorFromColorPicker ))
-		.OnColorPickerCancelled(FOnColorPickerCancelled::CreateSP( this, &FColorStructCustomization::OnColorPickerCancelled ))
-		.OnInteractivePickBegin(FSimpleDelegate::CreateSP( this, &FColorStructCustomization::OnColorPickerInteractiveBegin ))
-		.OnInteractivePickEnd(FSimpleDelegate::CreateSP( this, &FColorStructCustomization::OnColorPickerInteractiveEnd ))
+		.DisplayGamma(TAttribute<float>::Create(TAttribute<float>::FGetter::CreateUObject(GEngine, &UEngine::GetDisplayGamma)))
+		.OnColorCommitted(FOnLinearColorValueChanged::CreateSP(this, &FColorStructCustomization::OnSetColorFromColorPicker))
+		.OnColorPickerCancelled(FOnColorPickerCancelled::CreateSP(this, &FColorStructCustomization::OnColorPickerCancelled))
+		.OnInteractivePickBegin(FSimpleDelegate::CreateSP(this, &FColorStructCustomization::OnColorPickerInteractiveBegin))
+		.OnInteractivePickEnd(FSimpleDelegate::CreateSP(this, &FColorStructCustomization::OnColorPickerInteractiveEnd))
+		.sRGBOverride(sRGBOverride)
 		.TargetColorAttribute(InitialColor);
 }
 
-void FColorStructCustomization::OnSetColorFromColorPicker( FLinearColor NewColor )
+
+void FColorStructCustomization::OnSetColorFromColorPicker(FLinearColor NewColor)
 {
 	FString ColorString;
-	if( bIsLinearColor )
+	if (bIsLinearColor)
 	{
 		ColorString = NewColor.ToString();
 	}
@@ -647,39 +685,42 @@ void FColorStructCustomization::OnSetColorFromColorPicker( FLinearColor NewColor
 		ColorString = NewFColor.ToString();
 	}
 
-	StructPropertyHandle->SetValueFromFormattedString( ColorString, bIsInteractive ? EPropertyValueSetFlags::InteractiveChange : 0 );
+	StructPropertyHandle->SetValueFromFormattedString(ColorString, bIsInteractive ? EPropertyValueSetFlags::InteractiveChange : 0);
+	StructPropertyHandle->NotifyFinishedChangingProperties();
 }
 
-void FColorStructCustomization::OnColorPickerCancelled( FLinearColor OriginalColor )
+
+void FColorStructCustomization::OnColorPickerCancelled(FLinearColor OriginalColor)
 {
 	TArray<FString> PerObjectColors;
 
-	for( int32 ColorIndex = 0; ColorIndex < SavedPreColorPickerColors.Num(); ++ColorIndex )
+	for (int32 ColorIndex = 0; ColorIndex < SavedPreColorPickerColors.Num(); ++ColorIndex)
 	{
-		if( bIsLinearColor )
+		if (bIsLinearColor)
 		{
-			PerObjectColors.Add( SavedPreColorPickerColors[ColorIndex].ToString() );
+			PerObjectColors.Add(SavedPreColorPickerColors[ColorIndex].ToString());
 		}
 		else
 		{
 			const bool bSRGB = false;
-			FColor Color = SavedPreColorPickerColors[ColorIndex].ToFColor( bSRGB );
-			PerObjectColors.Add( Color.ToString() );
+			FColor Color = SavedPreColorPickerColors[ColorIndex].ToFColor(bSRGB);
+			PerObjectColors.Add(Color.ToString());
 		}
 	}
 
-	StructPropertyHandle->SetPerObjectValues( PerObjectColors );
-
+	StructPropertyHandle->SetPerObjectValues(PerObjectColors);
 	PerObjectColors.Empty();
 }
 
+
 void FColorStructCustomization::OnColorPickerInteractiveBegin()
 {
 	bIsInteractive = true;
 
-	GEditor->BeginTransaction( FText::Format( NSLOCTEXT("FColorStructCustomization", "SetColorProperty", "Edit {0}"), StructPropertyHandle->GetPropertyDisplayName() ) );
+	GEditor->BeginTransaction(FText::Format(NSLOCTEXT("FColorStructCustomization", "SetColorProperty", "Edit {0}"), StructPropertyHandle->GetPropertyDisplayName()));
 }
 
+
 void FColorStructCustomization::OnColorPickerInteractiveEnd()
 {
 	bIsInteractive = false;
@@ -695,6 +736,7 @@ void FColorStructCustomization::OnColorPickerInteractiveEnd()
 	GEditor->EndTransaction();
 }
 
+
 FLinearColor FColorStructCustomization::OnGetColorForColorBlock() const
 {
 	FLinearColor Color;
@@ -702,17 +744,18 @@ FLinearColor FColorStructCustomization::OnGetColorForColorBlock() const
 	return Color;
 }
 
+
 FPropertyAccess::Result FColorStructCustomization::GetColorAsLinear(FLinearColor& OutColor) const
 {
 	// Default to full alpha in case the alpha component is disabled.
 	OutColor.A = 1.0f;
 
 	// Get each color component 
-	for(int32 ChildIndex = 0; ChildIndex < SortedChildHandles.Num(); ++ChildIndex)
+	for (int32 ChildIndex = 0; ChildIndex < SortedChildHandles.Num(); ++ChildIndex)
 	{
 		FPropertyAccess::Result ValueResult = FPropertyAccess::Fail;
 
-		if(bIsLinearColor)
+		if (bIsLinearColor)
 		{
 			float ComponentValue = 0;
 			ValueResult = SortedChildHandles[ChildIndex]->GetValue(ComponentValue);
@@ -749,6 +792,7 @@ FPropertyAccess::Result FColorStructCustomization::GetColorAsLinear(FLinearColor
 	return FPropertyAccess::Success;
 }
 
+
 EVisibility FColorStructCustomization::GetMultipleValuesTextVisibility() const
 {
 	FLinearColor Color;
@@ -756,32 +800,27 @@ EVisibility FColorStructCustomization::GetMultipleValuesTextVisibility() const
 	return (ValueResult == FPropertyAccess::MultipleValues) ? EVisibility::Visible : EVisibility::Collapsed;
 }
 
+
 FReply FColorStructCustomization::OnMouseButtonDownColorBlock(const FGeometry& MyGeometry, const FPointerEvent& MouseEvent)
 {
-	if ( MouseEvent.GetEffectingButton() != EKeys::LeftMouseButton )
+	if (MouseEvent.GetEffectingButton() != EKeys::LeftMouseButton)
 	{
 		return FReply::Unhandled();
 	}
 	
-	CreateColorPicker( /*bUseAlpha*/ true, bDontUpdateWhileEditing );
-
-	//bIsInlineColorPickerVisible = !bIsInlineColorPickerVisible;
+	CreateColorPicker(true /*bUseAlpha*/);
 
 	return FReply::Handled();
 }
 
-EVisibility FColorStructCustomization::GetInlineColorPickerVisibility() const
-{
-	return bIsInlineColorPickerVisible ? EVisibility::Visible : EVisibility::Collapsed;
-}
 
 FReply FColorStructCustomization::OnOpenFullColorPickerClicked()
 {
-	CreateColorPicker( /*bUseAlpha*/ true, bDontUpdateWhileEditing );
-
+	CreateColorPicker(true /*bUseAlpha*/);
 	bIsInlineColorPickerVisible = false;
 
 	return FReply::Handled();
 }
 
+
 #undef LOCTEXT_NAMESPACE
diff --git a/Engine/Source/Editor/DetailCustomizations/Private/MathStructCustomizations.h b/Engine/Source/Editor/DetailCustomizations/Private/MathStructCustomizations.h
--- a/Engine/Source/Editor/DetailCustomizations/Private/MathStructCustomizations.h
+++ b/Engine/Source/Editor/DetailCustomizations/Private/MathStructCustomizations.h
@@ -2,10 +2,12 @@
 
 #pragma once
 
+
 /**
  * Base class for math struct customization (e.g, vector, rotator, color)                                                              
  */
-class FMathStructCustomization : public IPropertyTypeCustomization
+class FMathStructCustomization
+	: public IPropertyTypeCustomization
 {
 public:
 	static TSharedRef<IPropertyTypeCustomization> MakeInstance();
@@ -18,17 +20,18 @@ public:
 	virtual ~FMathStructCustomization() {}
 
 	/** IPropertyTypeCustomization instance */
-	virtual void CustomizeHeader( TSharedRef<class IPropertyHandle> StructPropertyHandle, class FDetailWidgetRow& HeaderRow, IPropertyTypeCustomizationUtils& StructCustomizationUtils ) override;
-	virtual void CustomizeChildren( TSharedRef<class IPropertyHandle> StructPropertyHandle, class IDetailChildrenBuilder& StructBuilder, IPropertyTypeCustomizationUtils& StructCustomizationUtils ) override;
+	virtual void CustomizeHeader(TSharedRef<class IPropertyHandle> StructPropertyHandle, class FDetailWidgetRow& HeaderRow, IPropertyTypeCustomizationUtils& StructCustomizationUtils) override;
+	virtual void CustomizeChildren(TSharedRef<class IPropertyHandle> StructPropertyHandle, class IDetailChildrenBuilder& StructBuilder, IPropertyTypeCustomizationUtils& StructCustomizationUtils) override;
 	
 protected:
+
 	/**
 	 * Makes the header row for the customization
 	 *
 	 * @param StructPropertyHandle	Handle to the struct property
 	 * @param Row	The header row to add widgets to
 	 */
-	virtual void MakeHeaderRow( TSharedRef<class IPropertyHandle>& StructPropertyHandle, FDetailWidgetRow& Row );
+	virtual void MakeHeaderRow(TSharedRef<class IPropertyHandle>& StructPropertyHandle, FDetailWidgetRow& Row);
 
 	/**
 	 * Gets the sorted children for the struct
@@ -36,7 +39,7 @@ protected:
 	 * @param StructPropertyHandle	The handle to the struct property
 	 * @param OutChildren			The child array that should be populated in the order that children should be displayed
 	 */
-	virtual void GetSortedChildren( TSharedRef<IPropertyHandle> StructPropertyHandle, TArray< TSharedRef<IPropertyHandle> >& OutChildren );
+	virtual void GetSortedChildren(TSharedRef<IPropertyHandle> StructPropertyHandle, TArray< TSharedRef<IPropertyHandle> >& OutChildren);
 
 	/**
 	 * Constructs a widget for the property handle 
@@ -44,9 +47,9 @@ protected:
 	 * @param StructurePropertyHandle	handle of the struct property
 	 * @param PropertyHandle	Child handle of the struct property
 	 */
-	virtual TSharedRef<SWidget> MakeChildWidget( 
+	virtual TSharedRef<SWidget> MakeChildWidget(
 		TSharedRef<IPropertyHandle>& StructurePropertyHandle,
-		TSharedRef<IPropertyHandle>& PropertyHandle );
+		TSharedRef<IPropertyHandle>& PropertyHandle);
 
 	/**
 	 * Gets the value as a float for the provided property handle
@@ -55,7 +58,7 @@ protected:
 	 * @return The value or unset if it could not be accessed
 	 */
 	template<typename NumericType>
-	TOptional<NumericType> OnGetValue( TWeakPtr<IPropertyHandle> WeakHandlePtr ) const;
+	TOptional<NumericType> OnGetValue(TWeakPtr<IPropertyHandle> WeakHandlePtr) const;
 
 	/**
 	 * Called when the value is committed from the property editor
@@ -65,14 +68,14 @@ protected:
 	 * @param WeakHandlePtr	Handle to the property that the new value is for
 	 */
 	template<typename NumericType>
-	void OnValueCommitted( NumericType NewValue, ETextCommit::Type CommitType, TWeakPtr<IPropertyHandle> WeakHandlePtr );
+	void OnValueCommitted(NumericType NewValue, ETextCommit::Type CommitType, TWeakPtr<IPropertyHandle> WeakHandlePtr);
 	
 	/**
 	 * Called to see if the value is enabled for editing
 	 *
 	 * @param WeakHandlePtr	Handle to the property that the new value is for
 	 */
-	bool IsValueEnabled( TWeakPtr<IPropertyHandle> WeakHandlePtr ) const;
+	bool IsValueEnabled(TWeakPtr<IPropertyHandle> WeakHandlePtr) const;
 
 	/**
 	 * Called when the value is changed in the property editor
@@ -81,7 +84,7 @@ protected:
 	 * @param WeakHandlePtr	Handle to the property that the new value is for
 	 */
 	template<typename NumericType>
-	void OnValueChanged( NumericType NewValue, TWeakPtr<IPropertyHandle> WeakHandlePtr );
+	void OnValueChanged(NumericType NewValue, TWeakPtr<IPropertyHandle> WeakHandlePtr);
 
 	/**
 	 * Called to set the value of the property handle.
@@ -94,6 +97,7 @@ protected:
 	void SetValue(NumericType NewValue, EPropertyValueSetFlags::Type Flags, TWeakPtr<IPropertyHandle> WeakHandlePtr);
 
 private:
+
 	/** Gets the brush to use for the lock icon. */
 	const FSlateBrush* GetPreserveScaleRatioImage() const;
 
@@ -104,17 +108,19 @@ private:
 	void OnPreserveScaleRatioToggled(ECheckBoxState NewState, TWeakPtr<IPropertyHandle> PropertyHandle);
 
 private:
+
 	/** Called when a value starts to be changed by a slider */
 	void OnBeginSliderMovement();
 
 	/** Called when a value stops being changed by a slider */
 	template<typename NumericType>
-	void OnEndSliderMovement( NumericType NewValue );
+	void OnEndSliderMovement(NumericType NewValue);
 
 	template<typename NumericType>
 	TSharedRef<SWidget> MakeNumericWidget(TSharedRef<IPropertyHandle>& StructurePropertyHandle, TSharedRef<IPropertyHandle>& PropertyHandle);
 
 protected:
+
 	/** All the sorted children of the struct that should be displayed */
 	TArray< TSharedRef<IPropertyHandle> > SortedChildHandles;
 
@@ -125,39 +131,48 @@ protected:
 	bool bPreserveScaleRatio;
 };
 
+
 /**
  * Customizes FRotator structs                                                              
  */
-class FRotatorStructCustomization : public FMathStructCustomization
+class FRotatorStructCustomization
+	: public FMathStructCustomization
 {
 public:
+
 	/** @return A new instance of this class */
 	static TSharedRef<IPropertyTypeCustomization> MakeInstance();
+
 private:
+
 	/** FMathStructCustomization interface */
-	virtual void GetSortedChildren( TSharedRef<IPropertyHandle> StructPropertyHandle, TArray< TSharedRef<IPropertyHandle> >& OutChildren ) override;
+	virtual void GetSortedChildren(TSharedRef<IPropertyHandle> StructPropertyHandle, TArray< TSharedRef<IPropertyHandle> >& OutChildren) override;
 };
 
+
 /**
  * Base class for color struct customization (FColor,FLinearColor)                                                             
  */
-class FColorStructCustomization : public FMathStructCustomization
+class FColorStructCustomization
+	: public FMathStructCustomization
 {
 public:
+
 	static TSharedRef<IPropertyTypeCustomization> MakeInstance();
 	
 protected:
+
 	FColorStructCustomization()
-		: bIgnoreAlpha( false )
+		: bIgnoreAlpha(false)
 		, bIsInlineColorPickerVisible(false)
 		, bIsInteractive(false)
 		, bDontUpdateWhileEditing(false)
 	{}
 
 	/** FMathStructCustomization interface */
-	virtual void CustomizeHeader( TSharedRef<class IPropertyHandle> StructPropertyHandle, class FDetailWidgetRow& HeaderRow, IPropertyTypeCustomizationUtils& StructCustomizationUtils ) override;
-	virtual void MakeHeaderRow( TSharedRef<class IPropertyHandle>& InStructPropertyHandle, FDetailWidgetRow& Row ) override;
-	virtual void GetSortedChildren( TSharedRef<IPropertyHandle> StructPropertyHandle, TArray< TSharedRef<IPropertyHandle> >& OutChildren ) override;
+	virtual void CustomizeHeader(TSharedRef<class IPropertyHandle> StructPropertyHandle, class FDetailWidgetRow& HeaderRow, IPropertyTypeCustomizationUtils& StructCustomizationUtils) override;
+	virtual void MakeHeaderRow(TSharedRef<class IPropertyHandle>& InStructPropertyHandle, FDetailWidgetRow& Row) override;
+	virtual void GetSortedChildren(TSharedRef<IPropertyHandle> StructPropertyHandle, TArray< TSharedRef<IPropertyHandle> >& OutChildren) override;
 
 	/** Creates the color widget that when clicked spawns the color picker window. */
 	TSharedRef<SWidget> CreateColorWidget();
@@ -181,7 +196,7 @@ protected:
 	 * @param bUseAlpha			If true alpha will be displayed, otherwise it is ignored
 	 * @param bOnlyRefreshOnOk	If true the value of the property will only be refreshed when the user clicks OK in the color picker
 	 */
-	void CreateColorPicker( bool bUseAlpha, bool bOnlyRefreshOnOk );
+	void CreateColorPicker(bool bUseAlpha);
 	
 	/** Creates a new color picker for interactively selecting color */
 	TSharedRef<class SColorPicker> CreateInlineColorPicker();
@@ -191,7 +206,7 @@ protected:
 	 * 
 	 * @param NewColor	The new color value
 	 */
-	void OnSetColorFromColorPicker( FLinearColor NewColor );
+	void OnSetColorFromColorPicker(FLinearColor NewColor);
 	
 	/**
 	 * Called when the user clicks cancel in the color picker
@@ -199,7 +214,7 @@ protected:
 	 *
 	 * @param OriginalColor Original color of the property
 	 */
-	void OnColorPickerCancelled( FLinearColor OriginalColor );
+	void OnColorPickerCancelled(FLinearColor OriginalColor);
 
 	/**
 	 * Called when the user enters an interactive color change (dragging something in the picker)
@@ -221,30 +236,35 @@ protected:
 	 */
 	FReply OnMouseButtonDownColorBlock(const FGeometry& MyGeometry, const FPointerEvent& MouseEvent);
 	
-	/** Gets the visibility of the inline color picker */
-	EVisibility GetInlineColorPickerVisibility() const;
-
 	/** Called when the user clicks on the the button to get the full color picker */
 	FReply OnOpenFullColorPickerClicked();
 
 protected:
+
 	/** Saved per struct colors in case the user clicks cancel in the color picker */
 	TArray<FLinearColor> SavedPreColorPickerColors;
+
 	/** Color struct handle */
 	TSharedPtr<IPropertyHandle> StructPropertyHandle;
+
 	/** True if the property is a linear color property */
 	bool bIsLinearColor;
+
 	/** True if the property wants to ignore the alpha component */
 	bool bIgnoreAlpha;
+
 	/** True if the inline color picker is visible */
 	bool bIsInlineColorPickerVisible;
+
 	/** True if the user is performing an interactive color change */
 	bool bIsInteractive;
+
 	/** Cached widget for the color picker to use as a parent */
 	TSharedPtr<SWidget> ColorPickerParentWidget;
+
 	/** The value won;t be updated while editing */
 	bool bDontUpdateWhileEditing;
+
 	/** Overrides the default state of the sRGB checkbox */
 	TOptional<bool> sRGBOverride;
 };
-
diff --git a/Engine/Source/Editor/EditorStyle/Private/SlateEditorStyle.cpp b/Engine/Source/Editor/EditorStyle/Private/SlateEditorStyle.cpp
--- a/Engine/Source/Editor/EditorStyle/Private/SlateEditorStyle.cpp
+++ b/Engine/Source/Editor/EditorStyle/Private/SlateEditorStyle.cpp
@@ -1622,9 +1622,9 @@ void FSlateEditorStyle::FStyle::SetupGeneralStyles()
 		Set( "Sequencer.DefaultSectionGripLeft", new BOX_BRUSH( "Sequencer/SectionGripLeft", FMargin(5.f/16.f) ) );
 		Set( "Sequencer.DefaultSectionGripRight", new BOX_BRUSH( "Sequencer/SectionGripRight", FMargin(5.f/16.f) ) );
 		Set( "Sequencer.SectionArea.Background", new FSlateColorBrush( FColor::White ) );
-		Set( "Sequencer.AnimationOutliner.SelectionBorder", new BOX_BRUSH( "Common/Selection_16x", 4.0f/16.0f, SelectionColor_Pressed ) );
-		Set( "Sequencer.AnimationOutliner.SelectionBorderHover", new BOX_BRUSH("Common/Selection_16x", 4.0f / 16.0f, FLinearColor(1.0f, 1.0f, 1.0f, 0.1f) ) );
-		Set( "Sequencer.AnimationOutliner.SelectionBorderInactive", new BOX_BRUSH("Common/Selection_16x", 4.0f / 16.0f, SelectionColor_Inactive));
+		Set( "Sequencer.AnimationOutliner.TopLevelBorder_Expanded", new BOX_BRUSH( "Sequencer/TopLevelNodeBorder_Expanded", FMargin(4.0f/16.0f) ) );
+		Set( "Sequencer.AnimationOutliner.TopLevelBorder_Collapsed", new BOX_BRUSH( "Sequencer/TopLevelNodeBorder_Collapsed", FMargin(4.0f/16.0f) ) );
+		Set( "Sequencer.AnimationOutliner.DefaultBorder", new FSlateColorBrush( FLinearColor::White ) );
 		Set( "Sequencer.AnimationOutliner.BoldFont", TTF_CORE_FONT( "Fonts/Roboto-Bold", 11 ) );
 		Set( "Sequencer.AnimationOutliner.RegularFont", TTF_CORE_FONT( "Fonts/Roboto-Regular", 9 ) );
 		Set( "Sequencer.ShotFilter", new IMAGE_BRUSH( "Sequencer/FilteredArea", FVector2D(74,74), FLinearColor::White, ESlateBrushTileType::Both ) );
@@ -1652,7 +1652,10 @@ void FSlateEditorStyle::FStyle::SetupGeneralStyles()
 		Set( "Sequencer.UnlockCamera", new IMAGE_BRUSH( "Icons/icon_ToggleActorPilotCameraView_16x", Icon16x16, FLinearColor(1.f, 1.f, 1.f, 0.5f) ) );
 		Set( "Sequencer.Section.SelectionBorder", new BORDER_BRUSH( TEXT("Sequencer/SectionHighlight"), FMargin(7.f/16.f) ) );
 		Set( "Sequencer.Section.LockedBorder", new BORDER_BRUSH( TEXT("Sequencer/SectionLocked"), FMargin(7.f/16.f) ) );
-		Set( "Sequencer.ShotTrack.SectionHandle", new IMAGE_BRUSH( "Old/White", Icon16x16, FLinearColor::Black ) );
+		Set( "Sequencer.Thumbnail.SectionHandle", new IMAGE_BRUSH( "Old/White", Icon16x16, FLinearColor::Black ) );
+		Set( "Sequencer.TrackHoverHighlight_Top", new IMAGE_BRUSH( TEXT("Sequencer/TrackHoverHighlight_Top"), FVector2D(4, 4) ) );
+		Set( "Sequencer.TrackHoverHighlight_Bottom", new IMAGE_BRUSH( TEXT("Sequencer/TrackHoverHighlight_Bottom"), FVector2D(4, 4) ) );
+		Set( "Sequencer.SpawnableIconOverlay", new IMAGE_BRUSH( TEXT("Sequencer/SpawnableIconOverlay"), FVector2D(13, 13) ) );
 
 		Set( "Sequencer.General", new IMAGE_BRUSH( "Icons/SequencerIcons/icon_Sequencer_View_24x", Icon24x24 ) );
 		Set( "Sequencer.General.Small", new IMAGE_BRUSH( "Icons/SequencerIcons/icon_Sequencer_View_24x", Icon20x20 ) );
@@ -1666,14 +1669,22 @@ void FSlateEditorStyle::FStyle::SetupGeneralStyles()
 		Set( "Sequencer.Tracks.Audio", new IMAGE_BRUSH("Icons/AssetIcons/SoundCue_16x", Icon16x16));
 		Set( "Sequencer.Tracks.Event", new IMAGE_BRUSH("Icons/SequencerIcons/icon_Sequencer_Event_16x", Icon16x16));
 		Set( "Sequencer.Tracks.Fade", new IMAGE_BRUSH("Icons/SequencerIcons/icon_Sequencer_Fade_16x", Icon16x16));
-		Set( "Sequencer.Tracks.Shot", new IMAGE_BRUSH("Icons/SequencerIcons/icon_Sequencer_Shot_16x", Icon16x16));
+		Set( "Sequencer.Tracks.CameraCut", new IMAGE_BRUSH("Icons/SequencerIcons/icon_Sequencer_CameraCut_16x", Icon16x16));
+		Set( "Sequencer.Tracks.CinematicShot", new IMAGE_BRUSH("Icons/SequencerIcons/icon_Sequencer_CinematicShot_16x", Icon16x16));
 		Set( "Sequencer.Tracks.Slomo", new IMAGE_BRUSH("Icons/SequencerIcons/icon_Sequencer_Slomo_16x", Icon16x16));
 		Set( "Sequencer.Tracks.Sub", new IMAGE_BRUSH("Icons/SequencerIcons/icon_Sequencer_Sub_16x", Icon16x16));
 
-		Set( "Sequencer.CursorDecorator_Marquee", new IMAGE_BRUSH( "Sequencer/CursorDecorator_Marquee", Icon16x16));
 		Set( "Sequencer.CursorDecorator_MarqueeAdd", new IMAGE_BRUSH( "Sequencer/CursorDecorator_MarqueeAdd", Icon16x16));
 		Set( "Sequencer.CursorDecorator_MarqueeSubtract", new IMAGE_BRUSH( "Sequencer/CursorDecorator_MarqueeSubtract", Icon16x16));
 
+		Set( "Sequencer.BreadcrumbText", FTextBlockStyle(NormalText)
+			.SetFont( TTF_CORE_FONT( "Fonts/Roboto-Bold", 11 ) )
+			.SetColorAndOpacity( FLinearColor( 1.0f, 1.0f, 1.0f ) )
+			.SetHighlightColor( FLinearColor( 1.0f, 1.0f, 1.0f ) )
+			.SetShadowOffset( FVector2D( 1,1 ) )
+			.SetShadowColorAndOpacity( FLinearColor(0,0,0,0.9f) ) );
+		Set( "Sequencer.BreadcrumbIcon", new IMAGE_BRUSH( "Common/SmallArrowRight", Icon10x10 ) );
+
 		const FButtonStyle DetailsKeyButton = FButtonStyle(NoBorder)
 			.SetNormal( IMAGE_BRUSH("Sequencer/AddKey_Details", FVector2D(11,11) )  )
 			.SetHovered( IMAGE_BRUSH("Sequencer/AddKey_Details", FVector2D(11,11), SelectionColor ) )
@@ -1696,6 +1707,8 @@ void FSlateEditorStyle::FStyle::SetupGeneralStyles()
 			.SetForegroundColor(FSlateColor::UseSubduedForeground())
 			.SetArrowsImage(FSlateNoResource())
 		);
+
+		Set( "Sequencer.FixedFont", TTF_FONT( "Fonts/DroidSansMono", 9 ) );
 	}
 
 	// Foliage Edit Mode
@@ -4474,12 +4487,6 @@ void FSlateEditorStyle::FStyle::SetupLevelEditorStyle()
 			.SetCheckedPressedImage( BOX_BRUSH( "PlacementMode/TabActive", 8.0f / 32.0f ) )
 			.SetPadding( 0 ) );
 
-		Set( "PlacementBrowser.Tab.ImportantText", FTextBlockStyle( NormalText )
-			.SetFont( TTF_CORE_FONT( "Fonts/Roboto-Bold", 10 ) )
-			.SetColorAndOpacity( FLinearColor( 1.0f, 1.0f, 1.0f, 1.0f ) )
-			.SetShadowOffset( FVector2D( 1, 1 ) )
-			.SetShadowColorAndOpacity( FLinearColor( 0, 0, 0, 0.9f ) ) );
-
 		Set( "PlacementBrowser.Tab.Text", FTextBlockStyle( NormalText )
 			.SetFont( TTF_CORE_FONT( "Fonts/Roboto-Bold", 10 ) )
 			.SetColorAndOpacity( FLinearColor( 1.0f, 1.0f, 1.0f, 0.9f ) )
diff --git a/Engine/Source/Editor/LevelEditor/Private/LevelEditorActions.cpp b/Engine/Source/Editor/LevelEditor/Private/LevelEditorActions.cpp
--- a/Engine/Source/Editor/LevelEditor/Private/LevelEditorActions.cpp
+++ b/Engine/Source/Editor/LevelEditor/Private/LevelEditorActions.cpp
@@ -61,7 +61,6 @@
 #include "Engine/StaticMesh.h"
 #include "Engine/Light.h"
 #include "Animation/SkeletalMeshActor.h"
-#include "Editor/UnrealEd/Public/Animation/AnimationRecorder.h"
 #include "Editor/KismetWidgets/Public/CreateBlueprintFromActorDialog.h"
 #include "EditorProjectSettings.h"
 #include "HierarchicalLODUtilities.h"
@@ -362,7 +361,7 @@ void FLevelEditorActionCallbacks::Save()
 
 void FLevelEditorActionCallbacks::SaveAs()
 {
-	FEditorFileUtils::SaveAs( GetWorld()->PersistentLevel );
+	FEditorFileUtils::SaveLevelAs( GetWorld()->PersistentLevel );
 }
 
 void FLevelEditorActionCallbacks::SaveAllLevels()
@@ -1773,6 +1772,8 @@ bool FLevelEditorActionCallbacks::ScaleGridSnap_IsChecked()
 
 bool FLevelEditorActionCallbacks::SaveAnimationFromSkeletalMeshComponent(AActor * EditorActor, AActor * SimActor, TArray<class USkeletalMeshComponent*> & OutEditorComponents)
 {
+	FLevelEditorModule& LevelEditorModule = FModuleManager::GetModuleChecked<FLevelEditorModule>( TEXT("LevelEditor") );
+
 	// currently blueprint actors don't work because their property can't get copied over. 
 	if (Cast<UBlueprintGeneratedClass>(EditorActor->GetClass()) != nullptr)
 	{
@@ -1806,28 +1807,23 @@ bool FLevelEditorActionCallbacks::SaveAnimationFromSkeletalMeshComponent(AActor
 					if (Comp->SkeletalMesh && Comp->SkeletalMesh->Skeleton && Comp->IsSimulatingPhysics())
 					{
 						// now record to animation
-						FAnimationRecorder Recorder;
-						if(Recorder.TriggerRecordAnimation(Comp))
+						class UAnimSequence* Sequence = LevelEditorModule.OnCaptureSingleFrameAnimSequence().IsBound() ? LevelEditorModule.OnCaptureSingleFrameAnimSequence().Execute(Comp) : nullptr;
+						if(Sequence)
 						{
-							class UAnimSequence * Sequence = Recorder.GetAnimationObject();
-							if(Sequence)
+							Comp->SetAnimationMode(EAnimationMode::AnimationSingleNode);
+							Comp->AnimationData.AnimToPlay = Sequence;
+							Comp->SetAnimation(Sequence);
+							Comp->SetSimulatePhysics(false);
+
+							// add the matching component to EditorCompoennts
+							class USkeletalMeshComponent * MatchingComponent = Cast<USkeletalMeshComponent>(EditorUtilities::FindMatchingComponentInstance(Comp, EditorActor));
+							if (MatchingComponent)
+							{
+								OutEditorComponents.Add(MatchingComponent);
+							}
+							else
 							{
-								Recorder.StopRecord(false);
-								Comp->SetAnimationMode(EAnimationMode::AnimationSingleNode);
-								Comp->AnimationData.AnimToPlay = Sequence;
-								Comp->SetAnimation(Sequence);
-								Comp->SetSimulatePhysics(false);
-
-								// add the matching component to EditorCompoennts
-								class USkeletalMeshComponent * MatchingComponent = Cast<USkeletalMeshComponent>(EditorUtilities::FindMatchingComponentInstance(Comp, EditorActor));
-								if (MatchingComponent)
-								{
-									OutEditorComponents.Add(MatchingComponent);
-								}
-								else
-								{
-									UE_LOG( LevelEditorActions, Warning, TEXT("Matching component could not be found %s(%s)"), *GetNameSafe(Comp), *GetNameSafe(EditorActor) );
-								}
+								UE_LOG(LevelEditorActions, Warning, TEXT("Matching component could not be found %s(%s)"), *GetNameSafe(Comp), *GetNameSafe(EditorActor));
 							}
 						}
 					}
diff --git a/Engine/Source/Editor/LevelEditor/Private/LevelViewportActions.cpp b/Engine/Source/Editor/LevelEditor/Private/LevelViewportActions.cpp
--- a/Engine/Source/Editor/LevelEditor/Private/LevelViewportActions.cpp
+++ b/Engine/Source/Editor/LevelEditor/Private/LevelViewportActions.cpp
@@ -52,7 +52,7 @@ void FLevelViewportCommands::RegisterCommands()
 
 	UI_COMMAND( ApplyMaterialToActor, "Apply Material", "Attempts to apply a dropped material to this object", EUserInterfaceActionType::Button, FInputChord() );
 
-	UI_COMMAND( ToggleCinematicPreview, "Toggles Cinematic Preview", "If enabled, allows Matinee or Sequencer previews to play in this viewport", EUserInterfaceActionType::ToggleButton, FInputChord() );
+	UI_COMMAND( ToggleCinematicPreview, "Cinematic Preview", "If enabled, allows Matinee or Sequencer previews to play in this viewport", EUserInterfaceActionType::ToggleButton, FInputChord() );
 
 	UI_COMMAND( FindInLevelScriptBlueprint, "Find In Level Script", "Finds references of a selected actor in the level script blueprint", EUserInterfaceActionType::Button, FInputChord(EModifierKey::Control, EKeys::K) );
 	UI_COMMAND( AdvancedSettings, "Advanced Settings...", "Opens the advanced viewport settings", EUserInterfaceActionType::Button, FInputChord());
diff --git a/Engine/Source/Editor/LevelEditor/Private/LevelViewportLayout.cpp b/Engine/Source/Editor/LevelEditor/Private/LevelViewportLayout.cpp
--- a/Engine/Source/Editor/LevelEditor/Private/LevelViewportLayout.cpp
+++ b/Engine/Source/Editor/LevelEditor/Private/LevelViewportLayout.cpp
@@ -4,8 +4,8 @@
 #include "LevelEditor.h"
 #include "LevelViewportTabContent.h"
 #include "LevelViewportLayout.h"
-#include "SLevelViewport.h"
 #include "SDockTab.h"
+#include "SLevelViewport.h"
 
 namespace ViewportLayoutDefs
 {
@@ -20,12 +20,8 @@ namespace ViewportLayoutDefs
 
 	/** Default immersive state for new layouts - will only be applied when no config data is restoring state */
 	static const bool bDefaultShouldBeImmersive = false;
-
-	/** Default maximized viewport for new layouts - will only be applied when no config data is restoring state */
-	static const int32 DefaultMaximizedViewportID = 1;
 }
 
-
 // SViewportsOverlay ////////////////////////////////////////////////
 
 /**
@@ -133,6 +129,47 @@ TSharedPtr<FLevelViewportTabContent> SViewportsOverlay::GetLevelViewportTab() co
 	return LevelViewportTab;
 }
 
+FLevelViewportLayoutEntity::FLevelViewportLayoutEntity(TSharedRef<SLevelViewport> InLevelViewport)
+	: LevelViewport(InLevelViewport)
+{}
+
+TSharedRef<SWidget> FLevelViewportLayoutEntity::AsWidget() const
+{
+	return LevelViewport;
+}
+
+TSharedPtr<SLevelViewport> FLevelViewportLayoutEntity::AsLevelViewport() const
+{
+	return LevelViewport;
+}
+
+bool FLevelViewportLayoutEntity::IsPlayInEditorViewportActive() const
+{
+	return LevelViewport->IsPlayInEditorViewportActive();
+}
+
+void FLevelViewportLayoutEntity::RegisterGameViewportIfPIE()
+{
+	return LevelViewport->RegisterGameViewportIfPIE();
+}
+
+void FLevelViewportLayoutEntity::SetKeyboardFocus()
+{
+	return LevelViewport->SetKeyboardFocusToThisViewport();
+}
+
+void FLevelViewportLayoutEntity::OnLayoutDestroyed()
+{
+	if (LevelViewport->IsPlayInEditorViewportActive() || LevelViewport->GetLevelViewportClient().IsSimulateInEditorViewport() )
+	{
+		GUnrealEd->EndPlayMap();
+	}
+}
+
+void FLevelViewportLayoutEntity::SaveConfig(const FString& ConfigSection)
+{
+	LevelViewport->SaveConfig(ConfigSection);
+}
 
 // FLevelViewportLayout /////////////////////////////
 
@@ -154,14 +191,9 @@ FLevelViewportLayout::FLevelViewportLayout()
 
 FLevelViewportLayout::~FLevelViewportLayout()
 {
-	//catchall if one of our viewports is running by PIE or SIE, end it.
-	for (int32 i=0; i<Viewports.Num(); i++)
+	for (auto& Pair : Viewports)
 	{
-		if (Viewports[i]->IsPlayInEditorViewportActive() || Viewports[i]->GetLevelViewportClient().IsSimulateInEditorViewport() )
-		{
-			GUnrealEd->EndPlayMap();
-			break;
-		}
+		Pair.Value->OnLayoutDestroyed();
 	}
 
 	// Make sure that we're not locking the immersive window after we go away
@@ -184,7 +216,8 @@ TSharedRef<SWidget> FLevelViewportLayout::BuildViewportLayout( TSharedPtr<SDockT
 	ParentTab = InParentDockTab;
 	ParentTabContent = InParentTab;
 	ParentLevelEditor = InParentLevelEditor;
-
+	MaximizedViewport = NAME_None;
+	
 	// We use an overlay so that we can draw a maximized viewport on top of the other viewports
 	TSharedPtr<SBorder> ViewportsBorder;
 	TSharedRef<SViewportsOverlay> ViewportsOverlay =
@@ -231,7 +264,9 @@ void FLevelViewportLayout::InitCommonLayoutFromString( const FString& SpecificLa
 {
 	bool bShouldBeMaximized = bIsMaximizeSupported && ViewportLayoutDefs::bDefaultShouldBeMaximized;
 	bool bShouldBeImmersive = ViewportLayoutDefs::bDefaultShouldBeImmersive;
-	int32 MaximizedViewportID = ViewportLayoutDefs::DefaultMaximizedViewportID;
+	
+	FName SavedMaximizedViewport = NAME_None;
+	
 	if (!SpecificLayoutString.IsEmpty())
 	{
 		const FString& IniSection = FLayoutSaveRestore::GetAdditionalLayoutConfigIni();
@@ -239,17 +274,20 @@ void FLevelViewportLayout::InitCommonLayoutFromString( const FString& SpecificLa
 		// NOTE: We don't support starting back up in immersive mode, even if the user shut down with a window that way.  See the
 		// comment below in SaveCommonLayoutString() for more info.
 		GConfig->GetBool(*IniSection, *(SpecificLayoutString + TEXT(".bIsMaximized")), bShouldBeMaximized, GEditorPerProjectIni);
-		GConfig->GetInt(*IniSection, *(SpecificLayoutString + TEXT(".MaximizedViewportID")), MaximizedViewportID, GEditorPerProjectIni);
+		FString MaximizedViewportString;
+		if (GConfig->GetString(*IniSection, *(SpecificLayoutString + TEXT(".MaximizedViewport")), MaximizedViewportString, GEditorPerProjectIni))
+		{
+			SavedMaximizedViewport = *MaximizedViewportString;
+		}
 	}
 	// Replacement layouts (those selected by the user via a command) don't start maximized so the layout can be seen clearly.
-	if (!bIsReplacement && bIsMaximizeSupported && MaximizedViewportID >= 0 && MaximizedViewportID < Viewports.Num() && (bShouldBeMaximized || bShouldBeImmersive))
+	if (!bIsReplacement && bIsMaximizeSupported && Viewports.Contains(SavedMaximizedViewport) && (bShouldBeMaximized || bShouldBeImmersive))
 	{
-		TSharedRef<SLevelViewport> LevelViewport = Viewports[MaximizedViewportID].ToSharedRef();
 		// we are not toggling maximize or immersive state but setting it directly
 		const bool bToggle=false;
 		// Do not allow animation at startup as it hitches
 		const bool bAllowAnimation=false;
-		MaximizeViewport( LevelViewport, bShouldBeMaximized, bShouldBeImmersive, bAllowAnimation );
+		MaximizeViewport(SavedMaximizedViewport, bShouldBeMaximized, bShouldBeImmersive, bAllowAnimation);
 	}
 }
 
@@ -258,15 +296,9 @@ void FLevelViewportLayout::SaveCommonLayoutString( const FString& SpecificLayout
 	const FString& IniSection = FLayoutSaveRestore::GetAdditionalLayoutConfigIni();
 
 	// Save all our data using the additional layout config
-	int32 MaximizedViewportID = INDEX_NONE;
-	for (int32 i = 0; i < Viewports.Num(); ++i)
+	for (auto& Pair : Viewports)
 	{
-		Viewports[i]->SaveConfig(SpecificLayoutString + FString::Printf(TEXT(".Viewport%i"), i));
-		if (Viewports[i] == MaximizedViewport)
-		{
-			check(MaximizedViewportID == INDEX_NONE);
-			MaximizedViewportID = i;
-		}
+		Pair.Value->SaveConfig(SpecificLayoutString + Pair.Key.ToString());
 	}
 
 	// We don't bother saving that we were in immersive mode, because we never want to start back up directly in immersive mode
@@ -280,10 +312,10 @@ void FLevelViewportLayout::SaveCommonLayoutString( const FString& SpecificLayout
 	{
 		GConfig->SetBool(*IniSection, *(SpecificLayoutString + TEXT(".bIsMaximized")), bIsMaximizeSupported && bIsMaximized, GEditorPerProjectIni);
 	}
-	GConfig->SetInt( *IniSection, *( SpecificLayoutString + TEXT( ".MaximizedViewportID" ) ), MaximizedViewportID, GEditorPerProjectIni );
+	GConfig->SetString( *IniSection, *( SpecificLayoutString + TEXT( ".MaximizedViewport" ) ), *MaximizedViewport.ToString(), GEditorPerProjectIni );
 }
 
-void FLevelViewportLayout::RequestMaximizeViewport( TSharedRef<class SLevelViewport> ViewportToMaximize, const bool bWantMaximize, const bool bWantImmersive, const bool bAllowAnimation )
+void FLevelViewportLayout::RequestMaximizeViewport( FName ViewportToMaximize, const bool bWantMaximize, const bool bWantImmersive, const bool bAllowAnimation )
 {
 	if( bAllowAnimation )
 	{
@@ -303,12 +335,14 @@ void FLevelViewportLayout::RequestMaximizeViewport( TSharedRef<class SLevelViewp
 	}
 }
 
-void FLevelViewportLayout::MaximizeViewport( TSharedRef<SLevelViewport> ViewportToMaximize, const bool bWantMaximize, const bool bWantImmersive, const bool bAllowAnimation )
+void FLevelViewportLayout::MaximizeViewport( FName ViewportToMaximize, const bool bWantMaximize, const bool bWantImmersive, const bool bAllowAnimation )
 {
+	TSharedPtr<IViewportLayoutEntity> Entity = Viewports.FindRef(ViewportToMaximize);
+
 	// Should never get into a situation where the viewport is being maximized and there is already a maximized viewport. 
 	// I.E Maximized viewport is NULL which means this is a new maximize or MaximizeViewport is equal to the passed in one which means this is a restore of the current maximized viewport
-	check( Viewports.Contains( ViewportToMaximize ) );
-	check( !MaximizedViewport.IsValid() || MaximizedViewport == ViewportToMaximize );
+	check( Entity.IsValid() );
+	check( MaximizedViewport.IsNone() || MaximizedViewport == ViewportToMaximize );
 
 	// If we're already in immersive mode, toggling maximize just needs to update some state (no visual change)
 	if( bIsImmersive )
@@ -340,7 +374,7 @@ void FLevelViewportLayout::MaximizeViewport( TSharedRef<SLevelViewport> Viewport
 		else
 		{
 			// Viewport is still within the splitter, so use it for metrics directly
-			OwnerWindow = FSlateApplication::Get().FindWidgetWindow( ViewportToMaximize, ViewportWidgetPath );
+			OwnerWindow = FSlateApplication::Get().FindWidgetWindow( Entity->AsWidget(), ViewportWidgetPath );
 		}
 		bIsQueryingLayoutMetrics = false;
 
@@ -422,19 +456,23 @@ void FLevelViewportLayout::MaximizeViewport( TSharedRef<SLevelViewport> Viewport
 				// Store the maximized viewport
 				MaximizedViewport = ViewportToMaximize;
 
-				// Replace our viewport with a dummy widget in it's place during the maximize transition.  We can't
-				// have a single viewport widget in two places at once!
-				ReplaceWidget( MaximizedViewport.ToSharedRef(), ViewportReplacementWidget.ToSharedRef() );
-
-				TSharedRef< SCanvas > Canvas = SNew( SCanvas );
-				Canvas->AddSlot()
-					.Position( TAttribute< FVector2D >( this, &FLevelViewportLayout::GetMaximizedViewportPositionOnCanvas ) )
-					.Size( TAttribute< FVector2D >( this, &FLevelViewportLayout::GetMaximizedViewportSizeOnCanvas ) )
-					[
-						MaximizedViewport.ToSharedRef()
-					];
-
-				ViewportsOverlayWidget = Canvas;
+				TSharedPtr<IViewportLayoutEntity> MaximizedEntity = Viewports.FindRef(MaximizedViewport);
+				if (MaximizedEntity.IsValid())
+				{
+					// Replace our viewport with a dummy widget in it's place during the maximize transition.  We can't
+					// have a single viewport widget in two places at once!
+					ReplaceWidget( MaximizedEntity->AsWidget(), ViewportReplacementWidget.ToSharedRef() );
+
+					TSharedRef< SCanvas > Canvas = SNew( SCanvas );
+					Canvas->AddSlot()
+						.Position( TAttribute< FVector2D >( this, &FLevelViewportLayout::GetMaximizedViewportPositionOnCanvas ) )
+						.Size( TAttribute< FVector2D >( this, &FLevelViewportLayout::GetMaximizedViewportSizeOnCanvas ) )
+						[
+							MaximizedEntity->AsWidget()
+						];
+
+					ViewportsOverlayWidget = Canvas;
+				}
 			}
 
 
@@ -532,7 +570,7 @@ FVector2D FLevelViewportLayout::GetMaximizedViewportSizeOnCanvas() const
 	// NOTE: Should ALWAYS be valid, however because MaximizedViewport is changed in Tick, it's possible
 	//       for widgets we're adding/removing to already have been reported by ArrangeChildren, thus
 	//       we need to be able to handle cases where widgets that are not bound can still have delegates fire
-	if( MaximizedViewport.IsValid() || bWasImmersive )
+	if( !MaximizedViewport.IsNone() || bWasImmersive )
 	{
 		FVector2D TargetSize = FVector2D::ZeroVector;
 		if( bIsImmersive || ( bIsTransitioning && bWasImmersive ) )
@@ -578,26 +616,26 @@ bool FLevelViewportLayout::IsVisible()  const
  * @param InViewport	The viewport within this layout that should be checked
  * @return true if the viewport is visible.  
  */
-bool FLevelViewportLayout::IsLevelViewportVisible( const SLevelViewport& InViewport ) const
+bool FLevelViewportLayout::IsLevelViewportVisible( FName InViewport ) const
 {
 	// The passed in viewport is visible if the current layout is visible and their is no maximized viewport or the viewport that is maximized was passed in.
-	return IsVisible() && ( !MaximizedViewport.IsValid() || MaximizedViewport.Get() == &InViewport );
+	return IsVisible() && ( MaximizedViewport.IsNone() || MaximizedViewport == InViewport );
 }
 
-bool FLevelViewportLayout::IsViewportMaximized( const SLevelViewport& InViewport ) const
+bool FLevelViewportLayout::IsViewportMaximized( FName InViewport ) const
 {
-	return bIsMaximized && MaximizedViewport.Get() == &InViewport;
+	return bIsMaximized && MaximizedViewport == InViewport;
 }
 
-bool FLevelViewportLayout::IsViewportImmersive( const SLevelViewport& InViewport ) const
+bool FLevelViewportLayout::IsViewportImmersive( FName InViewport ) const
 {
-	return bIsImmersive && MaximizedViewport.Get() == &InViewport;
+	return bIsImmersive && MaximizedViewport == InViewport;
 }
 
 EVisibility FLevelViewportLayout::OnGetNonMaximizedVisibility() const
 {
 	// The non-maximized viewports are not visible if there is a maximized viewport on top of them
-	return ( !bIsQueryingLayoutMetrics && MaximizedViewport.IsValid() && !bIsTransitioning && DeferredMaximizeCommands.Num() == 0 ) ? EVisibility::Collapsed : EVisibility::Visible;
+	return ( !bIsQueryingLayoutMetrics && !MaximizedViewport.IsNone() && !bIsTransitioning && DeferredMaximizeCommands.Num() == 0 ) ? EVisibility::Collapsed : EVisibility::Visible;
 }
 
 
@@ -605,6 +643,9 @@ void FLevelViewportLayout::FinishMaximizeTransition()
 {
 	if( bIsTransitioning )
 	{
+		TSharedPtr<IViewportLayoutEntity> MaximizedViewportEntity = Viewports.FindRef(MaximizedViewport);
+		check(MaximizedViewportEntity.IsValid());
+
 		// The transition animation is complete, allow the engine to tick normally
 		EndThrottleForAnimatedResize();
 
@@ -620,12 +661,12 @@ void FLevelViewportLayout::FinishMaximizeTransition()
 			}
 
 			// Finished transition from restored/maximized to immersive, if this is a PIE window we need to re-register it to capture the mouse.
-			MaximizedViewport->RegisterGameViewportIfPIE();
+			MaximizedViewportEntity->RegisterGameViewportIfPIE();
 		}
 		else if( bIsMaximized && !bWasImmersive )
 		{
 			// Finished transition from restored to immersive, if this is a PIE window we need to re-register it to capture the mouse.
-			MaximizedViewport->RegisterGameViewportIfPIE();
+			MaximizedViewportEntity->RegisterGameViewportIfPIE();
 		}
 		else if( bWasImmersive )	// Finished transition from immersive to restored/maximized
 		{
@@ -666,16 +707,15 @@ void FLevelViewportLayout::FinishMaximizeTransition()
 		}
 
 		// Stop transitioning
-		TSharedRef< SLevelViewport > ViewportToFocus = MaximizedViewport.ToSharedRef();
 		if( !bIsImmersive && !bIsMaximized )
 		{
 			// We're finished with this temporary overlay widget now
 			ViewportsOverlayWidget.Reset();
 
 			// Restore the viewport widget into the viewport layout splitter
-			ReplaceWidget( ViewportReplacementWidget.ToSharedRef(), MaximizedViewport.ToSharedRef() );
+			ReplaceWidget( ViewportReplacementWidget.ToSharedRef(), MaximizedViewportEntity->AsWidget() );
 
-			MaximizedViewport.Reset();
+			MaximizedViewport = NAME_None;
 		}
 		bIsTransitioning = false;
 
@@ -687,12 +727,12 @@ void FLevelViewportLayout::FinishMaximizeTransition()
 			FSlateApplication::Get().ClearKeyboardFocus( EFocusCause::SetDirectly );
 
 			// Set keyboard focus directly
-			ViewportToFocus->SetKeyboardFocusToThisViewport();
+			MaximizedViewportEntity->SetKeyboardFocus();
 		}
 
 		// If this is a PIE window we need to re-register since the maximized window will have registered itself
 		// as the game viewport.
-		ViewportToFocus->RegisterGameViewportIfPIE();
+		MaximizedViewportEntity->RegisterGameViewportIfPIE();
 	}
 }
 
@@ -716,7 +756,7 @@ void FLevelViewportLayout::Tick( float DeltaTime )
 			FMaximizeViewportCommand& Command = DeferredMaximizeCommands[i];
 
 			// Only bother with deferred maximize if we don't already have a maximized or immersive viewport unless we are toggling
-			if( !MaximizedViewport.IsValid() || Command.bToggle )
+			if( MaximizedViewport.IsNone() || Command.bToggle )
 			{
 				MaximizeViewport(Command.Viewport, Command.bMaximize, Command.bImmersive, Command.bAllowAnimation );
 			}
diff --git a/Engine/Source/Editor/LevelEditor/Private/LevelViewportLayout2x2.cpp b/Engine/Source/Editor/LevelEditor/Private/LevelViewportLayout2x2.cpp
--- a/Engine/Source/Editor/LevelEditor/Private/LevelViewportLayout2x2.cpp
+++ b/Engine/Source/Editor/LevelEditor/Private/LevelViewportLayout2x2.cpp
@@ -86,10 +86,10 @@ TSharedRef<SWidget> FLevelViewportLayout2x2::MakeViewportLayout(const FString& L
 		.IsEnabled(FSlateApplication::Get().GetNormalExecutionAttribute())
 	];
 
-	Viewports.Add( ViewportWidgetTL );
-	Viewports.Add( ViewportWidgetBL );
-	Viewports.Add( ViewportWidgetTR );
-	Viewports.Add( ViewportWidgetBR );
+	Viewports.Add( *TopLeftKey, MakeShareable( new FLevelViewportLayoutEntity(ViewportWidgetTL.ToSharedRef()) ) );
+	Viewports.Add( *BottomLeftKey, MakeShareable( new FLevelViewportLayoutEntity(ViewportWidgetBL.ToSharedRef()) ) );
+	Viewports.Add( *TopRightKey, MakeShareable( new FLevelViewportLayoutEntity(ViewportWidgetTR.ToSharedRef()) ) );
+	Viewports.Add( *BottomRightKey, MakeShareable( new FLevelViewportLayoutEntity(ViewportWidgetBR.ToSharedRef()) ) );
 
 	
 	// Make newly-created perspective viewports active by default
diff --git a/Engine/Source/Editor/LevelEditor/Private/LevelViewportLayoutFourPanes.cpp b/Engine/Source/Editor/LevelEditor/Private/LevelViewportLayoutFourPanes.cpp
--- a/Engine/Source/Editor/LevelEditor/Private/LevelViewportLayoutFourPanes.cpp
+++ b/Engine/Source/Editor/LevelEditor/Private/LevelViewportLayoutFourPanes.cpp
@@ -120,7 +120,7 @@ void FLevelViewportLayoutFourPanes::ReplaceWidget( TSharedRef< SWidget > Source,
 // FLevelViewportLayoutFourPanesLeft /////////////////////////////
 
 TSharedRef<SWidget> FLevelViewportLayoutFourPanesLeft::MakeFourPanelWidget(
-	TArray<TSharedPtr< SLevelViewport >>& ViewportWidgets,
+	TMap<FName, TSharedPtr< IViewportLayoutEntity >>& ViewportWidgets,
 	const FString& ViewportKey0, const FString& ViewportKey1, const FString& ViewportKey2, const FString& ViewportKey3,
 	float PrimarySplitterPercentage, float SecondarySplitterPercentage0, float SecondarySplitterPercentage1)
 {
@@ -176,10 +176,10 @@ TSharedRef<SWidget> FLevelViewportLayoutFourPanesLeft::MakeFourPanelWidget(
 			]
 		];
 
-	ViewportWidgets.Add( ViewportWidget0 );
-	ViewportWidgets.Add( ViewportWidget1 );
-	ViewportWidgets.Add( ViewportWidget2 );
-	ViewportWidgets.Add( ViewportWidget3 );
+	ViewportWidgets.Add( *ViewportKey0, MakeShareable( new FLevelViewportLayoutEntity(ViewportWidget0.ToSharedRef()) ) );
+	ViewportWidgets.Add( *ViewportKey1, MakeShareable( new FLevelViewportLayoutEntity(ViewportWidget1.ToSharedRef()) ) );
+	ViewportWidgets.Add( *ViewportKey2, MakeShareable( new FLevelViewportLayoutEntity(ViewportWidget2.ToSharedRef()) ) );
+	ViewportWidgets.Add( *ViewportKey3, MakeShareable( new FLevelViewportLayoutEntity(ViewportWidget3.ToSharedRef()) ) );
 
 	// Make newly-created perspective viewports active by default
 	GCurrentLevelEditingViewportClient = &ViewportWidget0->GetLevelViewportClient();
@@ -191,7 +191,7 @@ TSharedRef<SWidget> FLevelViewportLayoutFourPanesLeft::MakeFourPanelWidget(
 // FLevelViewportLayoutFourPanesRight /////////////////////////////
 
 TSharedRef<SWidget> FLevelViewportLayoutFourPanesRight::MakeFourPanelWidget(
-	TArray<TSharedPtr< SLevelViewport >>& ViewportWidgets,
+	TMap<FName, TSharedPtr< IViewportLayoutEntity >>& ViewportWidgets,
 	const FString& ViewportKey0, const FString& ViewportKey1, const FString& ViewportKey2, const FString& ViewportKey3,
 	float PrimarySplitterPercentage, float SecondarySplitterPercentage0, float SecondarySplitterPercentage1)
 {
@@ -247,10 +247,10 @@ TSharedRef<SWidget> FLevelViewportLayoutFourPanesRight::MakeFourPanelWidget(
 			.ConfigKey( ViewportKey0 )
 		];
 
-	ViewportWidgets.Add( ViewportWidget0 );
-	ViewportWidgets.Add( ViewportWidget1 );
-	ViewportWidgets.Add( ViewportWidget2 );
-	ViewportWidgets.Add( ViewportWidget3 );
+	ViewportWidgets.Add( *ViewportKey0, MakeShareable( new FLevelViewportLayoutEntity(ViewportWidget0.ToSharedRef()) ) );
+	ViewportWidgets.Add( *ViewportKey1, MakeShareable( new FLevelViewportLayoutEntity(ViewportWidget1.ToSharedRef()) ) );
+	ViewportWidgets.Add( *ViewportKey2, MakeShareable( new FLevelViewportLayoutEntity(ViewportWidget2.ToSharedRef()) ) );
+	ViewportWidgets.Add( *ViewportKey3, MakeShareable( new FLevelViewportLayoutEntity(ViewportWidget3.ToSharedRef()) ) );
 
 	// Make newly-created perspective viewports active by default
 	GCurrentLevelEditingViewportClient = &ViewportWidget0->GetLevelViewportClient();
@@ -262,7 +262,7 @@ TSharedRef<SWidget> FLevelViewportLayoutFourPanesRight::MakeFourPanelWidget(
 // FLevelViewportLayoutFourPanesTop /////////////////////////////
 
 TSharedRef<SWidget> FLevelViewportLayoutFourPanesTop::MakeFourPanelWidget(
-	TArray<TSharedPtr< SLevelViewport >>& ViewportWidgets,
+	TMap<FName, TSharedPtr< IViewportLayoutEntity >>& ViewportWidgets,
 	const FString& ViewportKey0, const FString& ViewportKey1, const FString& ViewportKey2, const FString& ViewportKey3,
 	float PrimarySplitterPercentage, float SecondarySplitterPercentage0, float SecondarySplitterPercentage1)
 {
@@ -318,10 +318,10 @@ TSharedRef<SWidget> FLevelViewportLayoutFourPanesTop::MakeFourPanelWidget(
 			]
 		];
 
-	ViewportWidgets.Add( ViewportWidget0 );
-	ViewportWidgets.Add( ViewportWidget1 );
-	ViewportWidgets.Add( ViewportWidget2 );
-	ViewportWidgets.Add( ViewportWidget3 );
+	ViewportWidgets.Add( *ViewportKey0, MakeShareable( new FLevelViewportLayoutEntity(ViewportWidget0.ToSharedRef()) ) );
+	ViewportWidgets.Add( *ViewportKey1, MakeShareable( new FLevelViewportLayoutEntity(ViewportWidget1.ToSharedRef()) ) );
+	ViewportWidgets.Add( *ViewportKey2, MakeShareable( new FLevelViewportLayoutEntity(ViewportWidget2.ToSharedRef()) ) );
+	ViewportWidgets.Add( *ViewportKey3, MakeShareable( new FLevelViewportLayoutEntity(ViewportWidget3.ToSharedRef()) ) );
 
 	// Make newly-created perspective viewports active by default
 	GCurrentLevelEditingViewportClient = &ViewportWidget0->GetLevelViewportClient();
@@ -333,7 +333,7 @@ TSharedRef<SWidget> FLevelViewportLayoutFourPanesTop::MakeFourPanelWidget(
 // FLevelViewportLayoutFourPanesBottom /////////////////////////////
 
 TSharedRef<SWidget> FLevelViewportLayoutFourPanesBottom::MakeFourPanelWidget(
-	TArray<TSharedPtr< SLevelViewport >>& ViewportWidgets,
+	TMap<FName, TSharedPtr< IViewportLayoutEntity >>& ViewportWidgets,
 	const FString& ViewportKey0, const FString& ViewportKey1, const FString& ViewportKey2, const FString& ViewportKey3,
 	float PrimarySplitterPercentage, float SecondarySplitterPercentage0, float SecondarySplitterPercentage1)
 {
@@ -389,10 +389,10 @@ TSharedRef<SWidget> FLevelViewportLayoutFourPanesBottom::MakeFourPanelWidget(
 			.ConfigKey( ViewportKey0 )
 		];
 
-	ViewportWidgets.Add( ViewportWidget0 );
-	ViewportWidgets.Add( ViewportWidget1 );
-	ViewportWidgets.Add( ViewportWidget2 );
-	ViewportWidgets.Add( ViewportWidget3 );
+	ViewportWidgets.Add( *ViewportKey0, MakeShareable( new FLevelViewportLayoutEntity(ViewportWidget0.ToSharedRef()) ) );
+	ViewportWidgets.Add( *ViewportKey1, MakeShareable( new FLevelViewportLayoutEntity(ViewportWidget1.ToSharedRef()) ) );
+	ViewportWidgets.Add( *ViewportKey2, MakeShareable( new FLevelViewportLayoutEntity(ViewportWidget2.ToSharedRef()) ) );
+	ViewportWidgets.Add( *ViewportKey3, MakeShareable( new FLevelViewportLayoutEntity(ViewportWidget3.ToSharedRef()) ) );
 
 	// Make newly-created perspective viewports active by default
 	GCurrentLevelEditingViewportClient = &ViewportWidget0->GetLevelViewportClient();
diff --git a/Engine/Source/Editor/LevelEditor/Private/LevelViewportLayoutFourPanes.h b/Engine/Source/Editor/LevelEditor/Private/LevelViewportLayoutFourPanes.h
--- a/Engine/Source/Editor/LevelEditor/Private/LevelViewportLayoutFourPanes.h
+++ b/Engine/Source/Editor/LevelEditor/Private/LevelViewportLayoutFourPanes.h
@@ -22,7 +22,7 @@ protected:
 	virtual TSharedRef<SWidget> MakeViewportLayout(const FString& LayoutString) override;
 
 	virtual TSharedRef<SWidget> MakeFourPanelWidget(
-		TArray<TSharedPtr< SLevelViewport >>& ViewportWidgets,
+		TMap<FName, TSharedPtr< IViewportLayoutEntity >>& ViewportWidgets,
 		const FString& ViewportKey0, const FString& ViewportKey1, const FString& ViewportKey2, const FString& ViewportKey3,
 		float PrimarySplitterPercentage, float SecondarySplitterPercentage0, float SecondarySplitterPercentage1) = 0;
 
@@ -45,7 +45,7 @@ public:
 	virtual const FName& GetLayoutTypeName() const override { return LevelViewportConfigurationNames::FourPanesLeft; }
 
 	virtual TSharedRef<SWidget> MakeFourPanelWidget(
-		TArray<TSharedPtr< SLevelViewport >>& ViewportWidgets,
+		TMap<FName, TSharedPtr< IViewportLayoutEntity >>& ViewportWidgets,
 		const FString& ViewportKey0, const FString& ViewportKey1, const FString& ViewportKey2, const FString& ViewportKey3,
 		float PrimarySplitterPercentage, float SecondarySplitterPercentage0, float SecondarySplitterPercentage1) override;
 };
@@ -59,7 +59,7 @@ public:
 	virtual const FName& GetLayoutTypeName() const override { return LevelViewportConfigurationNames::FourPanesRight; }
 
 	virtual TSharedRef<SWidget> MakeFourPanelWidget(
-		TArray<TSharedPtr< SLevelViewport >>& ViewportWidgets,
+		TMap<FName, TSharedPtr< IViewportLayoutEntity >>& ViewportWidgets,
 		const FString& ViewportKey0, const FString& ViewportKey1, const FString& ViewportKey2, const FString& ViewportKey3,
 		float PrimarySplitterPercentage, float SecondarySplitterPercentage0, float SecondarySplitterPercentage1) override;
 };
@@ -73,7 +73,7 @@ public:
 	virtual const FName& GetLayoutTypeName() const override { return LevelViewportConfigurationNames::FourPanesTop; }
 
 	virtual TSharedRef<SWidget> MakeFourPanelWidget(
-		TArray<TSharedPtr< SLevelViewport >>& ViewportWidgets,
+		TMap<FName, TSharedPtr< IViewportLayoutEntity >>& ViewportWidgets,
 		const FString& ViewportKey0, const FString& ViewportKey1, const FString& ViewportKey2, const FString& ViewportKey3,
 		float PrimarySplitterPercentage, float SecondarySplitterPercentage0, float SecondarySplitterPercentage1) override;
 };
@@ -87,7 +87,7 @@ public:
 	virtual const FName& GetLayoutTypeName() const override { return LevelViewportConfigurationNames::FourPanesBottom; }
 
 	virtual TSharedRef<SWidget> MakeFourPanelWidget(
-		TArray<TSharedPtr< SLevelViewport >>& ViewportWidgets,
+		TMap<FName, TSharedPtr< IViewportLayoutEntity >>& ViewportWidgets,
 		const FString& ViewportKey0, const FString& ViewportKey1, const FString& ViewportKey2, const FString& ViewportKey3,
 		float PrimarySplitterPercentage, float SecondarySplitterPercentage0, float SecondarySplitterPercentage1) override;
 };
diff --git a/Engine/Source/Editor/LevelEditor/Private/LevelViewportLayoutOnePane.cpp b/Engine/Source/Editor/LevelEditor/Private/LevelViewportLayoutOnePane.cpp
--- a/Engine/Source/Editor/LevelEditor/Private/LevelViewportLayoutOnePane.cpp
+++ b/Engine/Source/Editor/LevelEditor/Private/LevelViewportLayoutOnePane.cpp
@@ -49,7 +49,7 @@ TSharedRef<SWidget> FLevelViewportLayoutOnePane::MakeViewportLayout(const FStrin
 			.ConfigKey(ViewportKey)
 		];
 
-	Viewports.Add(ViewportWidget);
+	Viewports.Add( *ViewportKey, MakeShareable(new FLevelViewportLayoutEntity(ViewportWidget.ToSharedRef())));
 
 	// Make newly-created perspective viewports active by default
 	GCurrentLevelEditingViewportClient = &ViewportWidget->GetLevelViewportClient();
