Copying //UE4/Dev-Sequencer to //UE4/Main

diff --git a/Engine/Source/Editor/CurveAssetEditor/Private/CurveAssetEditor.cpp b/Engine/Source/Editor/CurveAssetEditor/Private/CurveAssetEditor.cpp
--- a/Engine/Source/Editor/CurveAssetEditor/Private/CurveAssetEditor.cpp
+++ b/Engine/Source/Editor/CurveAssetEditor/Private/CurveAssetEditor.cpp
@@ -179,7 +179,7 @@ TSharedPtr<FExtender> FCurveAssetEditor::GetToolbarExtender()
 {
 	struct Local
 	{
-		static void FillToolbar(FToolBarBuilder& ToolbarBuilder, TSharedRef<SWidget> InputSnapWidget, TSharedRef<SWidget> OutputSnapWidget)
+		static void FillToolbar(FToolBarBuilder& ToolbarBuilder, TSharedRef<SWidget> InputSnapWidget, TSharedRef<SWidget> OutputSnapWidget, FCurveAssetEditor* CurveAssetEditor)
 		{
 			ToolbarBuilder.BeginSection("Curve");
 			{
@@ -199,6 +199,14 @@ TSharedPtr<FExtender> FCurveAssetEditor::GetToolbarExtender()
 			}
 			ToolbarBuilder.EndSection();
 
+			ToolbarBuilder.AddComboButton(
+				FUIAction(),
+				FOnGetContent::CreateSP( CurveAssetEditor, &FCurveAssetEditor::MakeCurveEditorCurveOptionsMenu ),
+				LOCTEXT( "CurveEditorCurveOptions", "Curves Options" ),
+				LOCTEXT( "CurveEditorCurveOptionsToolTip", "Curve Options" ),
+				TAttribute<FSlateIcon>(),
+				true );
+
 			ToolbarBuilder.BeginSection("Snap");
 			{
 				ToolbarBuilder.AddToolBarButton(FRichCurveEditorCommands::Get().ToggleSnapping);
@@ -239,7 +247,7 @@ TSharedPtr<FExtender> FCurveAssetEditor::GetToolbarExtender()
 		"Asset",
 		EExtensionHook::After,
 		TrackWidget->GetCommands(),
-		FToolBarExtensionDelegate::CreateStatic(&Local::FillToolbar, InputSnapWidget, OutputSnapWidget)
+		FToolBarExtensionDelegate::CreateStatic(&Local::FillToolbar, InputSnapWidget, OutputSnapWidget, this)
 		);
 
 	return ToolbarExtender;
@@ -252,4 +260,53 @@ EOrientation FCurveAssetEditor::GetSnapLabelOrientation() const
 		: EOrientation::Orient_Vertical;
 }
 
+TSharedRef<SWidget> FCurveAssetEditor::MakeCurveEditorCurveOptionsMenu()
+{
+	struct FExtrapolationMenus
+	{
+		static void MakePreInfinityExtrapSubMenu(FMenuBuilder& MenuBuilder)
+		{
+			MenuBuilder.BeginSection( "Pre-Infinity Extrapolation", LOCTEXT( "CurveEditorMenuPreInfinityExtrapHeader", "Extrapolation" ) );
+			{
+				MenuBuilder.AddMenuEntry( FRichCurveEditorCommands::Get().SetPreInfinityExtrapCycle);
+				MenuBuilder.AddMenuEntry( FRichCurveEditorCommands::Get().SetPreInfinityExtrapCycleWithOffset);
+				MenuBuilder.AddMenuEntry( FRichCurveEditorCommands::Get().SetPreInfinityExtrapOscillate);
+				MenuBuilder.AddMenuEntry( FRichCurveEditorCommands::Get().SetPreInfinityExtrapLinear);
+				MenuBuilder.AddMenuEntry( FRichCurveEditorCommands::Get().SetPreInfinityExtrapConstant);
+			}
+			MenuBuilder.EndSection();
+		}
+
+		static void MakePostInfinityExtrapSubMenu(FMenuBuilder& MenuBuilder)
+		{
+			MenuBuilder.BeginSection( "Post-Infinity Extrapolation", LOCTEXT( "CurveEditorMenuPostInfinityExtrapHeader", "Extrapolation" ) );
+			{
+				MenuBuilder.AddMenuEntry( FRichCurveEditorCommands::Get().SetPostInfinityExtrapCycle);
+				MenuBuilder.AddMenuEntry( FRichCurveEditorCommands::Get().SetPostInfinityExtrapCycleWithOffset);
+				MenuBuilder.AddMenuEntry( FRichCurveEditorCommands::Get().SetPostInfinityExtrapOscillate);
+				MenuBuilder.AddMenuEntry( FRichCurveEditorCommands::Get().SetPostInfinityExtrapLinear);
+				MenuBuilder.AddMenuEntry( FRichCurveEditorCommands::Get().SetPostInfinityExtrapConstant);
+			}
+			MenuBuilder.EndSection();
+		}
+	};
+
+	FMenuBuilder MenuBuilder( true, TrackWidget->GetCommands());
+
+	MenuBuilder.AddMenuEntry( FRichCurveEditorCommands::Get().BakeCurve);
+	MenuBuilder.AddMenuEntry( FRichCurveEditorCommands::Get().ReduceCurve);
+
+	MenuBuilder.AddSubMenu(
+		LOCTEXT( "PreInfinitySubMenu", "Pre-Infinity" ),
+		LOCTEXT( "PreInfinitySubMenuToolTip", "Pre-Infinity Extrapolation" ),
+		FNewMenuDelegate::CreateStatic( &FExtrapolationMenus::MakePreInfinityExtrapSubMenu ) );
+
+	MenuBuilder.AddSubMenu(
+		LOCTEXT( "PostInfinitySubMenu", "Post-Infinity" ),
+		LOCTEXT( "PostInfinitySubMenuToolTip", "Post-Infinity Extrapolation" ),
+		FNewMenuDelegate::CreateStatic( &FExtrapolationMenus::MakePostInfinityExtrapSubMenu ) );
+	
+	return MenuBuilder.MakeWidget();
+}
+
 #undef LOCTEXT_NAMESPACE
\ No newline at end of file
diff --git a/Engine/Source/Editor/CurveAssetEditor/Private/CurveAssetEditor.h b/Engine/Source/Editor/CurveAssetEditor/Private/CurveAssetEditor.h
--- a/Engine/Source/Editor/CurveAssetEditor/Private/CurveAssetEditor.h
+++ b/Engine/Source/Editor/CurveAssetEditor/Private/CurveAssetEditor.h
@@ -35,6 +35,7 @@ public:
 	virtual FLinearColor GetWorldCentricTabColorScale() const override;
 
 	TSharedPtr<FExtender> GetToolbarExtender();
+	TSharedRef<SWidget> MakeCurveEditorCurveOptionsMenu();
 
 private:
 
diff --git a/Engine/Source/Editor/DetailCustomizations/DetailCustomizations.Build.cs b/Engine/Source/Editor/DetailCustomizations/DetailCustomizations.Build.cs
--- a/Engine/Source/Editor/DetailCustomizations/DetailCustomizations.Build.cs
+++ b/Engine/Source/Editor/DetailCustomizations/DetailCustomizations.Build.cs
@@ -47,6 +47,7 @@ public class DetailCustomizations : ModuleRules
                 "MeshUtilities",
 				"ConfigEditor",
                 "Persona",
+                "CinematicCamera",
 			}
 		);
 
diff --git a/Engine/Source/Editor/DetailCustomizations/Private/CameraFilmbackSettingsCustomization.cpp b/Engine/Source/Editor/DetailCustomizations/Private/CameraFilmbackSettingsCustomization.cpp
--- a/Engine/Source/Editor/DetailCustomizations/Private/CameraFilmbackSettingsCustomization.cpp
+++ b/Engine/Source/Editor/DetailCustomizations/Private/CameraFilmbackSettingsCustomization.cpp
@@ -3,7 +3,7 @@
 #include "DetailCustomizationsPrivatePCH.h"
 #include "CameraFilmbackSettingsCustomization.h"
 #include "ScopedTransaction.h"
-#include "Camera/CineCameraComponent.h"
+#include "CineCameraComponent.h"
 
 #define LOCTEXT_NAMESPACE "CameraFilmbackSettingsCustomization"
 
diff --git a/Engine/Source/Editor/DetailCustomizations/Private/CameraLensSettingsCustomization.cpp b/Engine/Source/Editor/DetailCustomizations/Private/CameraLensSettingsCustomization.cpp
--- a/Engine/Source/Editor/DetailCustomizations/Private/CameraLensSettingsCustomization.cpp
+++ b/Engine/Source/Editor/DetailCustomizations/Private/CameraLensSettingsCustomization.cpp
@@ -3,7 +3,7 @@
 #include "DetailCustomizationsPrivatePCH.h"
 #include "CameraLensSettingsCustomization.h"
 #include "ScopedTransaction.h"
-#include "Camera/CineCameraComponent.h"
+#include "CineCameraComponent.h"
 
 #define LOCTEXT_NAMESPACE "CameraLensSettingsCustomization"
 
@@ -68,10 +68,9 @@ void FCameraLensSettingsCustomization::CustomizeChildren(TSharedRef<IPropertyHan
 	{
 		TSharedRef<IPropertyHandle> ChildHandle = StructPropertyHandle->GetChildHandle(ChildIndex).ToSharedRef();
 		const FName PropertyName = ChildHandle->GetProperty()->GetFName();
-
 		PropertyHandles.Add(PropertyName, ChildHandle);
 	}
-
+	
 	// Retrieve special case properties
 	MinFocalLengthHandle = PropertyHandles.FindChecked(GET_MEMBER_NAME_CHECKED(FCameraLensSettings, MinFocalLength));
 	MaxFocalLengthHandle = PropertyHandles.FindChecked(GET_MEMBER_NAME_CHECKED(FCameraLensSettings, MaxFocalLength));
@@ -82,6 +81,12 @@ void FCameraLensSettingsCustomization::CustomizeChildren(TSharedRef<IPropertyHan
 
 	for (auto Iter(PropertyHandles.CreateConstIterator()); Iter; ++Iter)
 	{
+		if ((Iter.Value() == MinFocusDistanceHandle) || (Iter.Value() == MaxReproRatioHandle))
+		{
+			// skip showing these in the panel for now, as we don't really use them
+			continue;
+		}
+
 		IDetailPropertyRow& SettingsRow = ChildBuilder.AddChildProperty(Iter.Value().ToSharedRef());
 	}
 }
diff --git a/Engine/Source/Editor/DetailCustomizations/Private/ColorStructCustomization.cpp b/Engine/Source/Editor/DetailCustomizations/Private/ColorStructCustomization.cpp
new file mode 100644
--- /dev/null
+++ b/Engine/Source/Editor/DetailCustomizations/Private/ColorStructCustomization.cpp
@@ -0,0 +1,405 @@
+// Copyright 1998-2015 Epic Games, Inc. All Rights Reserved.
+
+#include "DetailCustomizationsPrivatePCH.h"
+#include "ColorStructCustomization.h"
+#include "IPropertyUtilities.h"
+#include "SColorPicker.h"
+
+
+#define LOCTEXT_NAMESPACE "FColorStructCustomization"
+
+
+TSharedRef<IPropertyTypeCustomization> FColorStructCustomization::MakeInstance() 
+{
+	return MakeShareable(new FColorStructCustomization);
+}
+
+
+void FColorStructCustomization::CustomizeHeader(TSharedRef<class IPropertyHandle> InStructPropertyHandle, class FDetailWidgetRow& InHeaderRow, IPropertyTypeCustomizationUtils& StructCustomizationUtils)
+{
+	StructPropertyHandle = InStructPropertyHandle;
+
+	bIsLinearColor = CastChecked<UStructProperty>(StructPropertyHandle->GetProperty())->Struct->GetFName() == NAME_LinearColor;
+	bIgnoreAlpha = StructPropertyHandle->GetProperty()->HasMetaData(TEXT("HideAlphaChannel"));
+	
+	if (StructPropertyHandle->GetProperty()->HasMetaData(TEXT("sRGB")))
+	{
+		sRGBOverride = StructPropertyHandle->GetProperty()->GetBoolMetaData(TEXT("sRGB"));
+	}
+
+	auto PropertyUtils = StructCustomizationUtils.GetPropertyUtilities();
+	bDontUpdateWhileEditing = PropertyUtils.IsValid() ? PropertyUtils->DontUpdateValueWhileEditing() : false;
+
+	FMathStructCustomization::CustomizeHeader(InStructPropertyHandle, InHeaderRow, StructCustomizationUtils);
+}
+
+
+void FColorStructCustomization::MakeHeaderRow(TSharedRef<class IPropertyHandle>& InStructPropertyHandle, FDetailWidgetRow& Row)
+{
+	// We'll set up reset to default ourselves
+	const bool bDisplayResetToDefault = false;
+	const FText DisplayNameOverride = FText::GetEmpty();
+	const FText DisplayToolTipOverride = FText::GetEmpty();
+	
+	TSharedPtr<SWidget> ColorWidget;
+	float ContentWidth = 250.0f;
+
+	if (InStructPropertyHandle->HasMetaData("InlineColorPicker"))
+	{
+		ColorWidget = CreateInlineColorPicker();
+		ContentWidth = 384.0f;
+	}
+	else
+	{
+		ColorWidget = CreateColorWidget();
+	}
+
+	Row.NameContent()
+	[
+		StructPropertyHandle->CreatePropertyNameWidget(DisplayNameOverride, DisplayToolTipOverride, bDisplayResetToDefault)
+	]
+	.ValueContent()
+	.MinDesiredWidth(ContentWidth)
+	[
+		ColorWidget.ToSharedRef()
+	];
+}
+
+
+TSharedRef<SWidget> FColorStructCustomization::CreateColorWidget()
+{
+	FSlateFontInfo NormalText = IDetailLayoutBuilder::GetDetailFont();
+
+	return SNew(SHorizontalBox)
+		+ SHorizontalBox::Slot()
+		.VAlign(VAlign_Center)
+		[
+			SNew(SOverlay)
+			+SOverlay::Slot()
+			[
+				// Displays the color with alpha unless it is ignored
+				SAssignNew(ColorPickerParentWidget, SColorBlock)
+				.Color(this, &FColorStructCustomization::OnGetColorForColorBlock)
+				.ShowBackgroundForAlpha(true)
+				.IgnoreAlpha(bIgnoreAlpha)
+				.OnMouseButtonDown(this, &FColorStructCustomization::OnMouseButtonDownColorBlock)
+				.Size(FVector2D(35.0f, 12.0f))
+			]
+			+SOverlay::Slot()
+			.HAlign(HAlign_Center)
+			.VAlign(VAlign_Center)
+			[
+				SNew(STextBlock)
+				.Text(NSLOCTEXT("PropertyEditor", "MultipleValues", "Multiple Values"))
+				.Font(NormalText)
+				.ColorAndOpacity(FSlateColor(FLinearColor::Black)) // we know the background is always white, so can safely set this to black
+				.Visibility(this, &FColorStructCustomization::GetMultipleValuesTextVisibility)
+			]
+		]
+		+ SHorizontalBox::Slot()
+		.VAlign(VAlign_Center)
+		[
+			// Displays the color without alpha
+			SNew(SColorBlock)
+			.Color(this, &FColorStructCustomization::OnGetColorForColorBlock)
+			.ShowBackgroundForAlpha(false)
+			.IgnoreAlpha(true)
+			.OnMouseButtonDown(this, &FColorStructCustomization::OnMouseButtonDownColorBlock)
+			.Size(FVector2D(35.0f, 12.0f))
+		];
+}
+
+
+void FColorStructCustomization::GetSortedChildren(TSharedRef<IPropertyHandle> InStructPropertyHandle, TArray< TSharedRef<IPropertyHandle> >& OutChildren)
+{
+	static const FName Red("R");
+	static const FName Green("G");
+	static const FName Blue("B");
+	static const FName Alpha("A");
+
+	// We control the order of the colors via this array so it always ends up R,G,B,A
+	TSharedPtr< IPropertyHandle > ColorProperties[4];
+
+	uint32 NumChildren;
+	InStructPropertyHandle->GetNumChildren(NumChildren);
+
+	for (uint32 ChildIndex = 0; ChildIndex < NumChildren; ++ChildIndex)
+	{
+		TSharedRef<IPropertyHandle> ChildHandle = InStructPropertyHandle->GetChildHandle(ChildIndex).ToSharedRef();
+		const FName PropertyName = ChildHandle->GetProperty()->GetFName();
+
+		if (PropertyName == Red)
+		{
+			ColorProperties[0] = ChildHandle;
+		}
+		else if (PropertyName == Green)
+		{
+			ColorProperties[1] = ChildHandle;
+		}
+		else if (PropertyName == Blue)
+		{
+			ColorProperties[2] = ChildHandle;
+		}
+		else
+		{
+			ColorProperties[3] = ChildHandle;
+		}
+	}
+
+	OutChildren.Add(ColorProperties[0].ToSharedRef());
+	OutChildren.Add(ColorProperties[1].ToSharedRef());
+	OutChildren.Add(ColorProperties[2].ToSharedRef());
+
+	// Alpha channel may not be used
+	if (!bIgnoreAlpha && ColorProperties[3].IsValid())
+	{
+		OutChildren.Add(ColorProperties[3].ToSharedRef());
+	}
+}
+
+
+void FColorStructCustomization::CreateColorPicker(bool bUseAlpha)
+{
+	int32 NumObjects = StructPropertyHandle->GetNumOuterObjects();
+
+	SavedPreColorPickerColors.Empty();
+	TArray<FString> PerObjectValues;
+	StructPropertyHandle->GetPerObjectValues(PerObjectValues);
+
+	for (int32 ObjectIndex = 0; ObjectIndex < NumObjects; ++ObjectIndex)
+	{
+		if (bIsLinearColor)
+		{
+			FLinearColor Color;
+			Color.InitFromString(PerObjectValues[ObjectIndex]);
+			SavedPreColorPickerColors.Add(Color);	
+		}
+		else
+		{
+			FColor Color;
+			Color.InitFromString(PerObjectValues[ObjectIndex]);
+			SavedPreColorPickerColors.Add(Color.ReinterpretAsLinear());
+		}
+	}
+
+	FLinearColor InitialColor;
+	GetColorAsLinear(InitialColor);
+
+	const bool bRefreshOnlyOnOk = bDontUpdateWhileEditing || StructPropertyHandle->HasMetaData("DontUpdateWhileEditing");
+
+	FColorPickerArgs PickerArgs;
+	{
+		PickerArgs.bUseAlpha = !bIgnoreAlpha;
+		PickerArgs.bOnlyRefreshOnMouseUp = false;
+		PickerArgs.bOnlyRefreshOnOk = bRefreshOnlyOnOk;
+		PickerArgs.sRGBOverride = sRGBOverride;
+		PickerArgs.DisplayGamma = TAttribute<float>::Create(TAttribute<float>::FGetter::CreateUObject(GEngine, &UEngine::GetDisplayGamma));
+		PickerArgs.OnColorCommitted = FOnLinearColorValueChanged::CreateSP(this, &FColorStructCustomization::OnSetColorFromColorPicker);
+		PickerArgs.OnColorPickerCancelled = FOnColorPickerCancelled::CreateSP(this, &FColorStructCustomization::OnColorPickerCancelled);
+		PickerArgs.OnInteractivePickBegin = FSimpleDelegate::CreateSP(this, &FColorStructCustomization::OnColorPickerInteractiveBegin);
+		PickerArgs.OnInteractivePickEnd = FSimpleDelegate::CreateSP(this, &FColorStructCustomization::OnColorPickerInteractiveEnd);
+		PickerArgs.InitialColorOverride = InitialColor;
+		PickerArgs.ParentWidget = ColorPickerParentWidget;
+	}
+
+	OpenColorPicker(PickerArgs);
+}
+
+
+TSharedRef<SColorPicker> FColorStructCustomization::CreateInlineColorPicker()
+{
+	int32 NumObjects = StructPropertyHandle->GetNumOuterObjects();
+
+	SavedPreColorPickerColors.Empty();
+	TArray<FString> PerObjectValues;
+	StructPropertyHandle->GetPerObjectValues(PerObjectValues);
+
+	for (int32 ObjectIndex = 0; ObjectIndex < NumObjects; ++ObjectIndex)
+	{
+		if (bIsLinearColor)
+		{
+			FLinearColor Color;
+			Color.InitFromString(PerObjectValues[ObjectIndex]);
+			SavedPreColorPickerColors.Add(Color);	
+		}
+		else
+		{
+			FColor Color;
+			Color.InitFromString(PerObjectValues[ObjectIndex]);
+			SavedPreColorPickerColors.Add(Color.ReinterpretAsLinear());
+		}
+	}
+
+	FLinearColor InitialColor;
+	GetColorAsLinear(InitialColor);
+
+	const bool bRefreshOnlyOnOk = /*bDontUpdateWhileEditing ||*/ StructPropertyHandle->HasMetaData("DontUpdateWhileEditing");
+	
+	return SNew(SColorPicker)
+		.DisplayInlineVersion(true)
+		.OnlyRefreshOnMouseUp(false)
+		.OnlyRefreshOnOk(bRefreshOnlyOnOk)
+		.DisplayGamma(TAttribute<float>::Create(TAttribute<float>::FGetter::CreateUObject(GEngine, &UEngine::GetDisplayGamma)))
+		.OnColorCommitted(FOnLinearColorValueChanged::CreateSP(this, &FColorStructCustomization::OnSetColorFromColorPicker))
+		.OnColorPickerCancelled(FOnColorPickerCancelled::CreateSP(this, &FColorStructCustomization::OnColorPickerCancelled))
+		.OnInteractivePickBegin(FSimpleDelegate::CreateSP(this, &FColorStructCustomization::OnColorPickerInteractiveBegin))
+		.OnInteractivePickEnd(FSimpleDelegate::CreateSP(this, &FColorStructCustomization::OnColorPickerInteractiveEnd))
+		.sRGBOverride(sRGBOverride)
+		.TargetColorAttribute(InitialColor);
+}
+
+
+void FColorStructCustomization::OnSetColorFromColorPicker(FLinearColor NewColor)
+{
+	FString ColorString;
+	if (bIsLinearColor)
+	{
+		ColorString = NewColor.ToString();
+	}
+	else
+	{
+		// Handled by the color picker
+		const bool bSRGB = false;
+		FColor NewFColor = NewColor.ToFColor(bSRGB);
+		ColorString = NewFColor.ToString();
+	}
+
+	StructPropertyHandle->SetValueFromFormattedString(ColorString, bIsInteractive ? EPropertyValueSetFlags::InteractiveChange : 0);
+	StructPropertyHandle->NotifyFinishedChangingProperties();
+}
+
+
+void FColorStructCustomization::OnColorPickerCancelled(FLinearColor OriginalColor)
+{
+	TArray<FString> PerObjectColors;
+
+	for (int32 ColorIndex = 0; ColorIndex < SavedPreColorPickerColors.Num(); ++ColorIndex)
+	{
+		if (bIsLinearColor)
+		{
+			PerObjectColors.Add(SavedPreColorPickerColors[ColorIndex].ToString());
+		}
+		else
+		{
+			const bool bSRGB = false;
+			FColor Color = SavedPreColorPickerColors[ColorIndex].ToFColor(bSRGB);
+			PerObjectColors.Add(Color.ToString());
+		}
+	}
+
+	StructPropertyHandle->SetPerObjectValues(PerObjectColors);
+	PerObjectColors.Empty();
+}
+
+
+void FColorStructCustomization::OnColorPickerInteractiveBegin()
+{
+	bIsInteractive = true;
+
+	GEditor->BeginTransaction(FText::Format(LOCTEXT("SetColorProperty", "Edit {0}"), StructPropertyHandle->GetPropertyDisplayName()));
+}
+
+
+void FColorStructCustomization::OnColorPickerInteractiveEnd()
+{
+	bIsInteractive = false;
+
+	if (!bDontUpdateWhileEditing)
+	{
+		// pushes the last value from the interactive change without the interactive flag
+		FString ColorString;
+		StructPropertyHandle->GetValueAsFormattedString(ColorString);
+		StructPropertyHandle->SetValueFromFormattedString(ColorString);
+	}
+
+	GEditor->EndTransaction();
+}
+
+
+FLinearColor FColorStructCustomization::OnGetColorForColorBlock() const
+{
+	FLinearColor Color;
+	GetColorAsLinear(Color);
+	return Color;
+}
+
+
+FPropertyAccess::Result FColorStructCustomization::GetColorAsLinear(FLinearColor& OutColor) const
+{
+	// Default to full alpha in case the alpha component is disabled.
+	OutColor.A = 1.0f;
+
+	// Get each color component 
+	for (int32 ChildIndex = 0; ChildIndex < SortedChildHandles.Num(); ++ChildIndex)
+	{
+		FPropertyAccess::Result ValueResult = FPropertyAccess::Fail;
+
+		if (bIsLinearColor)
+		{
+			float ComponentValue = 0;
+			ValueResult = SortedChildHandles[ChildIndex]->GetValue(ComponentValue);
+
+			OutColor.Component(ChildIndex) = ComponentValue;
+		}
+		else
+		{
+			uint8 ComponentValue = 0;
+			ValueResult = SortedChildHandles[ChildIndex]->GetValue(ComponentValue);
+
+			// Convert the FColor to a linear equivalent
+			OutColor.Component(ChildIndex) = ComponentValue/255.0f;
+		}
+
+		switch(ValueResult)
+		{
+		case FPropertyAccess::MultipleValues:
+			{
+				// Default the color to white if we've got multiple values selected
+				OutColor = FLinearColor::White;
+				return FPropertyAccess::MultipleValues;
+			}
+
+		case FPropertyAccess::Fail:
+			return FPropertyAccess::Fail;
+
+		default:
+			break;
+		}
+	}
+
+	// If we've got this far, we have to have successfully read a color with a single value
+	return FPropertyAccess::Success;
+}
+
+
+EVisibility FColorStructCustomization::GetMultipleValuesTextVisibility() const
+{
+	FLinearColor Color;
+	const FPropertyAccess::Result ValueResult = GetColorAsLinear(Color);
+	return (ValueResult == FPropertyAccess::MultipleValues) ? EVisibility::Visible : EVisibility::Collapsed;
+}
+
+
+FReply FColorStructCustomization::OnMouseButtonDownColorBlock(const FGeometry& MyGeometry, const FPointerEvent& MouseEvent)
+{
+	if (MouseEvent.GetEffectingButton() != EKeys::LeftMouseButton)
+	{
+		return FReply::Unhandled();
+	}
+	
+	CreateColorPicker(true /*bUseAlpha*/);
+
+	return FReply::Handled();
+}
+
+
+FReply FColorStructCustomization::OnOpenFullColorPickerClicked()
+{
+	CreateColorPicker(true /*bUseAlpha*/);
+	bIsInlineColorPickerVisible = false;
+
+	return FReply::Handled();
+}
+
+
+#undef LOCTEXT_NAMESPACE
diff --git a/Engine/Source/Editor/DetailCustomizations/Private/ColorStructCustomization.h b/Engine/Source/Editor/DetailCustomizations/Private/ColorStructCustomization.h
new file mode 100644
--- /dev/null
+++ b/Engine/Source/Editor/DetailCustomizations/Private/ColorStructCustomization.h
@@ -0,0 +1,133 @@
+// Copyright 1998-2015 Epic Games, Inc. All Rights Reserved.
+
+#pragma once
+
+#include "MathStructCustomizations.h"
+
+
+class SColorPicker;
+
+
+/**
+ * Base class for color struct customization (FColor,FLinearColor).
+ */
+class FColorStructCustomization
+	: public FMathStructCustomization
+{
+public:
+
+	static TSharedRef<IPropertyTypeCustomization> MakeInstance();
+	
+protected:
+
+	FColorStructCustomization()
+		: bIgnoreAlpha(false)
+		, bIsInlineColorPickerVisible(false)
+		, bIsInteractive(false)
+		, bDontUpdateWhileEditing(false)
+	{}
+
+protected:
+
+	/** Creates the color widget that when clicked spawns the color picker window. */
+	TSharedRef<SWidget> CreateColorWidget();
+
+	/**
+	 * Get the color used by this struct as a linear color value
+	 * @param InColor To be filled with the color value used by this struct, or white if this struct is being used to edit multiple values  
+	 * @return The result of trying to get the color value
+	 */
+	FPropertyAccess::Result GetColorAsLinear(FLinearColor& InColor) const;
+
+	/**
+	 * Does this struct have multiple values?
+	 * @return EVisibility::Visible if it does, EVisibility::Collapsed otherwise
+	 */
+	EVisibility GetMultipleValuesTextVisibility() const;
+
+	/**
+	 * Creates a new color picker for interactively selecting the color
+	 *
+	 * @param bUseAlpha If true alpha will be displayed, otherwise it is ignored
+	 * @param bOnlyRefreshOnOk If true the value of the property will only be refreshed when the user clicks OK in the color picker
+	 */
+	void CreateColorPicker(bool bUseAlpha);
+	
+	/** Creates a new color picker for interactively selecting color */
+	TSharedRef<SColorPicker> CreateInlineColorPicker();
+
+	/**
+	 * Called when the property is set from the color picker 
+	 * 
+	 * @param NewColor The new color value
+	 */
+	void OnSetColorFromColorPicker(FLinearColor NewColor);
+	
+	/**
+	 * Called when the user clicks cancel in the color picker
+	 * The values are reset to their original state when this happens
+	 *
+	 * @param OriginalColor Original color of the property
+	 */
+	void OnColorPickerCancelled(FLinearColor OriginalColor);
+
+	/**
+	 * Called when the user enters an interactive color change (dragging something in the picker)
+	 */
+	void OnColorPickerInteractiveBegin();
+
+	/**
+	 * Called when the user completes an interactive color change (dragging something in the picker)
+	 */
+	void OnColorPickerInteractiveEnd();
+
+	/**
+	 * @return The color that should be displayed in the color block                                                              
+	 */
+	FLinearColor OnGetColorForColorBlock() const;
+	
+	/**
+	 * Called when the user clicks in the color block (opens inline color picker)
+	 */
+	FReply OnMouseButtonDownColorBlock(const FGeometry& MyGeometry, const FPointerEvent& MouseEvent);
+	
+	/** Called when the user clicks on the the button to get the full color picker */
+	FReply OnOpenFullColorPickerClicked();
+
+protected:
+
+	// FMathStructCustomization interface
+
+	virtual void CustomizeHeader(TSharedRef<IPropertyHandle> StructPropertyHandle, FDetailWidgetRow& HeaderRow, IPropertyTypeCustomizationUtils& StructCustomizationUtils) override;
+	virtual void MakeHeaderRow(TSharedRef<IPropertyHandle>& InStructPropertyHandle, FDetailWidgetRow& Row) override;
+	virtual void GetSortedChildren(TSharedRef<IPropertyHandle> StructPropertyHandle, TArray< TSharedRef<IPropertyHandle> >& OutChildren) override;
+
+protected:
+
+	/** Saved per struct colors in case the user clicks cancel in the color picker */
+	TArray<FLinearColor> SavedPreColorPickerColors;
+
+	/** Color struct handle */
+	TSharedPtr<IPropertyHandle> StructPropertyHandle;
+
+	/** True if the property is a linear color property */
+	bool bIsLinearColor;
+
+	/** True if the property wants to ignore the alpha component */
+	bool bIgnoreAlpha;
+
+	/** True if the inline color picker is visible */
+	bool bIsInlineColorPickerVisible;
+
+	/** True if the user is performing an interactive color change */
+	bool bIsInteractive;
+
+	/** Cached widget for the color picker to use as a parent */
+	TSharedPtr<SWidget> ColorPickerParentWidget;
+
+	/** The value won;t be updated while editing */
+	bool bDontUpdateWhileEditing;
+
+	/** Overrides the default state of the sRGB check box */
+	TOptional<bool> sRGBOverride;
+};
diff --git a/Engine/Source/Editor/DetailCustomizations/Private/DetailCustomizations.cpp b/Engine/Source/Editor/DetailCustomizations/Private/DetailCustomizations.cpp
--- a/Engine/Source/Editor/DetailCustomizations/Private/DetailCustomizations.cpp
+++ b/Engine/Source/Editor/DetailCustomizations/Private/DetailCustomizations.cpp
@@ -105,11 +105,14 @@
 #include "MovieSceneCaptureCustomization.h"
 #include "TextCustomization.h"
 #include "AnimTrailNodeDetails.h"
-
 #include "MaterialProxySettingsCustomizations.h"
 #include "ImportantToggleSettingCustomization.h"
 #include "CameraFilmbackSettingsCustomization.h"
 #include "CameraLensSettingsCustomization.h"
+#include "ColorStructCustomization.h"
+#include "RotatorStructCustomization.h"
+#include "VectorStructCustomization.h"
+
 
 IMPLEMENT_MODULE( FDetailCustomizationsModule, DetailCustomizations );
 
@@ -161,7 +164,7 @@ void FDetailCustomizationsModule::RegisterPropertyTypeCustomizations()
 	RegisterCustomPropertyTypeLayout("DataTableRowHandle", FOnGetPropertyTypeCustomizationInstance::CreateStatic(&FDataTableCustomizationLayout::MakeInstance));
 	RegisterCustomPropertyTypeLayout("DataTableCategoryHandle", FOnGetPropertyTypeCustomizationInstance::CreateStatic(&FDataTableCategoryCustomizationLayout::MakeInstance));
 	RegisterCustomPropertyTypeLayout("CurveTableRowHandle", FOnGetPropertyTypeCustomizationInstance::CreateStatic(&FCurveTableCustomizationLayout::MakeInstance));
-	RegisterCustomPropertyTypeLayout(NAME_Vector, FOnGetPropertyTypeCustomizationInstance::CreateStatic(&FMathStructCustomization::MakeInstance));
+	RegisterCustomPropertyTypeLayout(NAME_Vector, FOnGetPropertyTypeCustomizationInstance::CreateStatic(&FVectorStructCustomization::MakeInstance));
 	RegisterCustomPropertyTypeLayout(NAME_Vector4, FOnGetPropertyTypeCustomizationInstance::CreateStatic(&FMathStructCustomization::MakeInstance));
 	RegisterCustomPropertyTypeLayout(NAME_Vector2D, FOnGetPropertyTypeCustomizationInstance::CreateStatic(&FMathStructCustomization::MakeInstance));
 	RegisterCustomPropertyTypeLayout(NAME_IntPoint, FOnGetPropertyTypeCustomizationInstance::CreateStatic(&FMathStructCustomization::MakeInstance));
@@ -336,4 +339,4 @@ void FDetailCustomizationsModule::RegisterCustomPropertyTypeLayout(FName Propert
 	static FName PropertyEditor("PropertyEditor");
 	FPropertyEditorModule& PropertyModule = FModuleManager::GetModuleChecked<FPropertyEditorModule>(PropertyEditor);
 	PropertyModule.RegisterCustomPropertyTypeLayout(PropertyTypeName, PropertyTypeLayoutDelegate);
-}
\ No newline at end of file
+}
diff --git a/Engine/Source/Editor/DetailCustomizations/Private/MathStructCustomizations.cpp b/Engine/Source/Editor/DetailCustomizations/Private/MathStructCustomizations.cpp
--- a/Engine/Source/Editor/DetailCustomizations/Private/MathStructCustomizations.cpp
+++ b/Engine/Source/Editor/DetailCustomizations/Private/MathStructCustomizations.cpp
@@ -4,7 +4,6 @@
 #include "Materials/MaterialExpressionConstant3Vector.h"
 #include "MathStructCustomizations.h"
 #include "IPropertyUtilities.h"
-#include "SColorPicker.h"
 #include "SNumericEntryBox.h"
 
 
@@ -374,7 +373,7 @@ void FMathStructCustomization::OnBeginSliderMovement()
 {
 	bIsUsingSlider = true;
 
-	GEditor->BeginTransaction(NSLOCTEXT("FMathStructCustomization", "SetVectorProperty", "Set Vector Property"));
+	GEditor->BeginTransaction(LOCTEXT("SetVectorProperty", "Set Vector Property"));
 }
 
 
@@ -387,440 +386,4 @@ void FMathStructCustomization::OnEndSliderMovement(NumericType NewValue)
 }
 
 
-TSharedRef<IPropertyTypeCustomization> FRotatorStructCustomization::MakeInstance() 
-{
-	return MakeShareable(new FRotatorStructCustomization);
-}
-
-
-void FRotatorStructCustomization::GetSortedChildren(TSharedRef<IPropertyHandle> StructPropertyHandle, TArray< TSharedRef<IPropertyHandle> >& OutChildren)
-{
-	static const FName Roll("Roll");
-	static const FName Pitch("Pitch");
-	static const FName Yaw("Yaw");
-
-	TSharedPtr< IPropertyHandle > RotatorChildren[3];
-
-	uint32 NumChildren;
-	StructPropertyHandle->GetNumChildren(NumChildren);
-
-	for (uint32 ChildIndex = 0; ChildIndex < NumChildren; ++ChildIndex)
-	{
-		TSharedRef<IPropertyHandle> ChildHandle = StructPropertyHandle->GetChildHandle(ChildIndex).ToSharedRef();
-		const FName PropertyName = ChildHandle->GetProperty()->GetFName();
-
-		if (PropertyName == Roll)
-		{
-			RotatorChildren[0] = ChildHandle;
-		}
-		else if (PropertyName == Pitch)
-		{
-			RotatorChildren[1] = ChildHandle;
-		}
-		else
-		{
-			check(PropertyName == Yaw);
-			RotatorChildren[2] = ChildHandle;
-		}
-	}
-
-	OutChildren.Add(RotatorChildren[0].ToSharedRef());
-	OutChildren.Add(RotatorChildren[1].ToSharedRef());
-	OutChildren.Add(RotatorChildren[2].ToSharedRef());
-}
-
-
-TSharedRef<IPropertyTypeCustomization> FColorStructCustomization::MakeInstance() 
-{
-	return MakeShareable(new FColorStructCustomization);
-}
-
-
-void FColorStructCustomization::CustomizeHeader(TSharedRef<class IPropertyHandle> InStructPropertyHandle, class FDetailWidgetRow& InHeaderRow, IPropertyTypeCustomizationUtils& StructCustomizationUtils)
-{
-	StructPropertyHandle = InStructPropertyHandle;
-
-	bIsLinearColor = CastChecked<UStructProperty>(StructPropertyHandle->GetProperty())->Struct->GetFName() == NAME_LinearColor;
-	bIgnoreAlpha = StructPropertyHandle->GetProperty()->HasMetaData(TEXT("HideAlphaChannel"));
-	
-	if (StructPropertyHandle->GetProperty()->HasMetaData(TEXT("sRGB")))
-	{
-		sRGBOverride = StructPropertyHandle->GetProperty()->GetBoolMetaData(TEXT("sRGB"));
-	}
-
-	auto PropertyUtils = StructCustomizationUtils.GetPropertyUtilities();
-	bDontUpdateWhileEditing = PropertyUtils.IsValid() ? PropertyUtils->DontUpdateValueWhileEditing() : false;
-
-	FMathStructCustomization::CustomizeHeader(InStructPropertyHandle, InHeaderRow, StructCustomizationUtils);
-}
-
-
-void FColorStructCustomization::MakeHeaderRow(TSharedRef<class IPropertyHandle>& InStructPropertyHandle, FDetailWidgetRow& Row)
-{
-	// We'll set up reset to default ourselves
-	const bool bDisplayResetToDefault = false;
-	const FText DisplayNameOverride = FText::GetEmpty();
-	const FText DisplayToolTipOverride = FText::GetEmpty();
-	
-	TSharedPtr<SWidget> ColorWidget;
-	float ContentWidth = 250.0f;
-
-	if (InStructPropertyHandle->HasMetaData("InlineColorPicker"))
-	{
-		ColorWidget = CreateInlineColorPicker();
-		ContentWidth = 384.0f;
-	}
-	else
-	{
-		ColorWidget = CreateColorWidget();
-	}
-
-	Row.NameContent()
-	[
-		StructPropertyHandle->CreatePropertyNameWidget(DisplayNameOverride, DisplayToolTipOverride, bDisplayResetToDefault)
-	]
-	.ValueContent()
-	.MinDesiredWidth(ContentWidth)
-	[
-		ColorWidget.ToSharedRef()
-	];
-}
-
-
-TSharedRef<SWidget> FColorStructCustomization::CreateColorWidget()
-{
-	FSlateFontInfo NormalText = IDetailLayoutBuilder::GetDetailFont();
-
-	return SNew(SHorizontalBox)
-		+ SHorizontalBox::Slot()
-		.VAlign(VAlign_Center)
-		[
-			SNew(SOverlay)
-			+SOverlay::Slot()
-			[
-				// Displays the color with alpha unless it is ignored
-				SAssignNew(ColorPickerParentWidget, SColorBlock)
-				.Color(this, &FColorStructCustomization::OnGetColorForColorBlock)
-				.ShowBackgroundForAlpha(true)
-				.IgnoreAlpha(bIgnoreAlpha)
-				.OnMouseButtonDown(this, &FColorStructCustomization::OnMouseButtonDownColorBlock)
-				.Size(FVector2D(35.0f, 12.0f))
-			]
-			+SOverlay::Slot()
-			.HAlign(HAlign_Center)
-			.VAlign(VAlign_Center)
-			[
-				SNew(STextBlock)
-				.Text(NSLOCTEXT("PropertyEditor", "MultipleValues", "Multiple Values"))
-				.Font(NormalText)
-				.ColorAndOpacity(FSlateColor(FLinearColor::Black)) // we know the background is always white, so can safely set this to black
-				.Visibility(this, &FColorStructCustomization::GetMultipleValuesTextVisibility)
-			]
-		]
-		+ SHorizontalBox::Slot()
-		.VAlign(VAlign_Center)
-		[
-			// Displays the color without alpha
-			SNew(SColorBlock)
-			.Color(this, &FColorStructCustomization::OnGetColorForColorBlock)
-			.ShowBackgroundForAlpha(false)
-			.IgnoreAlpha(true)
-			.OnMouseButtonDown(this, &FColorStructCustomization::OnMouseButtonDownColorBlock)
-			.Size(FVector2D(35.0f, 12.0f))
-		];
-}
-
-
-void FColorStructCustomization::GetSortedChildren(TSharedRef<IPropertyHandle> InStructPropertyHandle, TArray< TSharedRef<IPropertyHandle> >& OutChildren)
-{
-	static const FName Red("R");
-	static const FName Green("G");
-	static const FName Blue("B");
-	static const FName Alpha("A");
-
-	// We control the order of the colors via this array so it always ends up R,G,B,A
-	TSharedPtr< IPropertyHandle > ColorProperties[4];
-
-	uint32 NumChildren;
-	InStructPropertyHandle->GetNumChildren(NumChildren);
-
-	for (uint32 ChildIndex = 0; ChildIndex < NumChildren; ++ChildIndex)
-	{
-		TSharedRef<IPropertyHandle> ChildHandle = InStructPropertyHandle->GetChildHandle(ChildIndex).ToSharedRef();
-		const FName PropertyName = ChildHandle->GetProperty()->GetFName();
-
-		if (PropertyName == Red)
-		{
-			ColorProperties[0] = ChildHandle;
-		}
-		else if (PropertyName == Green)
-		{
-			ColorProperties[1] = ChildHandle;
-		}
-		else if (PropertyName == Blue)
-		{
-			ColorProperties[2] = ChildHandle;
-		}
-		else
-		{
-			ColorProperties[3] = ChildHandle;
-		}
-	}
-
-	OutChildren.Add(ColorProperties[0].ToSharedRef());
-	OutChildren.Add(ColorProperties[1].ToSharedRef());
-	OutChildren.Add(ColorProperties[2].ToSharedRef());
-
-	// Alpha channel may not be used
-	if (!bIgnoreAlpha && ColorProperties[3].IsValid())
-	{
-		OutChildren.Add(ColorProperties[3].ToSharedRef());
-	}
-}
-
-
-void FColorStructCustomization::CreateColorPicker(bool bUseAlpha)
-{
-	int32 NumObjects = StructPropertyHandle->GetNumOuterObjects();
-
-	SavedPreColorPickerColors.Empty();
-	TArray<FString> PerObjectValues;
-	StructPropertyHandle->GetPerObjectValues(PerObjectValues);
-
-	for (int32 ObjectIndex = 0; ObjectIndex < NumObjects; ++ObjectIndex)
-	{
-		if (bIsLinearColor)
-		{
-			FLinearColor Color;
-			Color.InitFromString(PerObjectValues[ObjectIndex]);
-			SavedPreColorPickerColors.Add(Color);	
-		}
-		else
-		{
-			FColor Color;
-			Color.InitFromString(PerObjectValues[ObjectIndex]);
-			SavedPreColorPickerColors.Add(Color.ReinterpretAsLinear());
-		}
-	}
-
-	FLinearColor InitialColor;
-	GetColorAsLinear(InitialColor);
-
-	const bool bRefreshOnlyOnOk = bDontUpdateWhileEditing || StructPropertyHandle->HasMetaData("DontUpdateWhileEditing");
-
-	FColorPickerArgs PickerArgs;
-	{
-		PickerArgs.bUseAlpha = !bIgnoreAlpha;
-		PickerArgs.bOnlyRefreshOnMouseUp = false;
-		PickerArgs.bOnlyRefreshOnOk = bRefreshOnlyOnOk;
-		PickerArgs.sRGBOverride = sRGBOverride;
-		PickerArgs.DisplayGamma = TAttribute<float>::Create(TAttribute<float>::FGetter::CreateUObject(GEngine, &UEngine::GetDisplayGamma));
-		PickerArgs.OnColorCommitted = FOnLinearColorValueChanged::CreateSP(this, &FColorStructCustomization::OnSetColorFromColorPicker);
-		PickerArgs.OnColorPickerCancelled = FOnColorPickerCancelled::CreateSP(this, &FColorStructCustomization::OnColorPickerCancelled);
-		PickerArgs.OnInteractivePickBegin = FSimpleDelegate::CreateSP(this, &FColorStructCustomization::OnColorPickerInteractiveBegin);
-		PickerArgs.OnInteractivePickEnd = FSimpleDelegate::CreateSP(this, &FColorStructCustomization::OnColorPickerInteractiveEnd);
-		PickerArgs.InitialColorOverride = InitialColor;
-		PickerArgs.ParentWidget = ColorPickerParentWidget;
-	}
-
-	OpenColorPicker(PickerArgs);
-}
-
-
-TSharedRef<SColorPicker> FColorStructCustomization::CreateInlineColorPicker()
-{
-	int32 NumObjects = StructPropertyHandle->GetNumOuterObjects();
-
-	SavedPreColorPickerColors.Empty();
-	TArray<FString> PerObjectValues;
-	StructPropertyHandle->GetPerObjectValues(PerObjectValues);
-
-	for (int32 ObjectIndex = 0; ObjectIndex < NumObjects; ++ObjectIndex)
-	{
-		if (bIsLinearColor)
-		{
-			FLinearColor Color;
-			Color.InitFromString(PerObjectValues[ObjectIndex]);
-			SavedPreColorPickerColors.Add(Color);	
-		}
-		else
-		{
-			FColor Color;
-			Color.InitFromString(PerObjectValues[ObjectIndex]);
-			SavedPreColorPickerColors.Add(Color.ReinterpretAsLinear());
-		}
-	}
-
-	FLinearColor InitialColor;
-	GetColorAsLinear(InitialColor);
-
-	const bool bRefreshOnlyOnOk = /*bDontUpdateWhileEditing ||*/ StructPropertyHandle->HasMetaData("DontUpdateWhileEditing");
-	
-	return SNew(SColorPicker)
-		.DisplayInlineVersion(true)
-		.OnlyRefreshOnMouseUp(false)
-		.OnlyRefreshOnOk(bRefreshOnlyOnOk)
-		.DisplayGamma(TAttribute<float>::Create(TAttribute<float>::FGetter::CreateUObject(GEngine, &UEngine::GetDisplayGamma)))
-		.OnColorCommitted(FOnLinearColorValueChanged::CreateSP(this, &FColorStructCustomization::OnSetColorFromColorPicker))
-		.OnColorPickerCancelled(FOnColorPickerCancelled::CreateSP(this, &FColorStructCustomization::OnColorPickerCancelled))
-		.OnInteractivePickBegin(FSimpleDelegate::CreateSP(this, &FColorStructCustomization::OnColorPickerInteractiveBegin))
-		.OnInteractivePickEnd(FSimpleDelegate::CreateSP(this, &FColorStructCustomization::OnColorPickerInteractiveEnd))
-		.sRGBOverride(sRGBOverride)
-		.TargetColorAttribute(InitialColor);
-}
-
-
-void FColorStructCustomization::OnSetColorFromColorPicker(FLinearColor NewColor)
-{
-	FString ColorString;
-	if (bIsLinearColor)
-	{
-		ColorString = NewColor.ToString();
-	}
-	else
-	{
-		// Handled by the color picker
-		const bool bSRGB = false;
-		FColor NewFColor = NewColor.ToFColor(bSRGB);
-		ColorString = NewFColor.ToString();
-	}
-
-	StructPropertyHandle->SetValueFromFormattedString(ColorString, bIsInteractive ? EPropertyValueSetFlags::InteractiveChange : 0);
-	StructPropertyHandle->NotifyFinishedChangingProperties();
-}
-
-
-void FColorStructCustomization::OnColorPickerCancelled(FLinearColor OriginalColor)
-{
-	TArray<FString> PerObjectColors;
-
-	for (int32 ColorIndex = 0; ColorIndex < SavedPreColorPickerColors.Num(); ++ColorIndex)
-	{
-		if (bIsLinearColor)
-		{
-			PerObjectColors.Add(SavedPreColorPickerColors[ColorIndex].ToString());
-		}
-		else
-		{
-			const bool bSRGB = false;
-			FColor Color = SavedPreColorPickerColors[ColorIndex].ToFColor(bSRGB);
-			PerObjectColors.Add(Color.ToString());
-		}
-	}
-
-	StructPropertyHandle->SetPerObjectValues(PerObjectColors);
-	PerObjectColors.Empty();
-}
-
-
-void FColorStructCustomization::OnColorPickerInteractiveBegin()
-{
-	bIsInteractive = true;
-
-	GEditor->BeginTransaction(FText::Format(NSLOCTEXT("FColorStructCustomization", "SetColorProperty", "Edit {0}"), StructPropertyHandle->GetPropertyDisplayName()));
-}
-
-
-void FColorStructCustomization::OnColorPickerInteractiveEnd()
-{
-	bIsInteractive = false;
-
-	if (!bDontUpdateWhileEditing)
-	{
-		// pushes the last value from the interactive change without the interactive flag
-		FString ColorString;
-		StructPropertyHandle->GetValueAsFormattedString(ColorString);
-		StructPropertyHandle->SetValueFromFormattedString(ColorString);
-	}
-
-	GEditor->EndTransaction();
-}
-
-
-FLinearColor FColorStructCustomization::OnGetColorForColorBlock() const
-{
-	FLinearColor Color;
-	GetColorAsLinear(Color);
-	return Color;
-}
-
-
-FPropertyAccess::Result FColorStructCustomization::GetColorAsLinear(FLinearColor& OutColor) const
-{
-	// Default to full alpha in case the alpha component is disabled.
-	OutColor.A = 1.0f;
-
-	// Get each color component 
-	for (int32 ChildIndex = 0; ChildIndex < SortedChildHandles.Num(); ++ChildIndex)
-	{
-		FPropertyAccess::Result ValueResult = FPropertyAccess::Fail;
-
-		if (bIsLinearColor)
-		{
-			float ComponentValue = 0;
-			ValueResult = SortedChildHandles[ChildIndex]->GetValue(ComponentValue);
-
-			OutColor.Component(ChildIndex) = ComponentValue;
-		}
-		else
-		{
-			uint8 ComponentValue = 0;
-			ValueResult = SortedChildHandles[ChildIndex]->GetValue(ComponentValue);
-
-			// Convert the FColor to a linear equivalent
-			OutColor.Component(ChildIndex) = ComponentValue/255.0f;
-		}
-
-		switch(ValueResult)
-		{
-		case FPropertyAccess::MultipleValues:
-			{
-				// Default the color to white if we've got multiple values selected
-				OutColor = FLinearColor::White;
-				return FPropertyAccess::MultipleValues;
-			}
-
-		case FPropertyAccess::Fail:
-			return FPropertyAccess::Fail;
-
-		default:
-			break;
-		}
-	}
-
-	// If we've got this far, we have to have successfully read a color with a single value
-	return FPropertyAccess::Success;
-}
-
-
-EVisibility FColorStructCustomization::GetMultipleValuesTextVisibility() const
-{
-	FLinearColor Color;
-	const FPropertyAccess::Result ValueResult = GetColorAsLinear(Color);
-	return (ValueResult == FPropertyAccess::MultipleValues) ? EVisibility::Visible : EVisibility::Collapsed;
-}
-
-
-FReply FColorStructCustomization::OnMouseButtonDownColorBlock(const FGeometry& MyGeometry, const FPointerEvent& MouseEvent)
-{
-	if (MouseEvent.GetEffectingButton() != EKeys::LeftMouseButton)
-	{
-		return FReply::Unhandled();
-	}
-	
-	CreateColorPicker(true /*bUseAlpha*/);
-
-	return FReply::Handled();
-}
-
-
-FReply FColorStructCustomization::OnOpenFullColorPickerClicked()
-{
-	CreateColorPicker(true /*bUseAlpha*/);
-	bIsInlineColorPickerVisible = false;
-
-	return FReply::Handled();
-}
-
-
 #undef LOCTEXT_NAMESPACE
diff --git a/Engine/Source/Editor/DetailCustomizations/Private/MathStructCustomizations.h b/Engine/Source/Editor/DetailCustomizations/Private/MathStructCustomizations.h
--- a/Engine/Source/Editor/DetailCustomizations/Private/MathStructCustomizations.h
+++ b/Engine/Source/Editor/DetailCustomizations/Private/MathStructCustomizations.h
@@ -2,6 +2,8 @@
 
 #pragma once
 
+#include "IPropertyTypeCustomization.h"
+
 
 /**
  * Base class for math struct customization (e.g, vector, rotator, color)                                                              
@@ -20,8 +22,8 @@ public:
 	virtual ~FMathStructCustomization() {}
 
 	/** IPropertyTypeCustomization instance */
-	virtual void CustomizeHeader(TSharedRef<class IPropertyHandle> StructPropertyHandle, class FDetailWidgetRow& HeaderRow, IPropertyTypeCustomizationUtils& StructCustomizationUtils) override;
-	virtual void CustomizeChildren(TSharedRef<class IPropertyHandle> StructPropertyHandle, class IDetailChildrenBuilder& StructBuilder, IPropertyTypeCustomizationUtils& StructCustomizationUtils) override;
+	virtual void CustomizeHeader(TSharedRef<IPropertyHandle> StructPropertyHandle, FDetailWidgetRow& HeaderRow, IPropertyTypeCustomizationUtils& StructCustomizationUtils) override;
+	virtual void CustomizeChildren(TSharedRef<IPropertyHandle> StructPropertyHandle, IDetailChildrenBuilder& StructBuilder, IPropertyTypeCustomizationUtils& StructCustomizationUtils) override;
 	
 protected:
 
@@ -31,7 +33,7 @@ protected:
 	 * @param StructPropertyHandle	Handle to the struct property
 	 * @param Row	The header row to add widgets to
 	 */
-	virtual void MakeHeaderRow(TSharedRef<class IPropertyHandle>& StructPropertyHandle, FDetailWidgetRow& Row);
+	virtual void MakeHeaderRow(TSharedRef<IPropertyHandle>& StructPropertyHandle, FDetailWidgetRow& Row);
 
 	/**
 	 * Gets the sorted children for the struct
@@ -130,141 +132,3 @@ protected:
 	/** True if the ratio is locked when scaling occurs (uniform scaling) */
 	bool bPreserveScaleRatio;
 };
-
-
-/**
- * Customizes FRotator structs                                                              
- */
-class FRotatorStructCustomization
-	: public FMathStructCustomization
-{
-public:
-
-	/** @return A new instance of this class */
-	static TSharedRef<IPropertyTypeCustomization> MakeInstance();
-
-private:
-
-	/** FMathStructCustomization interface */
-	virtual void GetSortedChildren(TSharedRef<IPropertyHandle> StructPropertyHandle, TArray< TSharedRef<IPropertyHandle> >& OutChildren) override;
-};
-
-
-/**
- * Base class for color struct customization (FColor,FLinearColor)                                                             
- */
-class FColorStructCustomization
-	: public FMathStructCustomization
-{
-public:
-
-	static TSharedRef<IPropertyTypeCustomization> MakeInstance();
-	
-protected:
-
-	FColorStructCustomization()
-		: bIgnoreAlpha(false)
-		, bIsInlineColorPickerVisible(false)
-		, bIsInteractive(false)
-		, bDontUpdateWhileEditing(false)
-	{}
-
-	/** FMathStructCustomization interface */
-	virtual void CustomizeHeader(TSharedRef<class IPropertyHandle> StructPropertyHandle, class FDetailWidgetRow& HeaderRow, IPropertyTypeCustomizationUtils& StructCustomizationUtils) override;
-	virtual void MakeHeaderRow(TSharedRef<class IPropertyHandle>& InStructPropertyHandle, FDetailWidgetRow& Row) override;
-	virtual void GetSortedChildren(TSharedRef<IPropertyHandle> StructPropertyHandle, TArray< TSharedRef<IPropertyHandle> >& OutChildren) override;
-
-	/** Creates the color widget that when clicked spawns the color picker window. */
-	TSharedRef<SWidget> CreateColorWidget();
-
-	/**
-	 * Get the color used by this struct as a linear color value
-	 * @param InColor To be filled with the color value used by this struct, or white if this struct is being used to edit multiple values  
-	 * @return The result of trying to get the color value
-	 */
-	FPropertyAccess::Result GetColorAsLinear(FLinearColor& InColor) const;
-
-	/**
-	 * Does this struct have multiple values?
-	 * @return EVisibility::Visible if it does, EVisibility::Collapsed otherwise
-	 */
-	EVisibility GetMultipleValuesTextVisibility() const;
-
-	/**
-	 * Creates a new color picker for interactively selecting the color
-	 *
-	 * @param bUseAlpha			If true alpha will be displayed, otherwise it is ignored
-	 * @param bOnlyRefreshOnOk	If true the value of the property will only be refreshed when the user clicks OK in the color picker
-	 */
-	void CreateColorPicker(bool bUseAlpha);
-	
-	/** Creates a new color picker for interactively selecting color */
-	TSharedRef<class SColorPicker> CreateInlineColorPicker();
-
-	/**
-	 * Called when the property is set from the color picker 
-	 * 
-	 * @param NewColor	The new color value
-	 */
-	void OnSetColorFromColorPicker(FLinearColor NewColor);
-	
-	/**
-	 * Called when the user clicks cancel in the color picker
-	 * The values are reset to their original state when this happens
-	 *
-	 * @param OriginalColor Original color of the property
-	 */
-	void OnColorPickerCancelled(FLinearColor OriginalColor);
-
-	/**
-	 * Called when the user enters an interactive color change (dragging something in the picker)
-	 */
-	void OnColorPickerInteractiveBegin();
-
-	/**
-	 * Called when the user completes an interactive color change (dragging something in the picker)
-	 */
-	void OnColorPickerInteractiveEnd();
-
-	/**
-	 * @return The color that should be displayed in the color block                                                              
-	 */
-	FLinearColor OnGetColorForColorBlock() const;
-	
-	/**
-	 * Called when the user clicks in the color block (opens inline color picker)
-	 */
-	FReply OnMouseButtonDownColorBlock(const FGeometry& MyGeometry, const FPointerEvent& MouseEvent);
-	
-	/** Called when the user clicks on the the button to get the full color picker */
-	FReply OnOpenFullColorPickerClicked();
-
-protected:
-
-	/** Saved per struct colors in case the user clicks cancel in the color picker */
-	TArray<FLinearColor> SavedPreColorPickerColors;
-
-	/** Color struct handle */
-	TSharedPtr<IPropertyHandle> StructPropertyHandle;
-
-	/** True if the property is a linear color property */
-	bool bIsLinearColor;
-
-	/** True if the property wants to ignore the alpha component */
-	bool bIgnoreAlpha;
-
-	/** True if the inline color picker is visible */
-	bool bIsInlineColorPickerVisible;
-
-	/** True if the user is performing an interactive color change */
-	bool bIsInteractive;
-
-	/** Cached widget for the color picker to use as a parent */
-	TSharedPtr<SWidget> ColorPickerParentWidget;
-
-	/** The value won;t be updated while editing */
-	bool bDontUpdateWhileEditing;
-
-	/** Overrides the default state of the sRGB checkbox */
-	TOptional<bool> sRGBOverride;
-};
diff --git a/Engine/Source/Editor/DetailCustomizations/Private/SlateColorCustomization.h b/Engine/Source/Editor/DetailCustomizations/Private/SlateColorCustomization.h
--- a/Engine/Source/Editor/DetailCustomizations/Private/SlateColorCustomization.h
+++ b/Engine/Source/Editor/DetailCustomizations/Private/SlateColorCustomization.h
@@ -2,23 +2,26 @@
 
 #pragma once
 
-#include "MathStructCustomizations.h"
+#include "ColorStructCustomization.h"
 
-class FSlateColorCustomization : public FColorStructCustomization
+
+class FSlateColorCustomization
+	: public FColorStructCustomization
 {
 public:
+
 	static TSharedRef<IPropertyTypeCustomization> MakeInstance();
 
-	/** IPropertyTypeCustomization interface */
-	virtual void CustomizeHeader( TSharedRef<class IPropertyHandle> StructPropertyHandle, class FDetailWidgetRow& HeaderRow, IPropertyTypeCustomizationUtils& StructCustomizationUtils ) override;
+public:
 
+	// IPropertyTypeCustomization interface
+
+	virtual void CustomizeHeader( TSharedRef<class IPropertyHandle> StructPropertyHandle, class FDetailWidgetRow& HeaderRow, IPropertyTypeCustomizationUtils& StructCustomizationUtils ) override;
 	virtual void MakeHeaderRow(TSharedRef<class IPropertyHandle>& InStructPropertyHandle, FDetailWidgetRow& Row) override;
 
 private:
 
-	/**
-	 * Called when the value is changed in the property editor
-	 */
+	/** Called when the value is changed in the property editor. */
 	virtual void OnValueChanged();
 
 	ECheckBoxState GetForegroundCheckState() const;
@@ -33,4 +36,3 @@ private:
 	TSharedPtr<IPropertyHandle> ColorRuleHandle;
 	TSharedPtr<IPropertyHandle> SpecifiedColorHandle;
 };
-
diff --git a/Engine/Source/Editor/DetailCustomizations/Private/VectorStructCustomization.cpp b/Engine/Source/Editor/DetailCustomizations/Private/VectorStructCustomization.cpp
new file mode 100644
--- /dev/null
+++ b/Engine/Source/Editor/DetailCustomizations/Private/VectorStructCustomization.cpp
@@ -0,0 +1,47 @@
+// Copyright 1998-2015 Epic Games, Inc. All Rights Reserved.
+
+#include "DetailCustomizationsPrivatePCH.h"
+#include "VectorStructCustomization.h"
+
+
+TSharedRef<IPropertyTypeCustomization> FVectorStructCustomization::MakeInstance() 
+{
+	return MakeShareable(new FVectorStructCustomization);
+}
+
+
+void FVectorStructCustomization::GetSortedChildren(TSharedRef<IPropertyHandle> StructPropertyHandle, TArray< TSharedRef<IPropertyHandle> >& OutChildren)
+{
+	static const FName X("X");
+	static const FName Y("Y");
+	static const FName Z("Z");
+
+	TSharedPtr<IPropertyHandle> VectorChildren[3];
+
+	uint32 NumChildren;
+	StructPropertyHandle->GetNumChildren(NumChildren);
+
+	for (uint32 ChildIndex = 0; ChildIndex < NumChildren; ++ChildIndex)
+	{
+		TSharedRef<IPropertyHandle> ChildHandle = StructPropertyHandle->GetChildHandle(ChildIndex).ToSharedRef();
+		const FName PropertyName = ChildHandle->GetProperty()->GetFName();
+
+		if (PropertyName == X)
+		{
+			VectorChildren[0] = ChildHandle;
+		}
+		else if (PropertyName == Y)
+		{
+			VectorChildren[1] = ChildHandle;
+		}
+		else
+		{
+			check(PropertyName == Z);
+			VectorChildren[2] = ChildHandle;
+		}
+	}
+
+	OutChildren.Add(VectorChildren[0].ToSharedRef());
+	OutChildren.Add(VectorChildren[1].ToSharedRef());
+	OutChildren.Add(VectorChildren[2].ToSharedRef());
+}
diff --git a/Engine/Source/Editor/DetailCustomizations/Private/VectorStructCustomization.h b/Engine/Source/Editor/DetailCustomizations/Private/VectorStructCustomization.h
new file mode 100644
--- /dev/null
+++ b/Engine/Source/Editor/DetailCustomizations/Private/VectorStructCustomization.h
@@ -0,0 +1,23 @@
+// Copyright 1998-2015 Epic Games, Inc. All Rights Reserved.
+
+#pragma once
+
+#include "MathStructCustomizations.h"
+
+
+/**
+ * Customizes FVector structs.
+ */
+class FVectorStructCustomization
+	: public FMathStructCustomization
+{
+public:
+
+	/** @return A new instance of this class */
+	static TSharedRef<IPropertyTypeCustomization> MakeInstance();
+
+private:
+
+	/** FMathStructCustomization interface */
+	virtual void GetSortedChildren(TSharedRef<IPropertyHandle> StructPropertyHandle, TArray<TSharedRef<IPropertyHandle>>& OutChildren) override;
+};
diff --git a/Engine/Source/Editor/EditorStyle/Private/SlateEditorStyle.cpp b/Engine/Source/Editor/EditorStyle/Private/SlateEditorStyle.cpp
--- a/Engine/Source/Editor/EditorStyle/Private/SlateEditorStyle.cpp
+++ b/Engine/Source/Editor/EditorStyle/Private/SlateEditorStyle.cpp
@@ -1359,6 +1359,8 @@ void FSlateEditorStyle::FStyle::SetupGeneralStyles()
 		Set( "AssetEditor.SaveAsset.Greyscale", new IMAGE_BRUSH( "Icons/icon_file_save_16px", Icon16x16 ) );
 		Set( "AssetEditor.SaveAsset", new IMAGE_BRUSH( "Icons/icon_SaveAsset_40x", Icon40x40 ) );
 		Set( "AssetEditor.SaveAsset.Small", new IMAGE_BRUSH( "Icons/icon_SaveAsset_40x", Icon20x20 ) );
+		Set( "AssetEditor.SaveAssetAs", new IMAGE_BRUSH( "Icons/icon_file_saveas_40x", Icon40x40 ) );
+		Set( "AssetEditor.SaveAssetAs.Small", new IMAGE_BRUSH( "Icons/icon_file_saveas_40x", Icon20x20 ) );
 		Set( "AssetEditor.ReimportAsset", new IMAGE_BRUSH( "Icons/icon_TextureEd_Reimport_40x", Icon40x40 ) );
 		Set( "AssetEditor.ReimportAsset.Small", new IMAGE_BRUSH( "Icons/icon_TextureEd_Reimport_40x", Icon20x20 ) );
 	}
@@ -1599,12 +1601,14 @@ void FSlateEditorStyle::FStyle::SetupGeneralStyles()
 		Set( "Sequencer.KeyLeft", new IMAGE_BRUSH( "Common/KeyLeft", FVector2D( 7.0f, 12.0f ) ) );
 		Set( "Sequencer.KeyRight", new IMAGE_BRUSH( "Common/KeyRight", FVector2D( 7.0f, 12.0f ) ) );
 		Set( "Sequencer.PartialKey", new IMAGE_BRUSH( "Sequencer/PartialKey", FVector2D(11.f, 11.f) ) );
+		Set( "Sequencer.Star", new IMAGE_BRUSH( "Sequencer/Star", Icon12x12 ) );
+		Set( "Sequencer.Empty", new IMAGE_BRUSH( "Sequencer/Empty", Icon12x12 ) );
 
 		Set( "Sequencer.Timeline.ScrubHandleDown", new BOX_BRUSH( "Sequencer/ScrubHandleDown", FMargin( 6.f/13.f, 5/12.f, 6/13.f, 8/12.f ) ) );
 		Set( "Sequencer.Timeline.ScrubHandleUp", new BOX_BRUSH( "Sequencer/ScrubHandleUp", FMargin( 6.f/13.f, 8/12.f, 6/13.f, 5/12.f ) ) );
 		Set( "Sequencer.Timeline.ScrubHandleWhole", new BOX_BRUSH( "Sequencer/ScrubHandleWhole", FMargin( 6.f/13.f, 10/24.f, 6/13.f, 10/24.f  ) ) );
-		Set( "Sequencer.Timeline.RangeHandleLeft", new BOX_BRUSH( "Sequencer/SectionGripLeft", FMargin(5.f/16.f) ) );
-		Set( "Sequencer.Timeline.RangeHandleRight", new BOX_BRUSH( "Sequencer/SectionGripRight", FMargin(5.f/16.f) ) );
+		Set( "Sequencer.Timeline.RangeHandleLeft", new BOX_BRUSH( "Sequencer/GenericGripLeft", FMargin(5.f/16.f) ) );
+		Set( "Sequencer.Timeline.RangeHandleRight", new BOX_BRUSH( "Sequencer/GenericGripRight", FMargin(5.f/16.f) ) );
 		Set( "Sequencer.Timeline.RangeHandle", new BOX_BRUSH( "Sequencer/GenericSectionBackground", FMargin(5.f/16.f) ) );
 		Set( "Sequencer.Timeline.NotifyAlignmentMarker", new IMAGE_BRUSH( "Sequencer/NotifyAlignmentMarker", FVector2D(10,19) ) );
 		Set( "Sequencer.Timeline.PlayRange_Top_L", new BOX_BRUSH( "Sequencer/PlayRange_Top_L", FMargin(1.f, 0.5f, 0.f, 0.5f) ) );
@@ -1618,10 +1622,31 @@ void FSlateEditorStyle::FStyle::SetupGeneralStyles()
 
 		Set( "Sequencer.Save", new IMAGE_BRUSH( "Icons/SequencerIcons/icon_Sequencer_Save_24x", Icon40x40 ) );
 		Set( "Sequencer.Save.Small", new IMAGE_BRUSH( "Icons/SequencerIcons/icon_Sequencer_Save_24x", Icon20x20 ) );
-		Set( "Sequencer.GenericSection.Background", new BOX_BRUSH( "Sequencer/GenericSectionBackground", FMargin(5.f/16.f) ) );
-		Set( "Sequencer.DefaultSectionGripLeft", new BOX_BRUSH( "Sequencer/SectionGripLeft", FMargin(5.f/16.f) ) );
-		Set( "Sequencer.DefaultSectionGripRight", new BOX_BRUSH( "Sequencer/SectionGripRight", FMargin(5.f/16.f) ) );
+		Set( "Sequencer.GenericGripLeft", new BOX_BRUSH( "Sequencer/GenericGripLeft", FMargin(5.f/16.f) ) );
+		Set( "Sequencer.GenericGripRight", new BOX_BRUSH( "Sequencer/GenericGripRight", FMargin(5.f/16.f) ) );
 		Set( "Sequencer.SectionArea.Background", new FSlateColorBrush( FColor::White ) );
+
+		Set( "Sequencer.Section.Background", new BOX_BRUSH( TEXT("Sequencer/SectionBackground"), FMargin(3.f/100.f) ) );
+		Set( "Sequencer.Section.BackgroundTint", new BOX_BRUSH( TEXT("Sequencer/SectionBackgroundTint"), FMargin(48.f/100.f) ) );
+		Set( "Sequencer.Section.SelectedSectionOverlay", new IMAGE_BRUSH( TEXT("Sequencer/SelectedSectionOverlay"), Icon16x16, FLinearColor::White, ESlateBrushTileType::Both ) );
+		Set( "Sequencer.Section.SelectedTrackTint", new BOX_BRUSH( TEXT("Sequencer/SelectedTrackTint"), FMargin(0.f, 0.5f) ) );
+		Set( "Sequencer.Section.SelectionBorder", new BORDER_BRUSH( TEXT("Sequencer/SectionHighlight"), FMargin(7.f/16.f) ) );
+		Set( "Sequencer.Section.LockedBorder", new BORDER_BRUSH( TEXT("Sequencer/SectionLocked"), FMargin(7.f/16.f) ) );
+		Set( "Sequencer.Section.SelectedSectionOverlay", new IMAGE_BRUSH( TEXT("Sequencer/SelectedSectionOverlay"), Icon16x16, FLinearColor::White, ESlateBrushTileType::Both ) );
+		Set( "Sequencer.Section.FilmBorder", new IMAGE_BRUSH( TEXT("Sequencer/SectionFilmBorder"), FVector2D(10, 7), FLinearColor::White, ESlateBrushTileType::Horizontal ) );
+		Set( "Sequencer.Section.GripLeft", new BOX_BRUSH( "Sequencer/SectionGripLeft", FMargin(5.f/16.f) ) );
+		Set( "Sequencer.Section.GripRight", new BOX_BRUSH( "Sequencer/SectionGripRight", FMargin(5.f/16.f) ) );
+
+		Set( "Sequencer.Section.StripeOverlay", new BOX_BRUSH( "Sequencer/SectionStripeOverlay", FMargin(0.f, .5f) ) );
+
+		Set( "Sequencer.AnimationOutliner.ColorStrip", FButtonStyle()
+			.SetNormal(FSlateNoResource())
+			.SetHovered(FSlateNoResource())
+			.SetPressed(FSlateNoResource())
+			.SetNormalPadding(FMargin(0,0,0,0))
+			.SetPressedPadding(FMargin(0,0,0,0))
+		);
+
 		Set( "Sequencer.AnimationOutliner.TopLevelBorder_Expanded", new BOX_BRUSH( "Sequencer/TopLevelNodeBorder_Expanded", FMargin(4.0f/16.0f) ) );
 		Set( "Sequencer.AnimationOutliner.TopLevelBorder_Collapsed", new BOX_BRUSH( "Sequencer/TopLevelNodeBorder_Collapsed", FMargin(4.0f/16.0f) ) );
 		Set( "Sequencer.AnimationOutliner.DefaultBorder", new FSlateColorBrush( FLinearColor::White ) );
@@ -1650,12 +1675,10 @@ void FSlateEditorStyle::FStyle::SetupGeneralStyles()
 		Set( "Sequencer.RenderMovie.Small", new IMAGE_BRUSH( "Sequencer/RenderMovie_16px", Icon16x16 ) );
 		Set( "Sequencer.LockCamera", new IMAGE_BRUSH( "Icons/icon_ToggleActorPilotCameraView_16x", Icon16x16 ) );
 		Set( "Sequencer.UnlockCamera", new IMAGE_BRUSH( "Icons/icon_ToggleActorPilotCameraView_16x", Icon16x16, FLinearColor(1.f, 1.f, 1.f, 0.5f) ) );
-		Set( "Sequencer.Section.SelectionBorder", new BORDER_BRUSH( TEXT("Sequencer/SectionHighlight"), FMargin(7.f/16.f) ) );
-		Set( "Sequencer.Section.LockedBorder", new BORDER_BRUSH( TEXT("Sequencer/SectionLocked"), FMargin(7.f/16.f) ) );
-		Set( "Sequencer.Thumbnail.SectionHandle", new IMAGE_BRUSH( "Old/White", Icon16x16, FLinearColor::Black ) );
 		Set( "Sequencer.TrackHoverHighlight_Top", new IMAGE_BRUSH( TEXT("Sequencer/TrackHoverHighlight_Top"), FVector2D(4, 4) ) );
 		Set( "Sequencer.TrackHoverHighlight_Bottom", new IMAGE_BRUSH( TEXT("Sequencer/TrackHoverHighlight_Bottom"), FVector2D(4, 4) ) );
 		Set( "Sequencer.SpawnableIconOverlay", new IMAGE_BRUSH( TEXT("Sequencer/SpawnableIconOverlay"), FVector2D(13, 13) ) );
+		Set( "Sequencer.DiscardChanges.Small", new IMAGE_BRUSH("Icons/icon_Delete_40x", Icon16x16));
 
 		Set( "Sequencer.General", new IMAGE_BRUSH( "Icons/SequencerIcons/icon_Sequencer_View_24x", Icon24x24 ) );
 		Set( "Sequencer.General.Small", new IMAGE_BRUSH( "Icons/SequencerIcons/icon_Sequencer_View_24x", Icon20x20 ) );
@@ -1673,6 +1696,7 @@ void FSlateEditorStyle::FStyle::SetupGeneralStyles()
 		Set( "Sequencer.Tracks.CinematicShot", new IMAGE_BRUSH("Icons/SequencerIcons/icon_Sequencer_CinematicShot_16x", Icon16x16));
 		Set( "Sequencer.Tracks.Slomo", new IMAGE_BRUSH("Icons/SequencerIcons/icon_Sequencer_Slomo_16x", Icon16x16));
 		Set( "Sequencer.Tracks.Sub", new IMAGE_BRUSH("Icons/SequencerIcons/icon_Sequencer_Sub_16x", Icon16x16));
+		Set( "Sequencer.Tracks.LevelVisibility", new IMAGE_BRUSH("Icons/SequencerIcons/icon_Sequencer_LevelVisibility_16x", Icon16x16));
 
 		Set( "Sequencer.CursorDecorator_MarqueeAdd", new IMAGE_BRUSH( "Sequencer/CursorDecorator_MarqueeAdd", Icon16x16));
 		Set( "Sequencer.CursorDecorator_MarqueeSubtract", new IMAGE_BRUSH( "Sequencer/CursorDecorator_MarqueeSubtract", Icon16x16));
diff --git a/Engine/Source/Editor/GameplayAbilitiesEditor/Private/AttributeDetails.cpp b/Engine/Source/Editor/GameplayAbilitiesEditor/Private/AttributeDetails.cpp
--- a/Engine/Source/Editor/GameplayAbilitiesEditor/Private/AttributeDetails.cpp
+++ b/Engine/Source/Editor/GameplayAbilitiesEditor/Private/AttributeDetails.cpp
@@ -303,113 +303,98 @@ void FScalableFloatDetails::CustomizeHeader( TSharedRef<class IPropertyHandle> S
 	StructPropertyHandle->GetNumChildren(NumChildren);
 
 	ValueProperty = StructPropertyHandle->GetChildHandle(GET_MEMBER_NAME_CHECKED(FScalableFloat,Value));
-	if (ValueProperty.IsValid() == false)
-	{
-		return;
-	}
-
 	CurveTableHandleProperty = StructPropertyHandle->GetChildHandle(GET_MEMBER_NAME_CHECKED(FScalableFloat,Curve));
-	if (CurveTableHandleProperty.IsValid() == false)
-	{
-		return;
-	}
 
-	RowNameProperty = CurveTableHandleProperty->GetChildHandle(GET_MEMBER_NAME_CHECKED(FCurveTableRowHandle, RowName));
-	if (RowNameProperty.IsValid() == false)
+	if(ValueProperty.IsValid() && CurveTableHandleProperty.IsValid())
 	{
-		return;
-	}
+		RowNameProperty = CurveTableHandleProperty->GetChildHandle(GET_MEMBER_NAME_CHECKED(FCurveTableRowHandle, RowName));
+		CurveTableProperty = CurveTableHandleProperty->GetChildHandle(GET_MEMBER_NAME_CHECKED(FCurveTableRowHandle, CurveTable));
 
-	CurveTableProperty = CurveTableHandleProperty->GetChildHandle(GET_MEMBER_NAME_CHECKED(FCurveTableRowHandle, CurveTable));
-	if (CurveTableProperty.IsValid() == false)
-	{
-		return;
-	}
+		CurrentSelectedItem = InitWidgetContent();
 
-	CurrentSelectedItem = InitWidgetContent();
-
-	FSimpleDelegate OnCurveTableChangedDelegate = FSimpleDelegate::CreateSP(this, &FScalableFloatDetails::OnCurveTableChanged);
-	CurveTableProperty->SetOnPropertyValueChanged(OnCurveTableChangedDelegate);
+		FSimpleDelegate OnCurveTableChangedDelegate = FSimpleDelegate::CreateSP(this, &FScalableFloatDetails::OnCurveTableChanged);
+		CurveTableProperty->SetOnPropertyValueChanged(OnCurveTableChangedDelegate);
 
-	HeaderRow
-		.NameContent()
-		[
-			StructPropertyHandle->CreatePropertyNameWidget()
-		]
-		.ValueContent()
-		.MinDesiredWidth( 600 )
-		.MaxDesiredWidth( 4096 )
-		[
-			SNew(SHorizontalBox)
-			.IsEnabled(TAttribute<bool>::Create(TAttribute<bool>::FGetter::CreateSP(this, &FScalableFloatDetails::IsEditable)))
-			
-			+SHorizontalBox::Slot()
-			.FillWidth(0.12f)
-			.HAlign(HAlign_Fill)
-			.Padding(0.f, 0.f, 2.f, 0.f)
-			[
-				ValueProperty->CreatePropertyValueWidget()
-			]
-		
-			+SHorizontalBox::Slot()
-			.FillWidth(0.40f)
-			.HAlign(HAlign_Fill)
-			.Padding(2.f, 0.f, 2.f, 0.f)
+		HeaderRow
+			.NameContent()
 			[
-				CurveTableProperty->CreatePropertyValueWidget()
+				StructPropertyHandle->CreatePropertyNameWidget()
 			]
-
-			+SHorizontalBox::Slot()
-			.FillWidth(0.23f)
-			.HAlign(HAlign_Fill)
-			.Padding(2.f, 0.f, 2.f, 0.f)
+			.ValueContent()
+			.MinDesiredWidth( 600 )
+			.MaxDesiredWidth( 4096 )
 			[
-				SAssignNew(RowNameComboButton, SComboButton)
-				.OnGetMenuContent(this, &FScalableFloatDetails::GetListContent)
-				.ContentPadding(FMargin(2.0f, 2.0f))
-				.Visibility(this, &FScalableFloatDetails::GetRowNameVisibility)
-				.ButtonContent()
+				SNew(SHorizontalBox)
+				.IsEnabled(TAttribute<bool>::Create(TAttribute<bool>::FGetter::CreateSP(this, &FScalableFloatDetails::IsEditable)))
+			
+				+SHorizontalBox::Slot()
+				.FillWidth(0.12f)
+				.HAlign(HAlign_Fill)
+				.Padding(0.f, 0.f, 2.f, 0.f)
 				[
-					SNew(STextBlock)
-					.Text(this, &FScalableFloatDetails::GetRowNameComboBoxContentText)
+					ValueProperty->CreatePropertyValueWidget()
+				]
+		
+				+SHorizontalBox::Slot()
+				.FillWidth(0.40f)
+				.HAlign(HAlign_Fill)
+				.Padding(2.f, 0.f, 2.f, 0.f)
+				[
+					CurveTableProperty->CreatePropertyValueWidget()
 				]
-			]
 
-			+SHorizontalBox::Slot()
-			.FillWidth(0.15f)
-			.HAlign(HAlign_Fill)
-			.Padding(2.f, 0.f, 2.f, 0.f)
-			[
-				SNew(SVerticalBox)
-				.Visibility(this, &FScalableFloatDetails::GetPreviewVisibility)
-				
-				+SVerticalBox::Slot()
-				.HAlign(HAlign_Center)
+				+SHorizontalBox::Slot()
+				.FillWidth(0.23f)
+				.HAlign(HAlign_Fill)
+				.Padding(2.f, 0.f, 2.f, 0.f)
 				[
-					SNew(STextBlock)
-					.Text(this, &FScalableFloatDetails::GetRowValuePreviewLabel)
+					SAssignNew(RowNameComboButton, SComboButton)
+					.OnGetMenuContent(this, &FScalableFloatDetails::GetListContent)
+					.ContentPadding(FMargin(2.0f, 2.0f))
+					.Visibility(this, &FScalableFloatDetails::GetRowNameVisibility)
+					.ButtonContent()
+					[
+						SNew(STextBlock)
+						.Text(this, &FScalableFloatDetails::GetRowNameComboBoxContentText)
+					]
 				]
 
-				+SVerticalBox::Slot()
-				.HAlign(HAlign_Center)
+				+SHorizontalBox::Slot()
+				.FillWidth(0.15f)
+				.HAlign(HAlign_Fill)
+				.Padding(2.f, 0.f, 2.f, 0.f)
 				[
-					SNew(STextBlock)
-					.Text(this, &FScalableFloatDetails::GetRowValuePreviewText)
+					SNew(SVerticalBox)
+					.Visibility(this, &FScalableFloatDetails::GetPreviewVisibility)
+				
+					+SVerticalBox::Slot()
+					.HAlign(HAlign_Center)
+					[
+						SNew(STextBlock)
+						.Text(this, &FScalableFloatDetails::GetRowValuePreviewLabel)
+					]
+
+					+SVerticalBox::Slot()
+					.HAlign(HAlign_Center)
+					[
+						SNew(STextBlock)
+						.Text(this, &FScalableFloatDetails::GetRowValuePreviewText)
+					]
 				]
-			]
 
-			+SHorizontalBox::Slot()
-			.FillWidth(0.1f)
-			.HAlign(HAlign_Fill)
-			.Padding(2.f, 0.f, 0.f, 0.f)
-			[
-				SNew(SSlider)
-				.Visibility(this, &FScalableFloatDetails::GetPreviewVisibility)
-				.ToolTipText(LOCTEXT("LevelPreviewToolTip", "Adjust the preview level."))
-				.Value(this, &FScalableFloatDetails::GetPreviewLevel)
-				.OnValueChanged(this, &FScalableFloatDetails::SetPreviewLevel)
-			]
-		];	
+				+SHorizontalBox::Slot()
+				.FillWidth(0.1f)
+				.HAlign(HAlign_Fill)
+				.Padding(2.f, 0.f, 0.f, 0.f)
+				[
+					SNew(SSlider)
+					.Visibility(this, &FScalableFloatDetails::GetPreviewVisibility)
+					.ToolTipText(LOCTEXT("LevelPreviewToolTip", "Adjust the preview level."))
+					.Value(this, &FScalableFloatDetails::GetPreviewLevel)
+					.OnValueChanged(this, &FScalableFloatDetails::SetPreviewLevel)
+				]
+			];	
+	}
 }
 
 END_SLATE_FUNCTION_BUILD_OPTIMIZATION
diff --git a/Engine/Source/Editor/LevelEditor/Private/SLevelEditor.cpp b/Engine/Source/Editor/LevelEditor/Private/SLevelEditor.cpp
--- a/Engine/Source/Editor/LevelEditor/Private/SLevelEditor.cpp
+++ b/Engine/Source/Editor/LevelEditor/Private/SLevelEditor.cpp
@@ -514,7 +514,12 @@ void SLevelEditor::AttachSequencer( TSharedPtr<SWidget> SequencerWidget, TShared
 	{
 		static void OnSequencerClosed( TSharedRef<SDockTab> DockTab, TWeakPtr<IAssetEditorInstance> InSequencerAssetEditor )
 		{
-			InSequencerAssetEditor.Pin()->CloseWindow();
+			TSharedPtr<IAssetEditorInstance> AssetEditorInstance = InSequencerAssetEditor.Pin();
+
+			if (AssetEditorInstance.IsValid())
+			{
+				InSequencerAssetEditor.Pin()->CloseWindow();
+			}
 		}
 	};
 
diff --git a/Engine/Source/Editor/LevelEditor/Private/SLevelViewport.cpp b/Engine/Source/Editor/LevelEditor/Private/SLevelViewport.cpp
--- a/Engine/Source/Editor/LevelEditor/Private/SLevelViewport.cpp
+++ b/Engine/Source/Editor/LevelEditor/Private/SLevelViewport.cpp
@@ -3771,6 +3771,18 @@ bool SLevelViewport::GetCameraInformationFromActor(AActor* Actor, FMinimalViewIn
 		}
 	}
 
+	// see if any actors are attached to us, directly or indirectly, that have an active camera component we might want to use
+	// #note: assumption here that attachment cannot be circular
+	TArray<AActor*> AttachedActors;
+	Actor->GetAttachedActors(AttachedActors);
+	for (AActor* AttachedActor : AttachedActors)
+	{
+		if (GetCameraInformationFromActor(AttachedActor, out_CameraInfo))
+		{
+			return true;
+		}
+	}
+
 	// no active cameras
 	return false;
 }
diff --git a/Engine/Source/Editor/MaterialEditor/Private/MaterialEditor.cpp b/Engine/Source/Editor/MaterialEditor/Private/MaterialEditor.cpp
--- a/Engine/Source/Editor/MaterialEditor/Private/MaterialEditor.cpp
+++ b/Engine/Source/Editor/MaterialEditor/Private/MaterialEditor.cpp
@@ -864,6 +864,18 @@ bool FMaterialEditor::ApproveSetPreviewAsset(UObject* InAsset)
 	return bApproved;
 }
 
+void FMaterialEditor::GetSaveableObjects(TArray<UObject*>& OutObjects) const
+{
+	if ((MaterialFunction != nullptr) && MaterialFunction->ParentFunction)
+	{
+		OutObjects.Add(MaterialFunction->ParentFunction);
+	}
+	else
+	{
+		OutObjects.Add(OriginalMaterial);
+	}
+}
+
 void FMaterialEditor::SaveAsset_Execute()
 {
 	UE_LOG(LogMaterialEditor, Log, TEXT("Saving and Compiling material %s"), *GetEditingObjects()[0]->GetName());
@@ -873,20 +885,19 @@ void FMaterialEditor::SaveAsset_Execute()
 		UpdateOriginalMaterial();
 	}
 
-	UPackage* Package = OriginalMaterial->GetOutermost();
+	IMaterialEditor::SaveAsset_Execute();
+}
 
-	if (MaterialFunction != NULL && MaterialFunction->ParentFunction)
+void FMaterialEditor::SaveAssetAs_Execute()
+{
+	UE_LOG(LogMaterialEditor, Log, TEXT("Saving and Compiling material %s"), *GetEditingObjects()[0]->GetName());
+
+	if (bMaterialDirty)
 	{
-		Package = MaterialFunction->ParentFunction->GetOutermost();
+		UpdateOriginalMaterial();
 	}
-	
-	if (Package)
-	{
-		TArray<UPackage*> PackagesToSave;
-		PackagesToSave.Add(Package);
 
-		FEditorFileUtils::PromptForCheckoutAndSave(PackagesToSave, false, false);
-	}
+	IMaterialEditor::SaveAssetAs_Execute();
 }
 
 bool FMaterialEditor::OnRequestClose()
diff --git a/Engine/Source/Editor/MaterialEditor/Private/MaterialEditor.h b/Engine/Source/Editor/MaterialEditor/Private/MaterialEditor.h
--- a/Engine/Source/Editor/MaterialEditor/Private/MaterialEditor.h
+++ b/Engine/Source/Editor/MaterialEditor/Private/MaterialEditor.h
@@ -366,12 +366,13 @@ public:
 	UMaterialEditorOptions* EditorOptions;
 	
 protected:
-	/** Called when "Save" is clicked for this asset */
+	//~ FAssetEditorToolkit interface
+	virtual void GetSaveableObjects(TArray<UObject*>& OutObjects) const override;
 	virtual void SaveAsset_Execute() override;
-	
-	/** Called when this toolkit would close */
+	virtual void SaveAssetAs_Execute() override;
 	virtual bool OnRequestClose() override;
 
+protected:
 	/** Called when the selection changes in the GraphEditor */
 	void OnSelectedNodesChanged(const TSet<class UObject*>& NewSelection);
 
diff --git a/Engine/Source/Editor/MovieSceneTools/Private/CurveKeyAreas/FloatCurveKeyArea.cpp b/Engine/Source/Editor/MovieSceneTools/Private/CurveKeyAreas/FloatCurveKeyArea.cpp
--- a/Engine/Source/Editor/MovieSceneTools/Private/CurveKeyAreas/FloatCurveKeyArea.cpp
+++ b/Engine/Source/Editor/MovieSceneTools/Private/CurveKeyAreas/FloatCurveKeyArea.cpp
@@ -105,7 +105,7 @@ void FFloatCurveKeyArea::DeleteKey(FKeyHandle KeyHandle)
 }
 
 
-FLinearColor FFloatCurveKeyArea::GetColor()
+TOptional<FLinearColor> FFloatCurveKeyArea::GetColor()
 {
 	return Color;
 }
diff --git a/Engine/Source/Editor/MovieSceneTools/Private/CurveKeyAreas/IntegralKeyArea.cpp b/Engine/Source/Editor/MovieSceneTools/Private/CurveKeyAreas/IntegralKeyArea.cpp
--- a/Engine/Source/Editor/MovieSceneTools/Private/CurveKeyAreas/IntegralKeyArea.cpp
+++ b/Engine/Source/Editor/MovieSceneTools/Private/CurveKeyAreas/IntegralKeyArea.cpp
@@ -47,9 +47,9 @@ void FIntegralCurveKeyAreaBase::DeleteKey(FKeyHandle KeyHandle)
 }
 
 
-FLinearColor FIntegralCurveKeyAreaBase::GetColor()
+TOptional<FLinearColor> FIntegralCurveKeyAreaBase::GetColor()
 {
-	return FLinearColor(0.1f, 0.1f, 0.1f, 0.7f);
+	return TOptional<FLinearColor>();
 }
 
 
diff --git a/Engine/Source/Editor/MovieSceneTools/Private/CurveKeyAreas/NameCurveKeyArea.cpp b/Engine/Source/Editor/MovieSceneTools/Private/CurveKeyAreas/NameCurveKeyArea.cpp
--- a/Engine/Source/Editor/MovieSceneTools/Private/CurveKeyAreas/NameCurveKeyArea.cpp
+++ b/Engine/Source/Editor/MovieSceneTools/Private/CurveKeyAreas/NameCurveKeyArea.cpp
@@ -59,9 +59,9 @@ void FNameCurveKeyArea::DeleteKey(FKeyHandle KeyHandle)
 }
 
 
-FLinearColor FNameCurveKeyArea::GetColor()
+TOptional<FLinearColor> FNameCurveKeyArea::GetColor()
 {
-	return FLinearColor(0.1f, 0.1f, 0.1f, 0.7f);
+	return TOptional<FLinearColor>();
 }
 
 
diff --git a/Engine/Source/Editor/MovieSceneTools/Private/MovieSceneToolHelpers.cpp b/Engine/Source/Editor/MovieSceneTools/Private/MovieSceneToolHelpers.cpp
--- a/Engine/Source/Editor/MovieSceneTools/Private/MovieSceneToolHelpers.cpp
+++ b/Engine/Source/Editor/MovieSceneTools/Private/MovieSceneToolHelpers.cpp
@@ -208,7 +208,7 @@ FString MovieSceneToolHelpers::GenerateNewShotName(const TArray<UMovieSceneSecti
 	return ComposeShotName(ProjectSettings->ShotPrefix, ProjectSettings->FirstShotNumber, ProjectSettings->FirstTakeNumber);
 }
 
-void MovieSceneToolHelpers::GatherTakes(const UMovieSceneSection* Section, TArray<uint32>& TakeNumbers)
+void MovieSceneToolHelpers::GatherTakes(const UMovieSceneSection* Section, TArray<uint32>& TakeNumbers, uint32& CurrentTakeNumber)
 {
 	const UMovieSceneCinematicShotSection* Shot = Cast<const UMovieSceneCinematicShotSection>(Section);
 	
@@ -226,9 +226,9 @@ void MovieSceneToolHelpers::GatherTakes(const UMovieSceneSection* Section, TArra
 
 	FString ShotPrefix;
 	uint32 ShotNumber = INDEX_NONE;
-	uint32 TakeNumber = INDEX_NONE;
+	CurrentTakeNumber = INDEX_NONE;
 
-	if (ParseShotName(Shot->GetShotDisplayName().ToString(), ShotPrefix, ShotNumber, TakeNumber))
+	if (ParseShotName(Shot->GetShotDisplayName().ToString(), ShotPrefix, ShotNumber, CurrentTakeNumber))
 	{
 		// Gather up all level sequence assets
 		FAssetRegistryModule& AssetRegistryModule = FModuleManager::LoadModuleChecked<FAssetRegistryModule>(TEXT("AssetRegistry"));
diff --git a/Engine/Source/Editor/MovieSceneTools/Private/MovieSceneToolsModule.cpp b/Engine/Source/Editor/MovieSceneTools/Private/MovieSceneToolsModule.cpp
--- a/Engine/Source/Editor/MovieSceneTools/Private/MovieSceneToolsModule.cpp
+++ b/Engine/Source/Editor/MovieSceneTools/Private/MovieSceneToolsModule.cpp
@@ -32,6 +32,7 @@
 #include "MaterialTrackEditor.h"
 #include "FadeTrackEditor.h"
 #include "SpawnTrackEditor.h"
+#include "LevelVisibilityTrackEditor.h"
 
 #include "MovieSceneClipboard.h"
 #include "SequencerClipboardReconciler.h"
@@ -88,6 +89,7 @@ public:
 		ComponentMaterialTrackCreateEditorHandle = SequencerModule.RegisterTrackEditor_Handle( FOnCreateTrackEditor::CreateStatic( &FComponentMaterialTrackEditor::CreateTrackEditor ) );
 		FadeTrackCreateEditorHandle = SequencerModule.RegisterTrackEditor_Handle( FOnCreateTrackEditor::CreateStatic( &FFadeTrackEditor::CreateTrackEditor ) );
 		SpawnTrackCreateEditorHandle = SequencerModule.RegisterTrackEditor_Handle( FOnCreateTrackEditor::CreateStatic( &FSpawnTrackEditor::CreateTrackEditor ) );
+		LevelVisibilityTrackCreateEditorHandle = SequencerModule.RegisterTrackEditor_Handle( FOnCreateTrackEditor::CreateStatic( &FLevelVisibilityTrackEditor::CreateTrackEditor ) );
 
 		RegisterClipboardConversions();
 	}
@@ -130,6 +132,7 @@ public:
 		SequencerModule.UnRegisterTrackEditor_Handle( ComponentMaterialTrackCreateEditorHandle );
 		SequencerModule.UnRegisterTrackEditor_Handle( FadeTrackCreateEditorHandle );
 		SequencerModule.UnRegisterTrackEditor_Handle( SpawnTrackCreateEditorHandle );
+		SequencerModule.UnRegisterTrackEditor_Handle( LevelVisibilityTrackCreateEditorHandle );
 	}
 
 	void RegisterClipboardConversions()
@@ -191,6 +194,7 @@ private:
 	FDelegateHandle ComponentMaterialTrackCreateEditorHandle;
 	FDelegateHandle FadeTrackCreateEditorHandle;
 	FDelegateHandle SpawnTrackCreateEditorHandle;
+	FDelegateHandle LevelVisibilityTrackCreateEditorHandle;
 };
 
 
diff --git a/Engine/Source/Editor/MovieSceneTools/Private/Sections/BoolPropertySection.cpp b/Engine/Source/Editor/MovieSceneTools/Private/Sections/BoolPropertySection.cpp
--- a/Engine/Source/Editor/MovieSceneTools/Private/Sections/BoolPropertySection.cpp
+++ b/Engine/Source/Editor/MovieSceneTools/Private/Sections/BoolPropertySection.cpp
@@ -12,14 +12,8 @@ void FBoolPropertySection::GenerateSectionLayout( class ISectionLayoutBuilder& L
 	LayoutBuilder.SetSectionAsKeyArea( KeyArea.ToSharedRef() );
 }
 
-int32 FBoolPropertySection::OnPaintSection( const FGeometry& AllottedGeometry, const FSlateRect& SectionClippingRect, FSlateWindowElementList& OutDrawElements, int32 LayerId, bool bParentEnabled ) const
+int32 FBoolPropertySection::OnPaintSection( FSequencerSectionPainter& Painter ) const
 {
-	const ESlateDrawEffect::Type DrawEffects = bParentEnabled ? ESlateDrawEffect::None : ESlateDrawEffect::DisabledEffect;
-
-	FTimeToPixel TimeToPixelConverter = SectionObject.IsInfinite() ?
-		FTimeToPixel( AllottedGeometry, Sequencer->GetViewRange() ) :
-		FTimeToPixel( AllottedGeometry, TRange<float>( SectionObject.GetStartTime(), SectionObject.GetEndTime() ) );
-
 	// custom drawing for bool curves
 	UMovieSceneBoolSection* BoolSection = Cast<UMovieSceneBoolSection>( &SectionObject );
 
@@ -47,36 +41,38 @@ int32 FBoolPropertySection::OnPaintSection( const FGeometry& AllottedGeometry, c
 	}
 	SectionSwitchTimes.Add( SectionObject.GetEndTime() );
 
+
+	int32 LayerId = Painter.PaintSectionBackground();
+
+	const ESlateDrawEffect::Type DrawEffects = Painter.bParentEnabled
+		? ESlateDrawEffect::None
+		: ESlateDrawEffect::DisabledEffect;
+
+	const FTimeToPixel& TimeConverter = Painter.GetTimeConverter();
+
+	static const int32 Height = 5;
+	const float VerticalOffset = Painter.SectionGeometry.GetLocalSize().Y * .5f - Height * .5f;
+
+	const FSlateBrush* BoolOverlayBrush = FEditorStyle::GetBrush("Sequencer.Section.StripeOverlay");
+
 	for ( int32 i = 0; i < SectionSwitchTimes.Num() - 1; ++i )
 	{
-		float FirstTime = SectionSwitchTimes[i];
-		float NextTime = SectionSwitchTimes[i + 1];
-		float FirstTimeInPixels = TimeToPixelConverter.TimeToPixel( FirstTime );
-		float NextTimeInPixels = TimeToPixelConverter.TimeToPixel( NextTime );
-		if ( BoolSection->Eval( FirstTime ) )
-		{
-			FSlateDrawElement::MakeBox(
-				OutDrawElements,
-				LayerId,
-				AllottedGeometry.ToPaintGeometry( FVector2D( FirstTimeInPixels, 0 ), FVector2D( NextTimeInPixels - FirstTimeInPixels, AllottedGeometry.Size.Y ) ),
-				FEditorStyle::GetBrush( "Sequencer.GenericSection.Background" ),
-				SectionClippingRect,
-				DrawEffects,
-				FLinearColor( 0.f, 1.f, 0.f, 1.f )
-				);
-		}
-		else
-		{
-			FSlateDrawElement::MakeBox(
-				OutDrawElements,
-				LayerId,
-				AllottedGeometry.ToPaintGeometry( FVector2D( FirstTimeInPixels, 0 ), FVector2D( NextTimeInPixels - FirstTimeInPixels, AllottedGeometry.Size.Y ) ),
-				FEditorStyle::GetBrush( "Sequencer.GenericSection.Background" ),
-				SectionClippingRect,
-				DrawEffects,
-				FLinearColor( 1.f, 0.f, 0.f, 1.f )
-				);
-		}
+		float ThisTime = SectionSwitchTimes[i];
+
+		const FColor Color = BoolSection->Eval(ThisTime) ? FColor(0, 255, 0, 127) : FColor(255, 0, 0, 127);
+		
+		FVector2D StartPos(TimeConverter.TimeToPixel(ThisTime), VerticalOffset);
+		FVector2D Size(TimeConverter.TimeToPixel(SectionSwitchTimes[i+1]) - StartPos.X, Height);
+
+		FSlateDrawElement::MakeBox(
+			Painter.DrawElements,
+			LayerId + 1,
+			Painter.SectionGeometry.ToPaintGeometry(StartPos, Size),
+			BoolOverlayBrush,
+			Painter.SectionClippingRect,
+			DrawEffects,
+			Color
+			);
 	}
 
 	return LayerId + 1;
diff --git a/Engine/Source/Editor/MovieSceneTools/Private/Sections/BoolPropertySection.h b/Engine/Source/Editor/MovieSceneTools/Private/Sections/BoolPropertySection.h
--- a/Engine/Source/Editor/MovieSceneTools/Private/Sections/BoolPropertySection.h
+++ b/Engine/Source/Editor/MovieSceneTools/Private/Sections/BoolPropertySection.h
@@ -22,7 +22,7 @@ public:
 	// FPropertySection interface
 
 	virtual void GenerateSectionLayout( class ISectionLayoutBuilder& LayoutBuilder ) const override;
-	virtual int32 OnPaintSection( const FGeometry& AllottedGeometry, const FSlateRect& SectionClippingRect, FSlateWindowElementList& OutDrawElements, int32 LayerId, bool bParentEnabled ) const override;
+	virtual int32 OnPaintSection( FSequencerSectionPainter& Painter ) const override;
 	virtual void SetIntermediateValue( FPropertyChangedParams PropertyChangedParams ) override;
 	virtual void ClearIntermediateValue() override;
 
diff --git a/Engine/Source/Editor/MovieSceneTools/Private/Sections/CameraCutSection.cpp b/Engine/Source/Editor/MovieSceneTools/Private/Sections/CameraCutSection.cpp
--- a/Engine/Source/Editor/MovieSceneTools/Private/Sections/CameraCutSection.cpp
+++ b/Engine/Source/Editor/MovieSceneTools/Private/Sections/CameraCutSection.cpp
@@ -42,9 +42,13 @@ void FCameraCutSection::BuildSectionContextMenu(FMenuBuilder& MenuBuilder, const
 	{
 		AActor* Actor = *ActorIt;
 
-		if ((Actor != GetCameraObject()) && Actor->IsListedInSceneOutliner() && Actor->IsA<ACameraActor>())
+		if ((Actor != GetCameraObject()) && Actor->IsListedInSceneOutliner())
 		{
-			AllCameras.Add(Actor);
+			UCameraComponent* CameraComponent = MovieSceneHelpers::CameraComponentFromActor(Actor);
+			if (CameraComponent)
+			{
+				AllCameras.Add(Actor);
+			}
 		}
 	}
 
@@ -80,7 +84,7 @@ FText FCameraCutSection::GetDisplayName() const
 /* FThumbnailSection interface
  *****************************************************************************/
 
-ACameraActor* FCameraCutSection::GetCameraObject() const
+AActor* FCameraCutSection::GetCameraObject() const
 {
 	UMovieSceneCameraCutSection* CameraCutSection = Cast<UMovieSceneCameraCutSection>(Section);
 
@@ -88,11 +92,11 @@ ACameraActor* FCameraCutSection::GetCameraObject() const
 	// @todo sequencer: the director track may be able to get cameras from sub-movie scenes
 	SequencerPtr.Pin()->GetRuntimeObjects(SequencerPtr.Pin()->GetRootMovieSceneSequenceInstance(),  CameraCutSection->GetCameraGuid(), OutObjects);
 
-	ACameraActor* ReturnCam = nullptr;
+	AActor* ReturnCam = nullptr;
 
 	if (OutObjects.Num() > 0)
 	{
-		ReturnCam = Cast<ACameraActor>(OutObjects[0]);
+		ReturnCam = Cast<AActor>(OutObjects[0]);
 
 		if (ReturnCam != nullptr)
 		{
@@ -106,7 +110,7 @@ ACameraActor* FCameraCutSection::GetCameraObject() const
 
 	if (OutObjects.Num() > 0)
 	{
-		ReturnCam = Cast<ACameraActor>(OutObjects[0]);
+		ReturnCam = Cast<AActor>(OutObjects[0]);
 
 		if (ReturnCam != nullptr)
 		{
@@ -117,9 +121,27 @@ ACameraActor* FCameraCutSection::GetCameraObject() const
 	return ReturnCam;
 }
 
+float FCameraCutSection::GetSectionHeight() const
+{
+	return FThumbnailSection::GetSectionHeight() + 10.f;
+}
+
+FMargin FCameraCutSection::GetContentPadding() const
+{
+	return FMargin(6.f, 12.f);
+}
+
+int32 FCameraCutSection::OnPaintSection(FSequencerSectionPainter& InPainter) const
+{
+	static const FSlateBrush* FilmBorder = FEditorStyle::GetBrush("Sequencer.Section.FilmBorder");
+
+	InPainter.LayerId = InPainter.PaintSectionBackground();
+	return FThumbnailSection::OnPaintSection(InPainter);
+}
+
 FText FCameraCutSection::HandleThumbnailTextBlockText() const
 {
-	ACameraActor* CameraActor = GetCameraObject();
+	AActor* CameraActor = GetCameraObject();
 	if (CameraActor)
 	{
 		return FText::FromString(CameraActor->GetActorLabel());
diff --git a/Engine/Source/Editor/MovieSceneTools/Private/Sections/CameraCutSection.h b/Engine/Source/Editor/MovieSceneTools/Private/Sections/CameraCutSection.h
--- a/Engine/Source/Editor/MovieSceneTools/Private/Sections/CameraCutSection.h
+++ b/Engine/Source/Editor/MovieSceneTools/Private/Sections/CameraCutSection.h
@@ -30,10 +30,12 @@ public:
 
 	virtual void BuildSectionContextMenu(FMenuBuilder& MenuBuilder, const FGuid& ObjectBinding) override;
 	virtual FText GetDisplayName() const override;
-
+	virtual float GetSectionHeight() const override;
+	virtual int32 OnPaintSection(FSequencerSectionPainter& InPainter) const override;
+	virtual FMargin GetContentPadding() const override;
 	// FThumbnail interface
 
-	virtual ACameraActor* GetCameraObject() const override;
+	virtual AActor* GetCameraObject() const override;
 	virtual FText HandleThumbnailTextBlockText() const override;
 
 private:
diff --git a/Engine/Source/Editor/MovieSceneTools/Private/Sections/CinematicShotSection.cpp b/Engine/Source/Editor/MovieSceneTools/Private/Sections/CinematicShotSection.cpp
--- a/Engine/Source/Editor/MovieSceneTools/Private/Sections/CinematicShotSection.cpp
+++ b/Engine/Source/Editor/MovieSceneTools/Private/Sections/CinematicShotSection.cpp
@@ -32,28 +32,48 @@ FCinematicShotSection::~FCinematicShotSection()
 {
 }
 
-int32 FCinematicShotSection::OnPaintSection(const FGeometry& AllottedGeometry, const FSlateRect& SectionClippingRect, FSlateWindowElementList& OutDrawElements, int32 LayerId, bool bParentEnabled) const
+float FCinematicShotSection::GetSectionHeight() const
 {
+	return FThumbnailSection::GetSectionHeight() + 2*13.f;
+}
+
+FMargin FCinematicShotSection::GetContentPadding() const
+{
+	return FMargin(8.f, 15.f);
+}
+
+int32 FCinematicShotSection::OnPaintSection(FSequencerSectionPainter& InPainter) const
+{
+	static const FSlateBrush* FilmBorder = FEditorStyle::GetBrush("Sequencer.Section.FilmBorder");
+
+	InPainter.LayerId = InPainter.PaintSectionBackground();
+
+	FVector2D SectionSize = InPainter.SectionGeometry.GetLocalSize();
+
+	FSlateDrawElement::MakeBox(
+		InPainter.DrawElements,
+		InPainter.LayerId++,
+		InPainter.SectionGeometry.ToPaintGeometry(FVector2D(SectionSize.X-2.f, 7.f), FSlateLayoutTransform(FVector2D(1.f, 4.f))),
+		FilmBorder,
+		InPainter.SectionClippingRect.InsetBy(FMargin(1.f)),
+		InPainter.bParentEnabled ? ESlateDrawEffect::None : ESlateDrawEffect::DisabledEffect
+	);
+
+	FSlateDrawElement::MakeBox(
+		InPainter.DrawElements,
+		InPainter.LayerId++,
+		InPainter.SectionGeometry.ToPaintGeometry(FVector2D(SectionSize.X-2.f, 7.f), FSlateLayoutTransform(FVector2D(1.f, SectionSize.Y - 11.f))),
+		FilmBorder,
+		InPainter.SectionClippingRect.InsetBy(FMargin(1.f)),
+		InPainter.bParentEnabled ? ESlateDrawEffect::None : ESlateDrawEffect::DisabledEffect
+	);
+
 	if (SequenceInstance.IsValid())
 	{
-		return FThumbnailSection::OnPaintSection(AllottedGeometry, SectionClippingRect, OutDrawElements, LayerId, bParentEnabled);
-	}
-	else
-	{
-		const ESlateDrawEffect::Type DrawEffects = bParentEnabled ? ESlateDrawEffect::None : ESlateDrawEffect::DisabledEffect;
-
-		FSlateDrawElement::MakeBox( 
-			OutDrawElements,
-			LayerId,
-			AllottedGeometry.ToPaintGeometry(),
-			FEditorStyle::GetBrush("Sequencer.GenericSection.Background"),
-			SectionClippingRect,
-			DrawEffects,
-			FLinearColor::Black
-		); 
-
-		return LayerId + 2;
+		return FThumbnailSection::OnPaintSection(InPainter);
 	}
+
+	return InPainter.LayerId;
 }
 
 void FCinematicShotSection::BuildSectionContextMenu(FMenuBuilder& MenuBuilder, const FGuid& ObjectBinding)
@@ -108,14 +128,15 @@ void FCinematicShotSection::BuildSectionContextMenu(FMenuBuilder& MenuBuilder, c
 void FCinematicShotSection::AddTakesMenu(FMenuBuilder& MenuBuilder)
 {
 	TArray<uint32> TakeNumbers;
-	MovieSceneToolHelpers::GatherTakes(&SectionObject, TakeNumbers);
+	uint32 CurrentTakeNumber = INDEX_NONE;
+	MovieSceneToolHelpers::GatherTakes(&SectionObject, TakeNumbers, CurrentTakeNumber);
 
 	for (auto TakeNumber : TakeNumbers)
 	{
 		MenuBuilder.AddMenuEntry(
 			FText::Format(LOCTEXT("TakeNumber", "Take {0}"), FText::AsNumber(TakeNumber)),
 			FText::Format(LOCTEXT("TakeNumberTooltip", "Switch to take {0}"), FText::AsNumber(TakeNumber)),
-			FSlateIcon(),
+			TakeNumber == CurrentTakeNumber ? FSlateIcon(FEditorStyle::GetStyleSetName(), "Sequencer.Star") : FSlateIcon(FEditorStyle::GetStyleSetName(), "Sequencer.Empty"),
 			FUIAction(FExecuteAction::CreateSP(CinematicShotTrackEditor.Pin().ToSharedRef(), &FCinematicShotTrackEditor::SwitchTake, &SectionObject, TakeNumber))
 		);
 	}
@@ -129,7 +150,7 @@ FText FCinematicShotSection::GetDisplayName() const
 /* FCinematicShotSection callbacks
  *****************************************************************************/
 
-ACameraActor* FCinematicShotSection::GetCameraObject() const
+AActor* FCinematicShotSection::GetCameraObject() const
 {
 	return CinematicShotTrackEditor.IsValid() ? CinematicShotTrackEditor.Pin()->GetCinematicShotCamera().Get() : nullptr;
 }
@@ -158,7 +179,7 @@ FReply FCinematicShotSection::OnSectionDoubleClicked(const FGeometry& SectionGeo
 	{
 		if (SequenceInstance.IsValid())
 		{
-			Sequencer.Pin()->FocusSequenceInstance(SequenceInstance.Pin().ToSharedRef());
+			Sequencer.Pin()->FocusSequenceInstance(SectionObject);
 		}
 	}
 
diff --git a/Engine/Source/Editor/MovieSceneTools/Private/Sections/CinematicShotSection.h b/Engine/Source/Editor/MovieSceneTools/Private/Sections/CinematicShotSection.h
--- a/Engine/Source/Editor/MovieSceneTools/Private/Sections/CinematicShotSection.h
+++ b/Engine/Source/Editor/MovieSceneTools/Private/Sections/CinematicShotSection.h
@@ -30,14 +30,15 @@ public:
 
 	// ISequencerSection interface
 
-	virtual int32 OnPaintSection( const FGeometry& AllottedGeometry, const FSlateRect& SectionClippingRect, FSlateWindowElementList& OutDrawElements, int32 LayerId, bool bParentEnabled ) const override;
+	virtual int32 OnPaintSection( FSequencerSectionPainter& Painter ) const override;
 	virtual void BuildSectionContextMenu(FMenuBuilder& MenuBuilder, const FGuid& ObjectBinding) override;
 	virtual FText GetDisplayName() const override;
 	virtual FReply OnSectionDoubleClicked(const FGeometry& SectionGeometry, const FPointerEvent& MouseEvent) override;
-
+	virtual float GetSectionHeight() const override;
+	virtual FMargin GetContentPadding() const override;
 	// FThumbnail interface
 
-	virtual ACameraActor* GetCameraObject() const override;
+	virtual AActor* GetCameraObject() const override;
 	virtual bool CanRename() const override { return true; }
 	virtual FText HandleThumbnailTextBlockText() const override;
 	virtual void HandleThumbnailTextBlockTextCommitted(const FText& NewThumbnailName, ETextCommit::Type CommitType) override;
diff --git a/Engine/Source/Editor/MovieSceneTools/Private/Sections/ColorPropertySection.cpp b/Engine/Source/Editor/MovieSceneTools/Private/Sections/ColorPropertySection.cpp
--- a/Engine/Source/Editor/MovieSceneTools/Private/Sections/ColorPropertySection.cpp
+++ b/Engine/Source/Editor/MovieSceneTools/Private/Sections/ColorPropertySection.cpp
@@ -23,10 +23,11 @@ void FColorPropertySection::GenerateSectionLayout( class ISectionLayoutBuilder&
 	LayoutBuilder.AddKeyArea( "A", NSLOCTEXT( "FColorPropertySection", "OpacityArea", "Opacity" ), AlphaKeyArea.ToSharedRef() );
 }
 
-
-int32 FColorPropertySection::OnPaintSection( const FGeometry& AllottedGeometry, const FSlateRect& SectionClippingRect, FSlateWindowElementList& OutDrawElements, int32 LayerId, bool bParentEnabled ) const
+int32 FColorPropertySection::OnPaintSection( FSequencerSectionPainter& Painter ) const
 {
-	const ESlateDrawEffect::Type DrawEffects = bParentEnabled ? ESlateDrawEffect::None : ESlateDrawEffect::DisabledEffect;
+	int32 LayerId = Painter.PaintSectionBackground(FColor(255,255,255));
+
+	const ESlateDrawEffect::Type DrawEffects = Painter.bParentEnabled ? ESlateDrawEffect::None : ESlateDrawEffect::DisabledEffect;
 
 	const UMovieSceneColorSection* ColorSection = Cast<const UMovieSceneColorSection>( &SectionObject );
 
@@ -36,15 +37,18 @@ int32 FColorPropertySection::OnPaintSection( const FGeometry& AllottedGeometry,
 
 	if ( !FMath::IsNearlyZero( SectionDuration ) )
 	{
-		LayerId = FPropertySection::OnPaintSection( AllottedGeometry, SectionClippingRect, OutDrawElements, LayerId, bParentEnabled );
-
-		FVector2D GradientSize = FVector2D( AllottedGeometry.Size.X, (AllottedGeometry.Size.Y / 4) - 3.0f );
+		FVector2D GradientSize = FVector2D( Painter.SectionGeometry.Size.X, (Painter.SectionGeometry.Size.Y / 4) - 3.0f );
 
-		FPaintGeometry PaintGeometry = AllottedGeometry.ToPaintGeometry( FVector2D( 0, 0 ), GradientSize );
+		FPaintGeometry PaintGeometry = Painter.SectionGeometry.ToPaintGeometry( FVector2D( 0, 0 ), GradientSize );
 
 		// If we are showing a background pattern and the colors is transparent, draw a checker pattern
-		const FSlateBrush* CheckerBrush = FEditorStyle::GetBrush( "Checker" );
-		FSlateDrawElement::MakeBox( OutDrawElements, LayerId, PaintGeometry, CheckerBrush, SectionClippingRect, DrawEffects );
+		FSlateDrawElement::MakeBox(
+			Painter.DrawElements,
+			LayerId,
+			PaintGeometry,
+			FEditorStyle::GetBrush( "Checker" ),
+			Painter.SectionClippingRect,
+			DrawEffects);
 
 		TArray<FSlateGradientStop> GradientStops;
 
@@ -57,19 +61,19 @@ int32 FColorPropertySection::OnPaintSection( const FGeometry& AllottedGeometry,
 			FLinearColor Color = ColorKeys[i].Value;
 			float TimeFraction = (Time - StartTime) / SectionDuration;
 
-			GradientStops.Add( FSlateGradientStop( FVector2D( TimeFraction * AllottedGeometry.Size.X, 0 ),
+			GradientStops.Add( FSlateGradientStop( FVector2D( TimeFraction * Painter.SectionGeometry.Size.X, 0 ),
 				Color ) );
 		}
 
 		if ( GradientStops.Num() > 0 )
 		{
 			FSlateDrawElement::MakeGradient(
-				OutDrawElements,
-				LayerId + 1,
+				Painter.DrawElements,
+				Painter.LayerId + 1,
 				PaintGeometry,
 				GradientStops,
 				Orient_Vertical,
-				SectionClippingRect,
+				Painter.SectionClippingRect,
 				DrawEffects
 				);
 		}
diff --git a/Engine/Source/Editor/MovieSceneTools/Private/Sections/ColorPropertySection.h b/Engine/Source/Editor/MovieSceneTools/Private/Sections/ColorPropertySection.h
--- a/Engine/Source/Editor/MovieSceneTools/Private/Sections/ColorPropertySection.h
+++ b/Engine/Source/Editor/MovieSceneTools/Private/Sections/ColorPropertySection.h
@@ -35,7 +35,7 @@ public:
 	// FPropertySection interface
 
 	virtual void GenerateSectionLayout(class ISectionLayoutBuilder& LayoutBuilder) const override;
-	virtual int32 OnPaintSection(const FGeometry& AllottedGeometry, const FSlateRect& SectionClippingRect, FSlateWindowElementList& OutDrawElements, int32 LayerId, bool bParentEnabled) const override;
+	virtual int32 OnPaintSection(FSequencerSectionPainter& Painter) const override;
 	virtual void SetIntermediateValue( FPropertyChangedParams PropertyChangedParams ) override;
 	virtual void ClearIntermediateValue() override;
 
diff --git a/Engine/Source/Editor/MovieSceneTools/Private/Sections/EventTrackSection.cpp b/Engine/Source/Editor/MovieSceneTools/Private/Sections/EventTrackSection.cpp
--- a/Engine/Source/Editor/MovieSceneTools/Private/Sections/EventTrackSection.cpp
+++ b/Engine/Source/Editor/MovieSceneTools/Private/Sections/EventTrackSection.cpp
@@ -26,10 +26,9 @@ UMovieSceneSection* FEventTrackSection::GetSectionObject()
 	return Section;
 }
 
-
-int32 FEventTrackSection::OnPaintSection(const FGeometry& AllottedGeometry, const FSlateRect& SectionClippingRect, FSlateWindowElementList& OutDrawElements, int32 LayerId, bool bParentEnabled) const
+int32 FEventTrackSection::OnPaintSection( FSequencerSectionPainter& InPainter ) const
 {
-	return LayerId + 1;
+	return InPainter.PaintSectionBackground();
 }
 
 
diff --git a/Engine/Source/Editor/MovieSceneTools/Private/Sections/EventTrackSection.h b/Engine/Source/Editor/MovieSceneTools/Private/Sections/EventTrackSection.h
--- a/Engine/Source/Editor/MovieSceneTools/Private/Sections/EventTrackSection.h
+++ b/Engine/Source/Editor/MovieSceneTools/Private/Sections/EventTrackSection.h
@@ -28,14 +28,12 @@ public:
 
 	virtual UMovieSceneSection* GetSectionObject() override;
 //	virtual TSharedRef<SWidget> GenerateSectionWidget() override;
-	virtual int32 OnPaintSection( const FGeometry& AllottedGeometry, const FSlateRect& SectionClippingRect, FSlateWindowElementList& OutDrawElements, int32 LayerId, bool bParentEnabled ) const override;
+	virtual int32 OnPaintSection( FSequencerSectionPainter& InPainter ) const override;
 //	virtual void Tick( const FGeometry& AllottedGeometry, const FGeometry& ParentGeometry, const double InCurrentTime, const float InDeltaTime ) override;
 	virtual FText GetDisplayName() const override;
 	virtual FText GetSectionTitle() const override;
 //	virtual float GetSectionHeight() const override;
 //	virtual float GetSectionGripSize() const override;
-//	virtual FName GetSectionGripLeftBrushName() const override;
-//	virtual FName GetSectionGripRightBrushName() const override;
 //	virtual bool AreSectionsConnected() const override { return true; }
 	virtual void GenerateSectionLayout( class ISectionLayoutBuilder& LayoutBuilder ) const override;
 //	virtual FReply OnSectionDoubleClicked( const FGeometry& SectionGeometry, const FPointerEvent& MouseEvent ) override;
