Copying //UE4/Dev-Framework to Dev-Main (//UE4/Dev-Main) @ 2926658

==========================
MAJOR FEATURES + CHANGES
==========================

Change 2821607 on 2016/01/08 by [REDACTED_PERSON]

	Added a way to limit amount of information logged by vlog by discarding logs from classes from outside of class whitelist #UE4

	This feature was followed by refactoring of functions taking FVisualLogEntry pointers to use references instead.

Change 2828384 on 2016/01/14 by [REDACTED_PERSON]

	Back out of visual log refactor done as part of CL#2821607 #UE4

Change 2910454 on 2016/03/15 by [REDACTED_PERSON]

	#ue4 - Properly exclude zero-distance MTD results in ComponentEncroachesBlockingGeometry_WithAdjustment() in the presense of multiple overlaps.

	UDN: https://udn.unrealengine.com/questions/270574/jeff-farris-hack-for-physx-mtd.html

Change 2910548 on 2016/03/15 by [REDACTED_PERSON]

	#ue4 - Handle MTD computation returning NaN direction when there is a "contact" with zero distance.

Change 2912311 on 2016/03/16 by [REDACTED_PERSON]

	Properly handle overlaps in C++ in documentation code and UE4 to Unity doc

Change 2913086 on 2016/03/17 by [REDACTED_PERSON]

	Adding ability to have 9 parameters to a dynamic delegate

Change 2913101 on 2016/03/17 by [REDACTED_PERSON]

	Fix some of the loctext error messages

Change 2913102 on 2016/03/17 by [REDACTED_PERSON]

	Console usability improvements

	Display console autocompletion commands from the lexicographically first element up to either the total number of commands or MAX_AUTOCOMPLETION_LINES, whichever is least. The previous behaviour started the list "in the middle" and hid the first elements if there were too many matches. Thus "[ab ac ad]" with "aa" hidden now becomes "[aa ab ac]" with "ad" hidden.
	To make scrolling work as expected, the input handling of the up and down arrow keys has been reversed so that the cursor iterates forward starting from the top with the down arrow key, and goes back up with the up arrow key. Command history is still accessed with the up arrow key.
	This commit also undoes one of the most evil uses of operator overloading I've ever seen, on par with "#define true false" but more subtle

	Color console autocomplete entries to denote their type: command, CVar or other (manual autocompletion entries). CVars are further divided into writeable and read-only variables.
	Assume that manual console autocompletion entries are commands. This makes the autocompletion list colors more consistent and less noisy

	Automatically select (but don't complete) a command on console character input. To prevent the autocomplete from becoming too trigger happy, no longer automatically complete commands for arbitrary key inputs that we happen to have a match for

	Allow cycling through console commands with the tab key
	Discriminate between first time and repeated tab presses and only scroll through autocomplete entries on the latter

	Fix off-by-one error in console: "x more matches" line was being shown when the number of autocomplete elements was equal to MAX_AUTOCOMPLETION_LINES
	Fix an off-by-one error that was causing the topmost console command to not be shown if there was an autocomplete scroll region

	Show the currently selected autocomplete entry faded out behind the user's typed input
	Slightly increase brightness of the normal input text colour to better distinguish between the typed and autocompleted parts of the input line

	Left-justify command descriptions in the console autocompletion box
	Detect overflow of console autocomplete lines on low resolutions and decrease the space used for description justification to compensate

	Make the console input, history and autocomplete colours user configurable

	Add console background transparency. Configurable, set to 15% by default

	Add missing closing quote to the console dump HTML template

	#github #2061: Console usability improvements from Mattiwatti

Change 2913104 on 2016/03/17 by [REDACTED_PERSON]

	Added indicator displayed on animation nodes when they use the 'fast path'

	Added checkbox that can be used to audit Blueprint fast-path usage.
	Switched almost all animation node widgets to derive from new SAnimationGraphNode. This creates the overlay widget that indicates whether this node is using a more optimal path.

	#doc Also added documentation tooltips and UDN doc files/images for the fast path systems.

Change 2913306 on 2016/03/17 by [REDACTED_PERSON]

	Cleaning up GetResourceSize
	- Made many call Super::GetResourceSize
	- Removed trivial implementations
	- Fixed HierarchicalInstanceStaticMeshComponent double counting an array

Change 2913535 on 2016/03/17 by [REDACTED_PERSON]

	fixed broken behavior tree graph data after subnode undo
	#ue4 UE-28198

Change 2913608 on 2016/03/17 by [REDACTED_PERSON]

	fixed behavior tree execution indices after undoing move in editor
	#ue4 UE-26705

Change 2913847 on 2016/03/17 by [REDACTED_PERSON]

	added new automation test for UE-28309
	#ue4

Change 2913849 on 2016/03/17 by [REDACTED_PERSON]

	fixed behavior tree skipping over branch when restart request comes during AbortCurrentTask call
	#ue4 UE-28309

Change 2913895 on 2016/03/17 by [REDACTED_PERSON]

	Added 'self' argument to Actor and PrimitiveComponent delegates that didn't already supply one
	Fixed up all C++ uses of these delegates

Change 2914743 on 2016/03/18 by [REDACTED_PERSON]

	Editing of primitive data in PhAT

diff --git a/Engine/Config/BaseInput.ini b/Engine/Config/BaseInput.ini
--- a/Engine/Config/BaseInput.ini
+++ b/Engine/Config/BaseInput.ini
@@ -41,6 +41,7 @@ bRequireCtrlToNavigateAutoComplete=False
 +DebugExecBindings=(Key=Tab,Command="FocusLastPIEWindow",Control=True,Shift=True)
 ; Please remove EnableGDT command binding after removing deprecated GameplayDebugger module (we are going to use Gameplay Debugger plugin soon)
 +DebugExecBindings=(Key=Apostrophe,Command="EnableGDT")
++DebugExecBindings=(Key=Quote,Command="EnableGDT")
 
 [/Script/Engine.Console]
 HistoryBot=-1
@@ -177,3 +178,9 @@ MaxScrollbackSize=1024
 +ManualAutoCompleteList=(Command="StopRecordingAnimation All",Desc="Stop all recording animations")
 +ManualAutoCompleteList=(Command="RecordSequence Filter ActorOrClassName",Desc="Record a level sequence from gameplay. Filter=<All|Actor|Class>")
 +ManualAutoCompleteList=(Command="StopRecordingSequence",Desc="Stop recording the current sequence. Only one sequence recording can be active at one time.")
+BackgroundOpacityPercentage=85.000000
+InputColor=(R=230,G=230,B=230,A=255)
+HistoryColor=(R=180,G=180,B=180,A=255)
+AutoCompleteCommandColor=(B=185,G=109,R=144,A=255)
+AutoCompleteCVarColor=(B=86,G=158,R=86,A=255)
+AutoCompleteFadedColor=(R=100,G=100,B=100,A=255)
\ No newline at end of file
diff --git a/Engine/Plugins/Experimental/Phya/Source/Phya/Classes/PhyaCollisionHandler.h b/Engine/Plugins/Experimental/Phya/Source/Phya/Classes/PhyaCollisionHandler.h
--- a/Engine/Plugins/Experimental/Phya/Source/Phya/Classes/PhyaCollisionHandler.h
+++ b/Engine/Plugins/Experimental/Phya/Source/Phya/Classes/PhyaCollisionHandler.h
@@ -34,7 +34,7 @@ class UPhyaCollisionHandler : public UPhysicsCollisionHandler
 
 	//~ Begin UPhysicsCollisionHandler Interface
 	virtual void InitCollisionHandler() override;
-	virtual void HandlePhysicsCollisions_AssumesLocked(const TArray<FCollisionNotifyInfo>& PendingCollisionNotifies) override;
+	virtual void HandlePhysicsCollisions_AssumesLocked(TArray<FCollisionNotifyInfo>& PendingCollisionNotifies) override;
 	//~ End UPhysicsCollisionHandler Interface
 
 	void TestImpact();
diff --git a/Engine/Plugins/Experimental/Phya/Source/Phya/Private/PhyaCollisionHandler.cpp b/Engine/Plugins/Experimental/Phya/Source/Phya/Private/PhyaCollisionHandler.cpp
--- a/Engine/Plugins/Experimental/Phya/Source/Phya/Private/PhyaCollisionHandler.cpp
+++ b/Engine/Plugins/Experimental/Phya/Source/Phya/Private/PhyaCollisionHandler.cpp
@@ -179,7 +179,7 @@ void UPhyaCollisionHandler::TestImpact()
 
 
 
-void UPhyaCollisionHandler::HandlePhysicsCollisions_AssumesLocked(const TArray<FCollisionNotifyInfo>& PendingCollisionNotifies)
+void UPhyaCollisionHandler::HandlePhysicsCollisions_AssumesLocked(TArray<FCollisionNotifyInfo>& PendingCollisionNotifies)
 {
 	if(GetWorld()->HasBegunPlay())
 	{
diff --git a/Engine/Source/Developer/AITestSuite/Classes/MockGameplayTasks.h b/Engine/Source/Developer/AITestSuite/Classes/MockGameplayTasks.h
--- a/Engine/Source/Developer/AITestSuite/Classes/MockGameplayTasks.h
+++ b/Engine/Source/Developer/AITestSuite/Classes/MockGameplayTasks.h
@@ -69,8 +69,5 @@ public:
 	UGameplayTasksComponent* GTComponent;
 
 	virtual UGameplayTasksComponent* GetGameplayTasksComponent(const UGameplayTask& Task) const override { return GTComponent; }	
-	virtual void OnTaskActivated(UGameplayTask& Task) override {}
-	virtual void OnTaskDeactivated(UGameplayTask& Task) override {}
-	virtual AActor* GetOwnerActor(const UGameplayTask* Task) const { return nullptr; }
-
-};
\ No newline at end of file
+	virtual AActor* GetGameplayTaskOwner(const UGameplayTask* Task) const { return nullptr; }
+};
diff --git a/Engine/Source/Developer/AITestSuite/Private/Tests/BTTest.cpp b/Engine/Source/Developer/AITestSuite/Private/Tests/BTTest.cpp
--- a/Engine/Source/Developer/AITestSuite/Private/Tests/BTTest.cpp
+++ b/Engine/Source/Developer/AITestSuite/Private/Tests/BTTest.cpp
@@ -512,16 +512,49 @@ struct FAITest_BTAbortDuringLatentAbort : public FAITest_SimpleBT
 	{
 		UBTCompositeNode& CompNode = FBTBuilder::AddSelector(*BTAsset);
 		{
-			FBTBuilder::AddTaskLatentFlags(CompNode, EBTNodeResult::Succeeded,
-				1, TEXT("Bool1"), 1, 2,
-				1, TEXT("Bool2"), 3, 4);
+			UBTCompositeNode& Comp1Node = FBTBuilder::AddSelector(CompNode);
 			{
-				FBTBuilder::WithDecoratorBlackboard(CompNode, EBasicKeyOperation::NotSet, EBTFlowAbortMode::Self, TEXT("Bool1"));
+				FBTBuilder::AddTask(Comp1Node, 0, EBTNodeResult::Succeeded);
+				{
+					FBTBuilder::WithDecoratorBlackboard(Comp1Node, EBasicKeyOperation::Set, EBTFlowAbortMode::LowerPriority, TEXT("Bool2"));
+				}
+
+				FBTBuilder::AddTaskLatentFlags(Comp1Node, EBTNodeResult::Succeeded,
+					1, TEXT("Bool1"), 1, 2,
+					1, TEXT("Bool2"), 3, 4);
+				{
+					FBTBuilder::WithDecoratorBlackboard(Comp1Node, EBasicKeyOperation::NotSet, EBTFlowAbortMode::Self, TEXT("Bool1"));
+				}
 			}
 
-			FBTBuilder::AddTask(CompNode, 5, EBTNodeResult::Succeeded);
+			FBTBuilder::AddTask(CompNode, 6, EBTNodeResult::Succeeded);
+		}
+
+		ExpectedResult.Add(1);
+		ExpectedResult.Add(3);
+		ExpectedResult.Add(4);
+		ExpectedResult.Add(0);
+	}
+};
+IMPLEMENT_AI_LATENT_TEST(FAITest_BTAbortDuringLatentAbort, "System.Engine.AI.Behavior Trees.Abort: during latent task abort (lower pri)")
+
+struct FAITest_BTAbortDuringLatentAbort2 : public FAITest_SimpleBT
+{
+	FAITest_BTAbortDuringLatentAbort2()
+	{
+		UBTCompositeNode& CompNode = FBTBuilder::AddSelector(*BTAsset);
+		{
+			UBTCompositeNode& Comp1Node = FBTBuilder::AddSelector(CompNode);
 			{
-				FBTBuilder::WithDecoratorBlackboard(CompNode, EBasicKeyOperation::Set, EBTFlowAbortMode::LowerPriority, TEXT("Bool2"));
+				FBTBuilder::AddTaskLatentFlags(Comp1Node, EBTNodeResult::Succeeded,
+					1, TEXT("Bool1"), 1, 2,
+					1, TEXT("Bool2"), 3, 4);
+				{
+					FBTBuilder::WithDecoratorBlackboard(Comp1Node, EBasicKeyOperation::NotSet, EBTFlowAbortMode::Self, TEXT("Bool1"));
+					FBTBuilder::WithDecoratorBlackboard(Comp1Node, EBasicKeyOperation::NotSet, EBTFlowAbortMode::Self, TEXT("Bool2"));
+				}
+
+				FBTBuilder::AddTask(Comp1Node, 5, EBTNodeResult::Succeeded);
 			}
 
 			FBTBuilder::AddTask(CompNode, 6, EBTNodeResult::Succeeded);
@@ -533,7 +566,7 @@ struct FAITest_BTAbortDuringLatentAbort : public FAITest_SimpleBT
 		ExpectedResult.Add(5);
 	}
 };
-IMPLEMENT_AI_LATENT_TEST(FAITest_BTAbortDuringLatentAbort, "System.Engine.AI.Behavior Trees.Abort: during latent task abort")
+IMPLEMENT_AI_LATENT_TEST(FAITest_BTAbortDuringLatentAbort2, "System.Engine.AI.Behavior Trees.Abort: during latent task abort (self)")
 
 struct FAITest_BTAbortDuringInstantAbort : public FAITest_SimpleBT
 {
@@ -547,7 +580,7 @@ struct FAITest_BTAbortDuringInstantAbort : public FAITest_SimpleBT
 
 				FBTBuilder::AddTask(Comp1Node, 1, EBTNodeResult::Succeeded);
 				{
-					FBTBuilder::WithDecoratorBlackboard(CompNode, EBasicKeyOperation::Set, EBTFlowAbortMode::None, TEXT("Bool3"));
+					FBTBuilder::WithDecoratorBlackboard(Comp1Node, EBasicKeyOperation::Set, EBTFlowAbortMode::None, TEXT("Bool3"));
 				}
 			}
 
@@ -574,7 +607,121 @@ struct FAITest_BTAbortDuringInstantAbort : public FAITest_SimpleBT
 		ExpectedResult.Add(7);
 	}
 };
-IMPLEMENT_AI_LATENT_TEST(FAITest_BTAbortDuringInstantAbort, "System.Engine.AI.Behavior Trees.Abort: during instant task abort")
+IMPLEMENT_AI_LATENT_TEST(FAITest_BTAbortDuringInstantAbort, "System.Engine.AI.Behavior Trees.Abort: during instant task abort (lower pri)")
+
+struct FAITest_BTAbortDuringInstantAbort2 : public FAITest_SimpleBT
+{
+	FAITest_BTAbortDuringInstantAbort2()
+	{
+		UBTCompositeNode& CompNode = FBTBuilder::AddSelector(*BTAsset);
+		{
+			UBTCompositeNode& Comp1Node = FBTBuilder::AddSelector(CompNode);
+			{
+				FBTBuilder::AddTaskLatentFlags(Comp1Node, EBTNodeResult::Succeeded,
+					1, TEXT("Bool1"), 0, 1, 
+					0, TEXT("Bool2"), 2, 3);
+				{
+					FBTBuilder::WithDecoratorBlackboard(Comp1Node, EBasicKeyOperation::NotSet, EBTFlowAbortMode::Self, TEXT("Bool1"));
+					FBTBuilder::WithDecoratorBlackboard(Comp1Node, EBasicKeyOperation::NotSet, EBTFlowAbortMode::Both, TEXT("Bool2"));
+				}
+
+				FBTBuilder::AddTask(Comp1Node, 4, EBTNodeResult::Succeeded);
+			}
+
+			FBTBuilder::AddTask(CompNode, 5, EBTNodeResult::Succeeded);
+		}
+
+		ExpectedResult.Add(0);
+		ExpectedResult.Add(2);
+		ExpectedResult.Add(3);
+		ExpectedResult.Add(4);
+	}
+};
+IMPLEMENT_AI_LATENT_TEST(FAITest_BTAbortDuringInstantAbort2, "System.Engine.AI.Behavior Trees.Abort: during instant task abort (self)")
+
+struct FAITest_BTAbortOnValueChangePass : public FAITest_SimpleBT
+{
+	FAITest_BTAbortOnValueChangePass()
+	{
+		UBTCompositeNode& CompNode = FBTBuilder::AddSelector(*BTAsset);
+		{
+			FBTBuilder::AddTask(CompNode, 0, EBTNodeResult::Failed);
+
+			UBTCompositeNode& CompNode2 = FBTBuilder::AddSequence(CompNode);
+			{
+				FBTBuilder::WithDecoratorBlackboard(CompNode, EArithmeticKeyOperation::NotEqual, 10, EBTFlowAbortMode::Self, EBTBlackboardRestart::ValueChange, TEXT("Int"));
+
+				FBTBuilder::AddTask(CompNode2, 1, EBTNodeResult::Succeeded);
+				FBTBuilder::AddTaskValueChange(CompNode2, 1, EBTNodeResult::Succeeded, TEXT("Int"));
+				FBTBuilder::AddTask(CompNode2, 2, EBTNodeResult::Succeeded);
+			}
+
+			FBTBuilder::AddTask(CompNode, 3, EBTNodeResult::Failed);
+		}
+
+		ExpectedResult.Add(0);
+		ExpectedResult.Add(1);
+		ExpectedResult.Add(1);
+		ExpectedResult.Add(2);
+	}
+};
+IMPLEMENT_AI_LATENT_TEST(FAITest_BTAbortOnValueChangePass, "System.Engine.AI.Behavior Trees.Abort: value change (pass)")
+
+struct FAITest_BTAbortOnValueChangeFail : public FAITest_SimpleBT
+{
+	FAITest_BTAbortOnValueChangeFail()
+	{
+		UBTCompositeNode& CompNode = FBTBuilder::AddSelector(*BTAsset);
+		{
+			FBTBuilder::AddTask(CompNode, 0, EBTNodeResult::Failed);
+
+			UBTCompositeNode& CompNode2 = FBTBuilder::AddSequence(CompNode);
+			{
+				FBTBuilder::WithDecoratorBlackboard(CompNode, EArithmeticKeyOperation::NotEqual, 10, EBTFlowAbortMode::Self, EBTBlackboardRestart::ValueChange, TEXT("Int"));
+
+				FBTBuilder::AddTask(CompNode2, 1, EBTNodeResult::Succeeded);
+				FBTBuilder::AddTaskValueChange(CompNode2, 10, EBTNodeResult::Succeeded, TEXT("Int"));
+				FBTBuilder::AddTask(CompNode2, 2, EBTNodeResult::Succeeded);
+			}
+
+			FBTBuilder::AddTask(CompNode, 3, EBTNodeResult::Failed);
+		}
+
+		ExpectedResult.Add(0);
+		ExpectedResult.Add(1);
+		ExpectedResult.Add(3);
+	}
+};
+IMPLEMENT_AI_LATENT_TEST(FAITest_BTAbortOnValueChangeFail, "System.Engine.AI.Behavior Trees.Abort: value change (fail)")
+
+struct FAITest_BTAbortOnValueChangeFailOther : public FAITest_SimpleBT
+{
+	FAITest_BTAbortOnValueChangeFailOther()
+	{
+		UBTCompositeNode& CompNode = FBTBuilder::AddSelector(*BTAsset);
+		{
+			FBTBuilder::AddTask(CompNode, 0, EBTNodeResult::Failed);
+
+			UBTCompositeNode& CompNode2 = FBTBuilder::AddSequence(CompNode);
+			{
+				FBTBuilder::WithDecoratorBlackboard(CompNode, EArithmeticKeyOperation::NotEqual, 10, EBTFlowAbortMode::Self, EBTBlackboardRestart::ValueChange, TEXT("Int"));
+				FBTBuilder::WithDecoratorBlackboard(CompNode, EBasicKeyOperation::NotSet, EBTFlowAbortMode::None, TEXT("Bool1"));
+
+				FBTBuilder::AddTask(CompNode2, 1, EBTNodeResult::Succeeded);
+				FBTBuilder::AddTaskFlagChange(CompNode2, true, EBTNodeResult::Succeeded, TEXT("Bool1"));
+				FBTBuilder::AddTaskValueChange(CompNode2, 1, EBTNodeResult::Succeeded, TEXT("Int"));
+				FBTBuilder::AddTask(CompNode2, 2, EBTNodeResult::Succeeded);
+			}
+
+			FBTBuilder::AddTask(CompNode, 3, EBTNodeResult::Failed);
+		}
+
+		ExpectedResult.Add(0);
+		ExpectedResult.Add(1);
+		ExpectedResult.Add(3);
+	}
+};
+IMPLEMENT_AI_LATENT_TEST(FAITest_BTAbortOnValueChangeFailOther, "System.Engine.AI.Behavior Trees.Abort: value change (other failed)")
 
 struct FAITest_BTLowPriObserverInLoop : public FAITest_SimpleBT
 {
@@ -696,4 +843,77 @@ struct FAITest_BTSubtreeAbortOut : public FAITest_SimpleBT
 };
 IMPLEMENT_AI_LATENT_TEST(FAITest_BTSubtreeAbortOut, "System.Engine.AI.Behavior Trees.Subtree: abort out")
 
+struct FAITest_BTServiceInstantTask : public FAITest_SimpleBT
+{
+	FAITest_BTServiceInstantTask()
+	{
+		UBTCompositeNode& CompNode = FBTBuilder::AddSequence(*BTAsset);
+		{
+			FBTBuilder::AddTask(CompNode, 0, EBTNodeResult::Succeeded);
+			{
+				FBTBuilder::WithTaskServiceLog(CompNode, 1, 2);
+			}
+
+			FBTBuilder::AddTask(CompNode, 3, EBTNodeResult::Succeeded);
+		}
+
+		ExpectedResult.Add(1);
+		ExpectedResult.Add(0);
+		ExpectedResult.Add(2);
+		ExpectedResult.Add(3);
+	}
+};
+IMPLEMENT_AI_LATENT_TEST(FAITest_BTServiceInstantTask, "System.Engine.AI.Behavior Trees.Service: instant task")
+
+struct FAITest_BTServiceLatentTask : public FAITest_SimpleBT
+{
+	FAITest_BTServiceLatentTask()
+	{
+		UBTCompositeNode& CompNode = FBTBuilder::AddSequence(*BTAsset);
+		{
+			FBTBuilder::AddTask(CompNode, 0, EBTNodeResult::Succeeded, 2);
+			{
+				FBTBuilder::WithTaskServiceLog(CompNode, 1, 2);
+			}
+
+			FBTBuilder::AddTask(CompNode, 3, EBTNodeResult::Succeeded);
+		}
+
+		ExpectedResult.Add(1);
+		ExpectedResult.Add(0);
+		ExpectedResult.Add(2);
+		ExpectedResult.Add(3);
+	}
+};
+IMPLEMENT_AI_LATENT_TEST(FAITest_BTServiceLatentTask, "System.Engine.AI.Behavior Trees.Service: latent task")
+
+struct FAITest_BTServiceAbortingTask : public FAITest_SimpleBT
+{
+	FAITest_BTServiceAbortingTask()
+	{
+		UBTCompositeNode& CompNode = FBTBuilder::AddSelector(*BTAsset);
+		{
+			FBTBuilder::AddTask(CompNode, 0, EBTNodeResult::Succeeded);
+			{
+				FBTBuilder::WithDecoratorBlackboard(CompNode, EBasicKeyOperation::Set, EBTFlowAbortMode::LowerPriority, TEXT("Bool1"));
+			}
+
+			FBTBuilder::AddTaskLatentFlags(CompNode, EBTNodeResult::Succeeded, 1, TEXT("Bool1"), 1, 2, 0, NAME_None, 3, 4);
+			{
+				FBTBuilder::WithTaskServiceLog(CompNode, 5, 6);
+			}
+
+			FBTBuilder::AddTask(CompNode, 7, EBTNodeResult::Succeeded);
+		}
+
+		ExpectedResult.Add(5);
+		ExpectedResult.Add(1);
+		ExpectedResult.Add(6);
+		ExpectedResult.Add(3);
+		ExpectedResult.Add(4);
+		ExpectedResult.Add(0);
+	}
+};
+IMPLEMENT_AI_LATENT_TEST(FAITest_BTServiceAbortingTask, "System.Engine.AI.Behavior Trees.Service: abort task")
+
 #undef LOCTEXT_NAMESPACE
diff --git a/Engine/Source/Developer/AITestSuite/Public/BTBuilder.h b/Engine/Source/Developer/AITestSuite/Public/BTBuilder.h
--- a/Engine/Source/Developer/AITestSuite/Public/BTBuilder.h
+++ b/Engine/Source/Developer/AITestSuite/Public/BTBuilder.h
@@ -13,8 +13,10 @@
 
 #include "BehaviorTree/TestBTTask_Log.h"
 #include "BehaviorTree/TestBTTask_SetFlag.h"
+#include "BehaviorTree/TestBTTask_SetValue.h"
 #include "BehaviorTree/TestBTTask_LatentWithFlags.h"
 #include "BehaviorTree/TestBTDecorator_DelayedAbort.h"
+#include "BehaviorTree/TestBTService_Log.h"
 
 struct FBTBuilder
 {
@@ -39,6 +41,10 @@ struct FBTBuilder
 		KeyData.KeyType = NewObject<UBlackboardKeyType_Bool>();
 		BB->Keys.Add(KeyData);
 
+		KeyData.EntryName = TEXT("Int");
+		KeyData.KeyType = NewObject<UBlackboardKeyType_Int>();
+		BB->Keys.Add(KeyData);
+
 		BB->UpdateParentKeys();
 
 		UBehaviorTree* TreeOb = NewObject<UBehaviorTree>();
@@ -148,6 +154,17 @@ struct FBTBuilder
 		ParentNode.Children[ChildIdx].ChildTask = TaskNode;
 	}
 
+	static void AddTaskValueChange(UBTCompositeNode& ParentNode, int32 Value, EBTNodeResult::Type NodeResult, FName IntKeyName = TEXT("Int"))
+	{
+		UTestBTTask_SetValue* TaskNode = NewObject<UTestBTTask_SetValue>(ParentNode.GetTreeAsset());
+		TaskNode->Value = Value;
+		TaskNode->TaskResult = NodeResult;
+		TaskNode->KeyName = IntKeyName;
+
+		const int32 ChildIdx = ParentNode.Children.AddZeroed(1);
+		ParentNode.Children[ChildIdx].ChildTask = TaskNode;
+	}
+
 	static void AddTaskSubtree(UBTCompositeNode& ParentNode, UBehaviorTree* TreeAsset)
 	{
 		UBTTask_RunBehavior* TaskNode = NewObject<UBTTask_RunBehavior>(ParentNode.GetTreeAsset());
@@ -190,7 +207,7 @@ struct FBTBuilder
 	static void WithDecoratorBlackboard(UBTCompositeNode& ParentNode, EBasicKeyOperation::Type Condition, EBTFlowAbortMode::Type Observer, FName BoolKeyName = TEXT("Bool1"))
 	{
 		UBTDecorator_Blackboard& BBDecorator = WithDecorator<UBTDecorator_Blackboard>(ParentNode);
-		
+
 		UByteProperty* ConditionProp = FindField<UByteProperty>(UBTDecorator_Blackboard::StaticClass(), TEXT("OperationType"));
 		uint8* ConditionPropData = ConditionProp->ContainerPtrToValuePtr<uint8>(&BBDecorator);
 		ConditionProp->SetIntPropertyValue(ConditionPropData, (uint64)Condition);
@@ -204,6 +221,31 @@ struct FBTBuilder
 		KeyPropData->SelectedKeyName = BoolKeyName;
 	}
 
+	static void WithDecoratorBlackboard(UBTCompositeNode& ParentNode, EArithmeticKeyOperation::Type Condition, int32 Value, EBTFlowAbortMode::Type Observer, EBTBlackboardRestart::Type NotifyMode, FName IntKeyName = TEXT("Int"))
+	{
+		UBTDecorator_Blackboard& BBDecorator = WithDecorator<UBTDecorator_Blackboard>(ParentNode);
+
+		UByteProperty* ConditionProp = FindField<UByteProperty>(UBTDecorator_Blackboard::StaticClass(), TEXT("OperationType"));
+		uint8* ConditionPropData = ConditionProp->ContainerPtrToValuePtr<uint8>(&BBDecorator);
+		ConditionProp->SetIntPropertyValue(ConditionPropData, (uint64)Condition);
+
+		UByteProperty* ObserverProp = FindField<UByteProperty>(UBTDecorator_Blackboard::StaticClass(), TEXT("FlowAbortMode"));
+		uint8* ObserverPropData = ObserverProp->ContainerPtrToValuePtr<uint8>(&BBDecorator);
+		ObserverProp->SetIntPropertyValue(ObserverPropData, (uint64)Observer);
+
+		UByteProperty* NotifyModeProp = FindField<UByteProperty>(UBTDecorator_Blackboard::StaticClass(), TEXT("NotifyObserver"));
+		uint8* NotifyModePropData = NotifyModeProp->ContainerPtrToValuePtr<uint8>(&BBDecorator);
+		NotifyModeProp->SetIntPropertyValue(NotifyModePropData, (uint64)NotifyMode);
+
+		UIntProperty* ConditionValueProp = FindField<UIntProperty>(UBTDecorator_Blackboard::StaticClass(), TEXT("IntValue"));
+		uint8* ConditionValuePropData = ConditionValueProp->ContainerPtrToValuePtr<uint8>(&BBDecorator);
+		ConditionValueProp->SetIntPropertyValue(ConditionValuePropData, (uint64)Value);
+
+		UStructProperty* KeyProp = FindField<UStructProperty>(UBTDecorator_Blackboard::StaticClass(), TEXT("BlackboardKey"));
+		FBlackboardKeySelector* KeyPropData = KeyProp->ContainerPtrToValuePtr<FBlackboardKeySelector>(&BBDecorator);
+		KeyPropData->SelectedKeyName = IntKeyName;
+	}
+
 	static void WithDecoratorDelayedAbort(UBTCompositeNode& ParentNode, int32 NumTicks, bool bAbortOnlyOnce = true)
 	{
 		UTestBTDecorator_DelayedAbort& AbortDecorator = WithDecorator<UTestBTDecorator_DelayedAbort>(ParentNode);
@@ -216,4 +258,39 @@ struct FBTBuilder
 		UBTDecorator_Loop& LoopDecorator = WithDecorator<UBTDecorator_Loop>(ParentNode);
 		LoopDecorator.NumLoops = NumLoops;
 	}
+	
+	template<class T>
+	static T& WithService(UBTCompositeNode& ParentNode, UClass* ServiceClass = T::StaticClass())
+	{
+		T* ServiceOb = NewObject<T>(ParentNode.GetTreeAsset());
+		ParentNode.Services.Add(ServiceOb);
+
+		return *ServiceOb;
+	}
+
+	static void WithServiceLog(UBTCompositeNode& ParentNode, int32 ActivationIndex, int32 DeactivationIndex)
+	{
+		UTestBTService_Log& LogService = WithService<UTestBTService_Log>(ParentNode);
+		LogService.LogActivation = ActivationIndex;
+		LogService.LogDeactivation = DeactivationIndex;
+	}
+
+	template<class T>
+	static T& WithTaskService(UBTCompositeNode& ParentNode, UClass* ServiceClass = T::StaticClass())
+	{
+		UBTTaskNode* TaskNode = ParentNode.Children.Last().ChildTask;
+		check(TaskNode);
+
+		T* ServiceOb = NewObject<T>(ParentNode.GetTreeAsset());
+		TaskNode->Services.Add(ServiceOb);
+
+		return *ServiceOb;
+	}
+
+	static void WithTaskServiceLog(UBTCompositeNode& ParentNode, int32 ActivationIndex, int32 DeactivationIndex)
+	{
+		UTestBTService_Log& LogService = WithTaskService<UTestBTService_Log>(ParentNode);
+		LogService.LogActivation = ActivationIndex;
+		LogService.LogDeactivation = DeactivationIndex;
+	}
 };
diff --git a/Engine/Source/Developer/BlueprintCompilerCppBackend/Private/BlueprintCompilerCppBackendUtils.cpp b/Engine/Source/Developer/BlueprintCompilerCppBackend/Private/BlueprintCompilerCppBackendUtils.cpp
--- a/Engine/Source/Developer/BlueprintCompilerCppBackend/Private/BlueprintCompilerCppBackendUtils.cpp
+++ b/Engine/Source/Developer/BlueprintCompilerCppBackend/Private/BlueprintCompilerCppBackendUtils.cpp
@@ -732,6 +732,7 @@ int32 FEmitHelper::ParseDelegateDetails(FEmitterLocalContext& EmitterContext, UF
 	case 6: ParamNumberStr = TEXT("_SixParams");	break;
 	case 7: ParamNumberStr = TEXT("_SevenParams");	break;
 	case 8: ParamNumberStr = TEXT("_EightParams");	break;
+	case 9: ParamNumberStr = TEXT("_NineParams");   break;
 	default: ParamNumberStr = TEXT("_TooMany");		break;
 	}
 
diff --git a/Engine/Source/Developer/DesktopPlatform/Private/Mac/DesktopPlatformMac.cpp b/Engine/Source/Developer/DesktopPlatform/Private/Mac/DesktopPlatformMac.cpp
--- a/Engine/Source/Developer/DesktopPlatform/Private/Mac/DesktopPlatformMac.cpp
+++ b/Engine/Source/Developer/DesktopPlatform/Private/Mac/DesktopPlatformMac.cpp
@@ -556,7 +556,7 @@ void FDesktopPlatformMac::EnumerateEngineInstallations(TMap<FString, FString> &O
 	TArray<FName> KeysToRemove;
 	for (auto It : Section)
 	{
-		const FString& EngineDir = It.Value;
+		const FString& EngineDir = It.Value.GetValue();
 		if (EngineDir.Contains("Unreal Engine.app/Contents/") || EngineDir.Contains("Epic Games Launcher.app/Contents/") || EngineDir.Contains("[REDACTED_USER_HOME]/UnrealEngine/Launcher") || !IFileManager::Get().DirectoryExists(*EngineDir))
 		{
 			KeysToRemove.Add(It.Key);
diff --git a/Engine/Source/Developer/GameplayDebugger/Private/GameplayDebuggingComponent.cpp b/Engine/Source/Developer/GameplayDebugger/Private/GameplayDebuggingComponent.cpp
--- a/Engine/Source/Developer/GameplayDebugger/Private/GameplayDebuggingComponent.cpp
+++ b/Engine/Source/Developer/GameplayDebugger/Private/GameplayDebuggingComponent.cpp
@@ -645,14 +645,26 @@ void UGameplayDebuggingComponent::CollectBasicBehaviorData(APawn* MyPawn)
 		CurrentAIAssets = TEXT("");
 	}
 
+	GameplayTasksState = TEXT("");
 	UGameplayTasksComponent* GTComponent = MyPawn->FindComponentByClass<UGameplayTasksComponent>();
 	if (GTComponent)
 	{
-		GameplayTasksState = FString::Printf(TEXT("Ticking Tasks: %s\nTask Queue: %s"), *GTComponent->GetTickingTasksDescription(), *GTComponent->GetTasksPriorityQueueDescription());
-	}
-	else
-	{
-		GameplayTasksState = TEXT("");
+		for (FConstGameplayTaskIterator It = GTComponent->GetPriorityQueueIterator(); It; ++It)
+		{
+			const UGameplayTask* QueueTask = *It;
+			if (QueueTask)
+			{
+				const UObject* OwnerOb = Cast<const UObject>(QueueTask->GetTaskOwner());
+
+				GameplayTasksState += FString::Printf(TEXT("{white}%s%s {%s}%s {white}Owner:{yellow}%s {white}Res:{yellow}%s\n"),
+					*QueueTask->GetName(),
+					QueueTask->GetInstanceName() != NAME_None ? *FString::Printf(TEXT(" {yellow}[%s]"), *QueueTask->GetInstanceName().ToString()) : TEXT(""),
+					QueueTask->IsActive() ? TEXT("green") : TEXT("orange"),
+					*QueueTask->GetTaskStateName(),
+					(OwnerOb == GTComponent) ? TEXT("default") : *GetNameSafe(OwnerOb),
+					*QueueTask->GetRequiredResources().GetDebugDescription());
+			}
+		}
 	}
 #endif //!(UE_BUILD_SHIPPING || UE_BUILD_TEST)
 }
diff --git a/Engine/Source/Developer/GameplayDebugger/Private/GameplayDebuggingHUDComponent.cpp b/Engine/Source/Developer/GameplayDebugger/Private/GameplayDebuggingHUDComponent.cpp
--- a/Engine/Source/Developer/GameplayDebugger/Private/GameplayDebuggingHUDComponent.cpp
+++ b/Engine/Source/Developer/GameplayDebugger/Private/GameplayDebuggingHUDComponent.cpp
@@ -459,8 +459,21 @@ void AGameplayDebuggingHUDComponent::DrawBasicData(APlayerController* PC, class
 		PrintString(DefaultContext, FString::Printf(TEXT("Ability: {yellow}%s\n"), *DebugComponent->AbilityInfo));
 	}
 
-	// putting gameplay tasks' stuff last since it can expand heavily
-	PrintString(DefaultContext, FString::Printf(TEXT("GameplayTasks:\n{yellow}%s\n"), *DebugComponent->GameplayTasksState));
+	// gameplay tasks
+	int32 NumTasks = 0;
+	if (DebugComponent->GameplayTasksState.Len() > 0)
+	{
+		int32 SearchStart = -2;
+		int32 PrevStart = 0;
+		do {
+			PrevStart = SearchStart + 1;
+			SearchStart = DebugComponent->GameplayTasksState.Find(TEXT("\n"), ESearchCase::IgnoreCase, ESearchDir::FromStart, PrevStart);
+			NumTasks++;
+		} while (SearchStart >= 0 && SearchStart > PrevStart);
+		
+		NumTasks--;
+	}
+	PrintString(DefaultContext, FString::Printf(TEXT("GameplayTasks: {yellow}%d\n%s\n"), NumTasks, *DebugComponent->GameplayTasksState));
 
 	DrawPath(PC, DebugComponent);
 #endif //!(UE_BUILD_SHIPPING || UE_BUILD_TEST)
diff --git a/Engine/Source/Developer/LauncherServices/Private/Profiles/LauncherProfileManager.cpp b/Engine/Source/Developer/LauncherServices/Private/Profiles/LauncherProfileManager.cpp
--- a/Engine/Source/Developer/LauncherServices/Private/Profiles/LauncherProfileManager.cpp
+++ b/Engine/Source/Developer/LauncherServices/Private/Profiles/LauncherProfileManager.cpp
@@ -419,7 +419,7 @@ void FLauncherProfileManager::LoadDeviceGroups( )
 			{
 				if (It.Key() == TEXT("DeviceGroup"))
 				{
-					DeviceGroups.Add(ParseDeviceGroup(*It.Value()));
+					DeviceGroups.Add(ParseDeviceGroup(It.Value().GetValue()));
 				}
 			}
 		}
diff --git a/Engine/Source/Developer/SimplygonMeshReduction/Private/SimplygonMeshReduction.cpp b/Engine/Source/Developer/SimplygonMeshReduction/Private/SimplygonMeshReduction.cpp
--- a/Engine/Source/Developer/SimplygonMeshReduction/Private/SimplygonMeshReduction.cpp
+++ b/Engine/Source/Developer/SimplygonMeshReduction/Private/SimplygonMeshReduction.cpp
@@ -221,7 +221,7 @@ public:
 	bool ReduceLODModel(
 		FStaticLODModel * SrcModel,
 		FStaticLODModel *& OutModel, 
-		FBoxSphereBounds & Bounds, 
+		const FBoxSphereBounds & Bounds, 
 		float &MaxDeviation, 
 		const FReferenceSkeleton& RefSkeleton, 
 		const FSkeletalMeshOptimizationSettings& Settings
@@ -355,7 +355,7 @@ public:
 		LODModels[LODIndex] = NewModel;
 
 		// Reduce LOD model with SrcMesh
-		if (ReduceLODModel(SrcModel, NewModel, SkeletalMesh->Bounds, MaxDeviation, SkeletalMesh->RefSkeleton, Settings))
+		if (ReduceLODModel(SrcModel, NewModel, SkeletalMesh->GetImportedBounds(), MaxDeviation, SkeletalMesh->RefSkeleton, Settings))
 		{
 			if (bCalcLODDistance)
 			{
diff --git a/Engine/Source/Developer/TargetDeviceServices/Private/TargetDeviceServiceManager.cpp b/Engine/Source/Developer/TargetDeviceServices/Private/TargetDeviceServiceManager.cpp
--- a/Engine/Source/Developer/TargetDeviceServices/Private/TargetDeviceServiceManager.cpp
+++ b/Engine/Source/Developer/TargetDeviceServices/Private/TargetDeviceServiceManager.cpp
@@ -187,7 +187,7 @@ void FTargetDeviceServiceManager::LoadSettings()
 			continue;
 		}
 
-		const FString& ServiceString = *It.Value();
+		const FString& ServiceString = It.Value().GetValue();
 
 		// ... parse device identifier...
 		FString DeviceName;
diff --git a/Engine/Source/Editor/AIGraph/Private/AIGraphNode.cpp b/Engine/Source/Editor/AIGraph/Private/AIGraphNode.cpp
--- a/Engine/Source/Editor/AIGraph/Private/AIGraphNode.cpp
+++ b/Engine/Source/Editor/AIGraph/Private/AIGraphNode.cpp
@@ -69,6 +69,11 @@ void UAIGraphNode::PostEditImport()
 void UAIGraphNode::PostEditUndo()
 {
 	ResetNodeOwner();
+	
+	if (ParentNode)
+	{
+		ParentNode->SubNodes.AddUnique(this);
+	}
 }
 
 #endif
diff --git a/Engine/Source/Editor/AIGraph/Private/AIGraphTypes.cpp b/Engine/Source/Editor/AIGraph/Private/AIGraphTypes.cpp
--- a/Engine/Source/Editor/AIGraph/Private/AIGraphTypes.cpp
+++ b/Engine/Source/Editor/AIGraph/Private/AIGraphTypes.cpp
@@ -15,6 +15,11 @@ FGraphNodeClassData::FGraphNodeClassData(UClass* InClass, const FString& InDepre
 	DeprecatedMessage(InDeprecatedMessage)
 {
 	Category = GetCategory();
+
+	if (InClass)
+	{
+		ClassName = InClass->GetName();
+	}
 }
 
 FGraphNodeClassData::FGraphNodeClassData(const FString& InAssetName, const FString& InGeneratedClassPackage, const FString& InClassName, UClass* InClass) :
diff --git a/Engine/Source/Editor/AddContentDialog/Private/ContentSourceProviders/FeaturePack/FeaturePackContentSource.cpp b/Engine/Source/Editor/AddContentDialog/Private/ContentSourceProviders/FeaturePack/FeaturePackContentSource.cpp
--- a/Engine/Source/Editor/AddContentDialog/Private/ContentSourceProviders/FeaturePack/FeaturePackContentSource.cpp
+++ b/Engine/Source/Editor/AddContentDialog/Private/ContentSourceProviders/FeaturePack/FeaturePackContentSource.cpp
@@ -503,7 +503,7 @@ bool FFeaturePackContentSource::ExtractListOfAdditionalFiles(const FString& InCo
 		{
 			if (FilePair.Key.ToString().Contains("Files"))
 			{
-				AdditionalFilesMap.Add(FilePair.Value);
+				AdditionalFilesMap.Add(FilePair.Value.GetValue());
 			}
 		}
 		BuildListOfAdditionalFiles(AdditionalFilesMap,InFileList,bContainsSource);
diff --git a/Engine/Source/Editor/AnimGraph/Classes/AnimGraphNode_Base.h b/Engine/Source/Editor/AnimGraph/Classes/AnimGraphNode_Base.h
--- a/Engine/Source/Editor/AnimGraph/Classes/AnimGraphNode_Base.h
+++ b/Engine/Source/Editor/AnimGraph/Classes/AnimGraphNode_Base.h
@@ -93,6 +93,14 @@ protected:
 	int32 ChildPropertyIndex;
 };
 
+UENUM()
+enum class EBlueprintUsage : uint8
+{
+	NoProperties,
+	DoesNotUseBlueprint,
+	UsesBlueprint
+};
+
 /**
   * This is the base class for any animation graph nodes that generate or consume an animation pose in
   * the animation blend graph.
@@ -107,6 +115,9 @@ class ANIMGRAPH_API UAnimGraphNode_Base : public UK2Node
 	UPROPERTY(EditAnywhere, BlueprintReadWrite, Category=PinOptions, EditFixedSize)
 	TArray<FOptionalPinFromProperty> ShowPinForProperties;
 
+	UPROPERTY(Transient)
+	EBlueprintUsage BlueprintUsage;
+
 	// UObject interface
 	virtual void PostEditChangeProperty(FPropertyChangedEvent& PropertyChangedEvent) override;
 	// End of UObject interface
diff --git a/Engine/Source/Editor/BehaviorTreeEditor/Classes/BehaviorTreeGraph.h b/Engine/Source/Editor/BehaviorTreeEditor/Classes/BehaviorTreeGraph.h
--- a/Engine/Source/Editor/BehaviorTreeEditor/Classes/BehaviorTreeGraph.h
+++ b/Engine/Source/Editor/BehaviorTreeEditor/Classes/BehaviorTreeGraph.h
@@ -54,6 +54,10 @@ protected:
 
 	void CollectAllNodeInstances(TSet<UObject*>& NodeInstances) override;
 
+#if WITH_EDITOR
+	virtual void PostEditUndo() override;
+#endif
+
 	void UpdateVersion_UnifiedSubNodes();
 	void UpdateVersion_InnerGraphWhitespace();
 };
diff --git a/Engine/Source/Editor/BehaviorTreeEditor/Private/BehaviorTreeGraph.cpp b/Engine/Source/Editor/BehaviorTreeEditor/Private/BehaviorTreeGraph.cpp
--- a/Engine/Source/Editor/BehaviorTreeEditor/Private/BehaviorTreeGraph.cpp
+++ b/Engine/Source/Editor/BehaviorTreeEditor/Private/BehaviorTreeGraph.cpp
@@ -320,7 +320,7 @@ namespace BTGraphHelpers
 			DecoratorInstances[i]->InitializeNode(RootNode, *ExecutionIndex, 0, TreeDepth);
 			if (bInitializeNodes)
 			{
-				DecoratorInstances[i]->InitializeDecorator(ChildIdx);
+				DecoratorInstances[i]->InitializeParentLink(ChildIdx);
 				*ExecutionIndex += 1;
 
 				// make sure that flow abort mode matches - skip for root level nodes
@@ -398,7 +398,7 @@ namespace BTGraphHelpers
 		for (int32 i = 0; i < DecoratorInstances.Num(); i++)
 		{
 			DecoratorInstances[i]->InitializeNode(RootNode, ExecutionIndex, 0, TreeDepth);
-			DecoratorInstances[i]->InitializeDecorator(ChildIdx);
+			DecoratorInstances[i]->InitializeParentLink(ChildIdx);
 			ExecutionIndex++;
 		}
 
@@ -463,12 +463,12 @@ namespace BTGraphHelpers
 		RootNode->Children.Reset();
 		RootNode->Services.Reset();
 
-		// collect services
+		// collect services (composite node)
 		if (RootEdNode->Services.Num())
 		{
-			for (int32 i = 0; i < RootEdNode->Services.Num(); i++)
+			for (int32 ServiceIdx = 0; ServiceIdx < RootEdNode->Services.Num(); ServiceIdx++)
 			{
-				UBTService* ServiceInstance = RootEdNode->Services[i] ? Cast<UBTService>(RootEdNode->Services[i]->NodeInstance) : NULL;
+				UBTService* ServiceInstance = RootEdNode->Services[ServiceIdx] ? Cast<UBTService>(RootEdNode->Services[ServiceIdx]->NodeInstance) : NULL;
 				if (ServiceInstance)
 				{
 					if (Cast<UBehaviorTree>(ServiceInstance->GetOuter()) == NULL)
@@ -527,13 +527,12 @@ namespace BTGraphHelpers
 				CollectDecorators(BTAsset, GraphNode, DecoratorInstances, DecoratorOperations, true, RootNode, ExecutionIndex, TreeDepth, ChildIdx);
 
 				// store child data
-				RootNode->Children.AddZeroed();
+				ChildIdx = RootNode->Children.AddDefaulted();
 				FBTCompositeChild& ChildInfo = RootNode->Children[ChildIdx];
 				ChildInfo.ChildComposite = CompositeInstance;
 				ChildInfo.ChildTask = TaskInstance;
 				ChildInfo.Decorators = DecoratorInstances;
 				ChildInfo.DecoratorOps = DecoratorOperations;
-				ChildIdx++;
 
 				UBTNode* ChildNode = CompositeInstance ? (UBTNode*)CompositeInstance : (UBTNode*)TaskInstance;
 				if (ChildNode && Cast<UBehaviorTree>(ChildNode->GetOuter()) == NULL)
@@ -550,6 +549,29 @@ namespace BTGraphHelpers
 					*ExecutionIndex += SubtreeTask->GetInjectedNodesCount();
 				}
 
+				// collect services (task node)
+				if (TaskInstance)
+				{
+					for (int32 ServiceIdx = 0; ServiceIdx < GraphNode->Services.Num(); ServiceIdx++)
+					{
+						UBTService* ServiceInstance = GraphNode->Services[ServiceIdx] ? Cast<UBTService>(GraphNode->Services[ServiceIdx]->NodeInstance) : NULL;
+						if (ServiceInstance)
+						{
+							if (Cast<UBehaviorTree>(ServiceInstance->GetOuter()) == NULL)
+							{
+								ServiceInstance->Rename(NULL, BTAsset);
+							}
+
+							ServiceInstance->InitializeNode(RootNode, *ExecutionIndex, 0, TreeDepth);
+							ServiceInstance->InitializeParentLink(ChildIdx);
+							*ExecutionIndex += 1;
+
+							TaskInstance->Services.Add(ServiceInstance);
+						}
+					}
+				}
+
+				// assign execution index to child node
 				ChildNode->InitializeNode(RootNode, *ExecutionIndex, 0, TreeDepth);
 				*ExecutionIndex += 1;
 
@@ -1085,9 +1107,21 @@ void UBehaviorTreeGraph::RebuildChildOrder(UEdGraphNode* ParentNode)
 	if (bUpdateExecutionOrder)
 	{
 		UpdateAsset(KeepRebuildCounter);
+		Modify();
 	}
 }
 
+#if WITH_EDITOR
+void UBehaviorTreeGraph::PostEditUndo()
+{
+	Super::PostEditUndo();
+
+	// make sure that all execution indices are up to date
+	UpdateAsset(KeepRebuildCounter);
+	Modify();
+}
+#endif // WITH_EDITOR
+
 namespace BTAutoArrangeHelpers
 {
 	struct FNodeBoundsInfo
@@ -1262,3 +1296,4 @@ void UBehaviorTreeGraph::UpdateVersion_InnerGraphWhitespace()
 		}
 	}
 }
+
diff --git a/Engine/Source/Editor/BehaviorTreeEditor/Private/BehaviorTreeGraphNode.cpp b/Engine/Source/Editor/BehaviorTreeEditor/Private/BehaviorTreeGraphNode.cpp
--- a/Engine/Source/Editor/BehaviorTreeEditor/Private/BehaviorTreeGraphNode.cpp
+++ b/Engine/Source/Editor/BehaviorTreeEditor/Private/BehaviorTreeGraphNode.cpp
@@ -184,6 +184,29 @@ void UBehaviorTreeGraphNode::FindDiffs(UEdGraphNode* OtherNode, FDiffResults& Re
 	}
 }
 
+#if WITH_EDITOR
+void UBehaviorTreeGraphNode::PostEditUndo()
+{
+	Super::PostEditUndo();
+
+	UBehaviorTreeGraphNode* MyParentNode = Cast<UBehaviorTreeGraphNode>(ParentNode);
+	if (MyParentNode)
+	{
+		const bool bIsDecorator = (Cast<UBTDecorator>(NodeInstance) != nullptr) || IsA(UBehaviorTreeGraphNode_CompositeDecorator::StaticClass());
+		const bool bIsService = Cast<UBTService>(NodeInstance) != nullptr;
+
+		if (bIsDecorator)
+		{
+			MyParentNode->Decorators.AddUnique(this);
+		}
+		else if (bIsService)
+		{
+			MyParentNode->Services.AddUnique(this);
+		}
+	}
+}
+#endif
+
 void UBehaviorTreeGraphNode::OnSubNodeAdded(UAIGraphNode* NodeTemplate)
 {
 	UBehaviorTreeGraphNode* BTGraphNode = Cast<UBehaviorTreeGraphNode>(NodeTemplate);
diff --git a/Engine/Source/Editor/BehaviorTreeEditor/Private/BehaviorTreeGraphNode_CompositeDecorator.cpp b/Engine/Source/Editor/BehaviorTreeEditor/Private/BehaviorTreeGraphNode_CompositeDecorator.cpp
--- a/Engine/Source/Editor/BehaviorTreeEditor/Private/BehaviorTreeGraphNode_CompositeDecorator.cpp
+++ b/Engine/Source/Editor/BehaviorTreeEditor/Private/BehaviorTreeGraphNode_CompositeDecorator.cpp
@@ -173,7 +173,7 @@ void UBehaviorTreeGraphNode_CompositeDecorator::SetDecoratorData(class UBTCompos
 void UBehaviorTreeGraphNode_CompositeDecorator::InitializeDecorator(class UBTDecorator* InnerDecorator)
 {
 	InnerDecorator->InitializeNode(ParentNodeInstance, 0, 0, 0);
-	InnerDecorator->InitializeDecorator(ChildIndex);
+	InnerDecorator->InitializeParentLink(ChildIndex);
 }
 
 void UBehaviorTreeGraphNode_CompositeDecorator::OnBlackboardUpdate()
diff --git a/Engine/Source/Editor/BehaviorTreeEditor/Private/BehaviorTreeGraphNode_Task.cpp b/Engine/Source/Editor/BehaviorTreeEditor/Private/BehaviorTreeGraphNode_Task.cpp
--- a/Engine/Source/Editor/BehaviorTreeEditor/Private/BehaviorTreeGraphNode_Task.cpp
+++ b/Engine/Source/Editor/BehaviorTreeEditor/Private/BehaviorTreeGraphNode_Task.cpp
@@ -32,4 +32,5 @@ FText UBehaviorTreeGraphNode_Task::GetNodeTitle(ENodeTitleType::Type TitleType)
 void UBehaviorTreeGraphNode_Task::GetContextMenuActions(const FGraphNodeContextMenuBuilder& Context) const
 {
 	AddContextMenuActionsDecorators(Context);
+	AddContextMenuActionsServices(Context);
 }
diff --git a/Engine/Source/Editor/BehaviorTreeEditor/Private/DetailCustomizations/BlackboardSelectorDetails.cpp b/Engine/Source/Editor/BehaviorTreeEditor/Private/DetailCustomizations/BlackboardSelectorDetails.cpp
--- a/Engine/Source/Editor/BehaviorTreeEditor/Private/DetailCustomizations/BlackboardSelectorDetails.cpp
+++ b/Engine/Source/Editor/BehaviorTreeEditor/Private/DetailCustomizations/BlackboardSelectorDetails.cpp
@@ -78,16 +78,16 @@ void FBlackboardSelectorDetails::CacheBlackboardData()
 	FPropertyAccess::Result Result = MyFilterProperty->GetNumElements(NumElements);
 	if (Result == FPropertyAccess::Success)
 	{
-		for (uint32 i = 0; i < NumElements; i++)
+		for (uint32 Idx = 0; Idx < NumElements; Idx++)
 		{
 			UObject* FilterOb;
-			Result = MyFilterProperty->GetElement(i)->GetValue(FilterOb);
+			Result = MyFilterProperty->GetElement(Idx)->GetValue(FilterOb);
 			if (Result == FPropertyAccess::Success)
 			{
-				UBlackboardKeyType* CasterFilterOb = Cast<UBlackboardKeyType>(FilterOb);
-				if (CasterFilterOb)
+				UBlackboardKeyType* KeyFilterOb = Cast<UBlackboardKeyType>(FilterOb);
+				if (KeyFilterOb)
 				{
-					FilterObjects.Add(CasterFilterOb);
+					FilterObjects.Add(KeyFilterOb);
 				}
 			}
 		}
@@ -95,9 +95,9 @@ void FBlackboardSelectorDetails::CacheBlackboardData()
 
 	TArray<UObject*> MyObjects;
 	MyStructProperty->GetOuterObjects(MyObjects);
-	for (int32 iOb = 0; iOb < MyObjects.Num(); iOb++)
+	for (int32 ObjectIdx = 0; ObjectIdx < MyObjects.Num(); ObjectIdx++)
 	{
-		const UBlackboardData* BlackboardAsset = FindBlackboardAsset(MyObjects[iOb]);
+		const UBlackboardData* BlackboardAsset = FindBlackboardAsset(MyObjects[ObjectIdx]);
 		if (BlackboardAsset)
 		{
 			CachedBlackboardAsset = BlackboardAsset;
@@ -105,9 +105,9 @@ void FBlackboardSelectorDetails::CacheBlackboardData()
 			TArray<FName> ProcessedNames;
 			for (UBlackboardData* It = (UBlackboardData*)BlackboardAsset; It; It = It->Parent)
 			{
-				for (int32 iKey = 0; iKey < It->Keys.Num(); iKey++)
+				for (int32 KeyIdx = 0; KeyIdx < It->Keys.Num(); KeyIdx++)
 				{
-					const FBlackboardEntry& EntryInfo = It->Keys[iKey];
+					const FBlackboardEntry& EntryInfo = It->Keys[KeyIdx];
 					bool bIsKeyOverridden = ProcessedNames.Contains(EntryInfo.EntryName);
 					bool bIsEntryAllowed = !bIsKeyOverridden && (EntryInfo.KeyType != NULL);
 
@@ -116,9 +116,9 @@ void FBlackboardSelectorDetails::CacheBlackboardData()
 					if (bIsEntryAllowed && FilterObjects.Num())
 					{
 						bool bFilterPassed = false;
-						for (int32 iFilter = 0; iFilter < FilterObjects.Num(); iFilter++)
+						for (int32 FilterIdx = 0; FilterIdx < FilterObjects.Num(); FilterIdx++)
 						{
-							if (EntryInfo.KeyType->IsAllowedByFilter(FilterObjects[iFilter]))
+							if (EntryInfo.KeyType->IsAllowedByFilter(FilterObjects[FilterIdx]))
 							{
 								bFilterPassed = true;
 								break;
@@ -151,7 +151,10 @@ void FBlackboardSelectorDetails::InitKeyFromProperty()
 		{
 			if (bNoneIsAllowedValue == false)
 			{
-				OnKeyComboChange(0);
+				const FName PropName = MyStructProperty->GetProperty() ? MyStructProperty->GetProperty()->GetFName() : NAME_None;
+				const int32 KeyNameIdx = KeyValues.IndexOfByKey(PropName);
+
+				OnKeyComboChange(KeyNameIdx != INDEX_NONE ? KeyNameIdx : 0);
 			}
 			else
 			{
@@ -167,10 +170,10 @@ TSharedRef<SWidget> FBlackboardSelectorDetails::OnGetKeyContent() const
 {
 	FMenuBuilder MenuBuilder(true, NULL);
 
-	for (int32 i = 0; i < KeyValues.Num(); i++)
+	for (int32 Idx = 0; Idx < KeyValues.Num(); Idx++)
 	{
-		FUIAction ItemAction( FExecuteAction::CreateSP( this, &FBlackboardSelectorDetails::OnKeyComboChange, i ) );
-		MenuBuilder.AddMenuEntry( FText::FromName( KeyValues[i] ), TAttribute<FText>(), FSlateIcon(), ItemAction);
+		FUIAction ItemAction( FExecuteAction::CreateSP( this, &FBlackboardSelectorDetails::OnKeyComboChange, Idx) );
+		MenuBuilder.AddMenuEntry( FText::FromName( KeyValues[Idx] ), TAttribute<FText>(), FSlateIcon(), ItemAction);
 	}
 
 	return MenuBuilder.MakeWidget();
diff --git a/Engine/Source/Editor/BehaviorTreeEditor/Private/EdGraphSchema_BehaviorTree.cpp b/Engine/Source/Editor/BehaviorTreeEditor/Private/EdGraphSchema_BehaviorTree.cpp
--- a/Engine/Source/Editor/BehaviorTreeEditor/Private/EdGraphSchema_BehaviorTree.cpp
+++ b/Engine/Source/Editor/BehaviorTreeEditor/Private/EdGraphSchema_BehaviorTree.cpp
@@ -378,7 +378,7 @@ const FPinConnectionResponse UEdGraphSchema_BehaviorTree::CanMergeNodes(const UE
 		}
 
 		if ((bNodeAIsDecorator && (bNodeBIsComposite || bNodeBIsTask || bNodeBIsDecorator))
-			|| (bNodeAIsService && (bNodeBIsComposite || bNodeBIsService)))
+			|| (bNodeAIsService && (bNodeBIsComposite || bNodeBIsTask || bNodeBIsService)))
 		{
 			return FPinConnectionResponse(CONNECT_RESPONSE_MAKE, TEXT(""));
 		}
diff --git a/Engine/Source/Editor/BlueprintGraph/Private/K2Node_BaseAsyncTask.cpp b/Engine/Source/Editor/BlueprintGraph/Private/K2Node_BaseAsyncTask.cpp
--- a/Engine/Source/Editor/BlueprintGraph/Private/K2Node_BaseAsyncTask.cpp
+++ b/Engine/Source/Editor/BlueprintGraph/Private/K2Node_BaseAsyncTask.cpp
@@ -477,9 +477,9 @@ UK2Node::ERedirectType UK2Node_BaseAsyncTask::DoPinsMatchForReconstruction(const
 					FString OldPinString;
 					FString NewPinString;
 
-					FParse::Value(*It.Value(), TEXT("ProxyClassName="), ProxyClassString);
-					FParse::Value(*It.Value(), TEXT("OldPinName="), OldPinString);
-					FParse::Value(*It.Value(), TEXT("NewPinName="), NewPinString);
+					FParse::Value(*It.Value().GetValue(), TEXT("ProxyClassName="), ProxyClassString);
+					FParse::Value(*It.Value().GetValue(), TEXT("OldPinName="), OldPinString);
+					FParse::Value(*It.Value().GetValue(), TEXT("NewPinName="), NewPinString);
 
 					UClass* RedirectProxyClass = FindObject<UClass>(ANY_PACKAGE, *ProxyClassString);
 					if (RedirectProxyClass)
diff --git a/Engine/Source/Editor/ContentBrowser/Private/NewAssetOrClassContextMenu.cpp b/Engine/Source/Editor/ContentBrowser/Private/NewAssetOrClassContextMenu.cpp
--- a/Engine/Source/Editor/ContentBrowser/Private/NewAssetOrClassContextMenu.cpp
+++ b/Engine/Source/Editor/ContentBrowser/Private/NewAssetOrClassContextMenu.cpp
@@ -315,7 +315,7 @@ void FNewAssetOrClassContextMenu::MakeContextMenu(
 	if (InOnNewAssetRequested.IsBound())
 	{
 		// Add Basic Asset
-		MenuBuilder.BeginSection("ContentBrowserNewBasicAsset", LOCTEXT("BasicAssetsMenuHeading", "Create Basic Asset") );
+		MenuBuilder.BeginSection("ContentBrowserNewBasicAsset", LOCTEXT("CreateBasicAssetsMenuHeading", "Create Basic Asset") );
 		{
 			CreateNewAssetMenuCategory(
 				MenuBuilder, 
@@ -328,7 +328,7 @@ void FNewAssetOrClassContextMenu::MakeContextMenu(
 		MenuBuilder.EndSection(); //ContentBrowserNewBasicAsset
 
 		// Add Advanced Asset
-		MenuBuilder.BeginSection("ContentBrowserNewAdvancedAsset", LOCTEXT("AdvancedAssetsMenuHeading", "Create Advanced Asset"));
+		MenuBuilder.BeginSection("ContentBrowserNewAdvancedAsset", LOCTEXT("CreateAdvancedAssetsMenuHeading", "Create Advanced Asset"));
 		{
 			FAssetToolsModule& AssetToolsModule = FAssetToolsModule::GetModule();
 
diff --git a/Engine/Source/Editor/ContentBrowser/Private/SAssetView.cpp b/Engine/Source/Editor/ContentBrowser/Private/SAssetView.cpp
--- a/Engine/Source/Editor/ContentBrowser/Private/SAssetView.cpp
+++ b/Engine/Source/Editor/ContentBrowser/Private/SAssetView.cpp
@@ -2542,7 +2542,7 @@ TSharedRef<SWidget> SAssetView::GetViewButtonContent()
 
 		MenuBuilder.AddMenuEntry(
 			LOCTEXT("ShowL10NFolderOption", "Show Localized Assets"),
-			LOCTEXT("ShowFoldersOptionToolTip", "Show assets within the localized asset directory."),
+			LOCTEXT("ShowL10NFolderOptionToolTip", "Show assets within the localized asset directory."),
 			FSlateIcon(),
 			FUIAction(
 			FExecuteAction::CreateSP(this, &SAssetView::ToggleShowL10NFolder),
diff --git a/Engine/Source/Editor/DestructibleMeshEditor/Private/SDestructibleMeshEditorViewport.cpp b/Engine/Source/Editor/DestructibleMeshEditor/Private/SDestructibleMeshEditorViewport.cpp
--- a/Engine/Source/Editor/DestructibleMeshEditor/Private/SDestructibleMeshEditorViewport.cpp
+++ b/Engine/Source/Editor/DestructibleMeshEditor/Private/SDestructibleMeshEditorViewport.cpp
@@ -146,8 +146,9 @@ void FDestructibleMeshEditorViewportClient::SetPreviewComponent( UDestructibleCo
 	UDestructibleMesh* DestructibleMesh = DestructibleMeshEditorPtr.Pin()->GetDestructibleMesh();
 	if (DestructibleMesh != NULL)
 	{
-		SetViewLocation( FVector(0,-DestructibleMesh->Bounds.SphereRadius / (75.0f * (float)PI / 360.0f),0.5f*DestructibleMesh->Bounds.BoxExtent.Z) );
-		SetViewRotation( FRotator(0,90.f,0) );
+		FBoxSphereBounds MeshBounds = DestructibleMesh->GetImportedBounds();
+		SetViewLocation(FVector(0.0f, -MeshBounds.SphereRadius / (75.0f * (float)PI / 360.0f), 0.5f * MeshBounds.BoxExtent.Z));
+		SetViewRotation(FRotator(0.0f, 90.0f, 0.0f));
 	}
 }
 
diff --git a/Engine/Source/Editor/DetailCustomizations/Private/AnimMontageSegmentDetails.cpp b/Engine/Source/Editor/DetailCustomizations/Private/AnimMontageSegmentDetails.cpp
--- a/Engine/Source/Editor/DetailCustomizations/Private/AnimMontageSegmentDetails.cpp
+++ b/Engine/Source/Editor/DetailCustomizations/Private/AnimMontageSegmentDetails.cpp
@@ -235,7 +235,7 @@ void SAnimationSegmentViewport::InitSkeleton()
 				//Place the camera at a good viewer position
 				FVector NewPosition = LevelViewportClient->GetViewLocation();
 				NewPosition.Normalize();
-				LevelViewportClient->SetViewLocation(NewPosition * (PreviewMesh->Bounds.SphereRadius*1.5f));
+				LevelViewportClient->SetViewLocation(NewPosition * (PreviewMesh->GetImportedBounds().SphereRadius*1.5f));
 			}
 		}
 	}
diff --git a/Engine/Source/Editor/DetailCustomizations/Private/AnimSequenceDetails.cpp b/Engine/Source/Editor/DetailCustomizations/Private/AnimSequenceDetails.cpp
--- a/Engine/Source/Editor/DetailCustomizations/Private/AnimSequenceDetails.cpp
+++ b/Engine/Source/Editor/DetailCustomizations/Private/AnimSequenceDetails.cpp
@@ -503,7 +503,7 @@ void SAnimationRefPoseViewport::InitSkeleton()
 				NewPosition.Normalize();
 				if(PreviewSkeletalMesh)
 				{
-					NewPosition *= (PreviewSkeletalMesh->Bounds.SphereRadius*1.5f);
+					NewPosition *= (PreviewSkeletalMesh->GetImportedBounds().SphereRadius*1.5f);
 				}
 				LevelViewportClient->SetViewLocation( NewPosition );
 			}
diff --git a/Engine/Source/Editor/DetailCustomizations/Private/DetailCustomizations.cpp b/Engine/Source/Editor/DetailCustomizations/Private/DetailCustomizations.cpp
--- a/Engine/Source/Editor/DetailCustomizations/Private/DetailCustomizations.cpp
+++ b/Engine/Source/Editor/DetailCustomizations/Private/DetailCustomizations.cpp
@@ -174,6 +174,7 @@ void FDetailCustomizationsModule::RegisterPropertyTypeCustomizations()
 	RegisterCustomPropertyTypeLayout(NAME_Color, FOnGetPropertyTypeCustomizationInstance::CreateStatic(&FColorStructCustomization::MakeInstance));
 	RegisterCustomPropertyTypeLayout(NAME_Matrix, FOnGetPropertyTypeCustomizationInstance::CreateStatic(&FMatrixStructCustomization::MakeInstance));
 	RegisterCustomPropertyTypeLayout(NAME_Transform, FOnGetPropertyTypeCustomizationInstance::CreateStatic(&FTransformStructCustomization::MakeInstance));
+	RegisterCustomPropertyTypeLayout(NAME_Quat, FOnGetPropertyTypeCustomizationInstance::CreateStatic(&FQuatStructCustomization::MakeInstance));
 	RegisterCustomPropertyTypeLayout("SlateColor", FOnGetPropertyTypeCustomizationInstance::CreateStatic(&FSlateColorCustomization::MakeInstance));
 	RegisterCustomPropertyTypeLayout("AttenuationSettings", FOnGetPropertyTypeCustomizationInstance::CreateStatic(&FAttenuationSettingsCustomization::MakeInstance));
 	RegisterCustomPropertyTypeLayout("DialogueContext", FOnGetPropertyTypeCustomizationInstance::CreateStatic(&FDialogueContextStructCustomization::MakeInstance));
diff --git a/Engine/Source/Editor/DetailCustomizations/Private/MathStructProxyCustomizations.cpp b/Engine/Source/Editor/DetailCustomizations/Private/MathStructProxyCustomizations.cpp
--- a/Engine/Source/Editor/DetailCustomizations/Private/MathStructProxyCustomizations.cpp
+++ b/Engine/Source/Editor/DetailCustomizations/Private/MathStructProxyCustomizations.cpp
@@ -17,7 +17,7 @@ void FMathStructProxyCustomization::MakeHeaderRow( TSharedRef<class IPropertyHan
 }
 
 template<typename ProxyType, typename NumericType>
-TSharedRef<SWidget> FMathStructProxyCustomization::MakeNumericProxyWidget(TSharedRef<IPropertyHandle>& StructPropertyHandle, TSharedRef< TProxyProperty<ProxyType, NumericType> >& ProxyValue, const FText& Label, const FLinearColor& LabelColor, const FLinearColor& LabelBackgroundColor)
+TSharedRef<SWidget> FMathStructProxyCustomization::MakeNumericProxyWidget(TSharedRef<IPropertyHandle>& StructPropertyHandle, TSharedRef< TProxyProperty<ProxyType, NumericType> >& ProxyValue, const FText& Label, bool bRotationInDegrees, const FLinearColor& LabelColor, const FLinearColor& LabelBackgroundColor)
 {
 	TWeakPtr<IPropertyHandle> WeakHandlePtr = StructPropertyHandle;
 
@@ -36,8 +36,8 @@ TSharedRef<SWidget> FMathStructProxyCustomization::MakeNumericProxyWidget(TShare
 		.AllowSpin( StructPropertyHandle->GetNumOuterObjects() == 1 )
 		.MinValue(TOptional<NumericType>())
 		.MaxValue(TOptional<NumericType>())
-		.MaxSliderValue(TOptional<NumericType>())
-		.MinSliderValue(TOptional<NumericType>())
+		.MaxSliderValue(bRotationInDegrees ? 360.0f : TOptional<NumericType>())
+		.MinSliderValue(bRotationInDegrees ? 0.0f : TOptional<NumericType>())
 		.Label()
 		[
 			SNumericEntryBox<float>::BuildLabel( Label, LabelColor, LabelBackgroundColor )
@@ -110,98 +110,119 @@ void FMatrixStructCustomization::MakeHeaderRow(TSharedRef<class IPropertyHandle>
 	];
 }
 
-void FMatrixStructCustomization::CustomizeChildren(TSharedRef<class IPropertyHandle> StructPropertyHandle, class IDetailChildrenBuilder& StructBuilder, IPropertyTypeCustomizationUtils& StructCustomizationUtils)
+void FMatrixStructCustomization::CustomizeLocation(TSharedRef<class IPropertyHandle> StructPropertyHandle, FDetailWidgetRow& Row)
 {
-	FMathStructProxyCustomization::CustomizeChildren(StructPropertyHandle, StructBuilder, StructCustomizationUtils);
-
 	TWeakPtr<IPropertyHandle> WeakHandlePtr = StructPropertyHandle;
 
-	StructBuilder.AddChildContent(LOCTEXT("LocationLabel", "Location"))
-		.CopyAction(FUIAction(FExecuteAction::CreateSP(this, &FMatrixStructCustomization::OnCopy, FTransformField::Location, WeakHandlePtr)))
-		.PasteAction(FUIAction(FExecuteAction::CreateSP(this, &FMatrixStructCustomization::OnPaste, FTransformField::Location, WeakHandlePtr)))
-		.NameContent()
+	Row
+	.CopyAction(FUIAction(FExecuteAction::CreateSP(this, &FMatrixStructCustomization::OnCopy, FTransformField::Location, WeakHandlePtr)))
+	.PasteAction(FUIAction(FExecuteAction::CreateSP(this, &FMatrixStructCustomization::OnPaste, FTransformField::Location, WeakHandlePtr)))
+	.NameContent()
+	[
+		StructPropertyHandle->CreatePropertyNameWidget(LOCTEXT("LocationLabel", "Location"))
+	]
+	.ValueContent()
+	.MinDesiredWidth(375.0f)
+	.MaxDesiredWidth(375.0f)
+	[
+		SNew(SHorizontalBox)
+		+ SHorizontalBox::Slot()
+		.Padding(FMargin(0.0f, 2.0f, 3.0f, 2.0f))
 		[
-			StructPropertyHandle->CreatePropertyNameWidget(LOCTEXT("LocationLabel", "Location"))
+			MakeNumericProxyWidget<FVector, float>(StructPropertyHandle, CachedTranslationX, LOCTEXT("TranslationX", "X"), false, FLinearColor::White, SNumericEntryBox<float>::RedLabelBackgroundColor)
 		]
-		.ValueContent()
-		.MinDesiredWidth(375.0f)
-		.MaxDesiredWidth(375.0f)
+		+ SHorizontalBox::Slot()
+		.Padding(FMargin(0.0f, 2.0f, 3.0f, 2.0f))
 		[
-			SNew(SHorizontalBox)
-			+ SHorizontalBox::Slot()
-			.Padding(FMargin(0.0f, 2.0f, 3.0f, 2.0f))
-			[
-				MakeNumericProxyWidget<FVector, float>(StructPropertyHandle, CachedTranslationX, LOCTEXT("TranslationX", "X"), FLinearColor::White, SNumericEntryBox<float>::RedLabelBackgroundColor)
-			]
-			+ SHorizontalBox::Slot()
-			.Padding(FMargin(0.0f, 2.0f, 3.0f, 2.0f))
-			[
-				MakeNumericProxyWidget<FVector, float>(StructPropertyHandle, CachedTranslationY, LOCTEXT("TranslationY", "Y"), FLinearColor::White, SNumericEntryBox<float>::GreenLabelBackgroundColor)
-			]
-			+ SHorizontalBox::Slot()
-			.Padding(FMargin(0.0f, 2.0f, 0.0f, 2.0f))
-			[
-				MakeNumericProxyWidget<FVector, float>(StructPropertyHandle, CachedTranslationZ, LOCTEXT("TranslationZ", "Z"), FLinearColor::White, SNumericEntryBox<float>::BlueLabelBackgroundColor)
-			]
-		];
+			MakeNumericProxyWidget<FVector, float>(StructPropertyHandle, CachedTranslationY, LOCTEXT("TranslationY", "Y"), false, FLinearColor::White, SNumericEntryBox<float>::GreenLabelBackgroundColor)
+		]
+		+ SHorizontalBox::Slot()
+		.Padding(FMargin(0.0f, 2.0f, 0.0f, 2.0f))
+		[
+			MakeNumericProxyWidget<FVector, float>(StructPropertyHandle, CachedTranslationZ, LOCTEXT("TranslationZ", "Z"), false, FLinearColor::White, SNumericEntryBox<float>::BlueLabelBackgroundColor)
+		]
+	];
+}
+
 
-	StructBuilder.AddChildContent(LOCTEXT("RotationLabel", "Rotation"))
-		.CopyAction(FUIAction(FExecuteAction::CreateSP(this, &FMatrixStructCustomization::OnCopy, FTransformField::Rotation, WeakHandlePtr)))
-		.PasteAction(FUIAction(FExecuteAction::CreateSP(this, &FMatrixStructCustomization::OnPaste, FTransformField::Rotation, WeakHandlePtr)))
-		.NameContent()
+void FMatrixStructCustomization::CustomizeRotation(TSharedRef<class IPropertyHandle> StructPropertyHandle, FDetailWidgetRow& Row)
+{
+	TWeakPtr<IPropertyHandle> WeakHandlePtr = StructPropertyHandle;
+
+	Row
+	.CopyAction(FUIAction(FExecuteAction::CreateSP(this, &FMatrixStructCustomization::OnCopy, FTransformField::Rotation, WeakHandlePtr)))
+	.PasteAction(FUIAction(FExecuteAction::CreateSP(this, &FMatrixStructCustomization::OnPaste, FTransformField::Rotation, WeakHandlePtr)))
+	.NameContent()
+	[
+		StructPropertyHandle->CreatePropertyNameWidget(LOCTEXT("RotationLabel", "Rotation"))
+	]
+	.ValueContent()
+	.MinDesiredWidth(375.0f)
+	.MaxDesiredWidth(375.0f)
+	[
+		SNew(SHorizontalBox)
+		+ SHorizontalBox::Slot()
+		.Padding(FMargin(0.0f, 2.0f, 3.0f, 2.0f))
 		[
-			StructPropertyHandle->CreatePropertyNameWidget(LOCTEXT("RotationLabel", "Rotation"))
+			MakeNumericProxyWidget<FRotator, float>(StructPropertyHandle, CachedRotationRoll, LOCTEXT("RotationRoll", "X"), true, FLinearColor::White, SNumericEntryBox<float>::RedLabelBackgroundColor)
 		]
-		.ValueContent()
-		.MinDesiredWidth(375.0f)
-		.MaxDesiredWidth(375.0f)
+		+ SHorizontalBox::Slot()
+		.Padding(FMargin(0.0f, 2.0f, 3.0f, 2.0f))
 		[
-			SNew(SHorizontalBox)
-			+ SHorizontalBox::Slot()
-			.Padding(FMargin(0.0f, 2.0f, 3.0f, 2.0f))
-			[
-				MakeNumericProxyWidget<FRotator, float>(StructPropertyHandle, CachedRotationRoll, LOCTEXT("RotationRoll", "X"), FLinearColor::White, SNumericEntryBox<float>::RedLabelBackgroundColor)
-			]
-			+ SHorizontalBox::Slot()
-			.Padding(FMargin(0.0f, 2.0f, 3.0f, 2.0f))
-			[
-				MakeNumericProxyWidget<FRotator, float>(StructPropertyHandle, CachedRotationPitch, LOCTEXT("RotationPitch", "Y"), FLinearColor::White, SNumericEntryBox<float>::GreenLabelBackgroundColor)
-			]
-			+ SHorizontalBox::Slot()
-			.Padding(FMargin(0.0f, 2.0f, 0.0f, 2.0f))
-			[
-				MakeNumericProxyWidget<FRotator, float>(StructPropertyHandle, CachedRotationYaw, LOCTEXT("RotationYaw", "Z"), FLinearColor::White, SNumericEntryBox<float>::BlueLabelBackgroundColor)
-			]
-		];
+			MakeNumericProxyWidget<FRotator, float>(StructPropertyHandle, CachedRotationPitch, LOCTEXT("RotationPitch", "Y"), true, FLinearColor::White, SNumericEntryBox<float>::GreenLabelBackgroundColor)
+		]
+		+ SHorizontalBox::Slot()
+		.Padding(FMargin(0.0f, 2.0f, 0.0f, 2.0f))
+		[
+			MakeNumericProxyWidget<FRotator, float>(StructPropertyHandle, CachedRotationYaw, LOCTEXT("RotationYaw", "Z"), true, FLinearColor::White, SNumericEntryBox<float>::BlueLabelBackgroundColor)
+		]
+	];
+}
 
-	StructBuilder.AddChildContent(LOCTEXT("ScaleLabel", "Scale"))
-		.CopyAction(FUIAction(FExecuteAction::CreateSP(this, &FMatrixStructCustomization::OnCopy, FTransformField::Scale, WeakHandlePtr)))
-		.PasteAction(FUIAction(FExecuteAction::CreateSP(this, &FMatrixStructCustomization::OnPaste, FTransformField::Scale, WeakHandlePtr)))
-		.NameContent()
+
+void FMatrixStructCustomization::CustomizeScale(TSharedRef<class IPropertyHandle> StructPropertyHandle, FDetailWidgetRow& Row)
+{
+	TWeakPtr<IPropertyHandle> WeakHandlePtr = StructPropertyHandle;
+
+	Row
+	.CopyAction(FUIAction(FExecuteAction::CreateSP(this, &FMatrixStructCustomization::OnCopy, FTransformField::Scale, WeakHandlePtr)))
+	.PasteAction(FUIAction(FExecuteAction::CreateSP(this, &FMatrixStructCustomization::OnPaste, FTransformField::Scale, WeakHandlePtr)))
+	.NameContent()
+	[
+		StructPropertyHandle->CreatePropertyNameWidget(LOCTEXT("ScaleLabel", "Scale"))
+	]
+	.ValueContent()
+	.MinDesiredWidth(375.0f)
+	.MaxDesiredWidth(375.0f)
+	[
+		SNew(SHorizontalBox)
+		+ SHorizontalBox::Slot()
+		.Padding(FMargin(0.0f, 2.0f, 3.0f, 2.0f))
 		[
-			StructPropertyHandle->CreatePropertyNameWidget(LOCTEXT("ScaleLabel", "Scale"))
+			MakeNumericProxyWidget<FVector, float>(StructPropertyHandle, CachedScaleX, LOCTEXT("ScaleX", "X"), false, FLinearColor::White, SNumericEntryBox<float>::RedLabelBackgroundColor)
 		]
-		.ValueContent()
-		.MinDesiredWidth(375.0f)
-		.MaxDesiredWidth(375.0f)
+		+ SHorizontalBox::Slot()
+		.Padding(FMargin(0.0f, 2.0f, 3.0f, 2.0f))
 		[
-			SNew(SHorizontalBox)
-			+ SHorizontalBox::Slot()
-			.Padding(FMargin(0.0f, 2.0f, 3.0f, 2.0f))
-			[
-				MakeNumericProxyWidget<FVector, float>(StructPropertyHandle, CachedScaleX, LOCTEXT("ScaleX", "X"), FLinearColor::White, SNumericEntryBox<float>::RedLabelBackgroundColor)
-			]
-			+ SHorizontalBox::Slot()
-			.Padding(FMargin(0.0f, 2.0f, 3.0f, 2.0f))
-			[
-				MakeNumericProxyWidget<FVector, float>(StructPropertyHandle, CachedScaleY, LOCTEXT("ScaleY", "Y"), FLinearColor::White, SNumericEntryBox<float>::GreenLabelBackgroundColor)
-			]
-			+ SHorizontalBox::Slot()
-			.Padding(FMargin(0.0f, 2.0f, 0.0f, 2.0f))
-			[
-				MakeNumericProxyWidget<FVector, float>(StructPropertyHandle, CachedScaleZ, LOCTEXT("ScaleZ", "Z"), FLinearColor::White, SNumericEntryBox<float>::BlueLabelBackgroundColor)
-			]
-		];
+			MakeNumericProxyWidget<FVector, float>(StructPropertyHandle, CachedScaleY, LOCTEXT("ScaleY", "Y"), false, FLinearColor::White, SNumericEntryBox<float>::GreenLabelBackgroundColor)
+		]
+		+ SHorizontalBox::Slot()
+		.Padding(FMargin(0.0f, 2.0f, 0.0f, 2.0f))
+		[
+			MakeNumericProxyWidget<FVector, float>(StructPropertyHandle, CachedScaleZ, LOCTEXT("ScaleZ", "Z"), false, FLinearColor::White, SNumericEntryBox<float>::BlueLabelBackgroundColor)
+		]
+	];
+}
+
+void FMatrixStructCustomization::CustomizeChildren(TSharedRef<class IPropertyHandle> StructPropertyHandle, class IDetailChildrenBuilder& StructBuilder, IPropertyTypeCustomizationUtils& StructCustomizationUtils)
+{
+	FMathStructProxyCustomization::CustomizeChildren(StructPropertyHandle, StructBuilder, StructCustomizationUtils);
+
+	TWeakPtr<IPropertyHandle> WeakHandlePtr = StructPropertyHandle;
+
+	CustomizeLocation(StructPropertyHandle, StructBuilder.AddChildContent(LOCTEXT("RotationLabel", "Rotation")));
+	CustomizeRotation(StructPropertyHandle, StructBuilder.AddChildContent(LOCTEXT("LocationLabel", "Location")));
+	CustomizeScale(StructPropertyHandle, StructBuilder.AddChildContent(LOCTEXT("ScaleLabel", "Scale")));
 }
 
 void FMatrixStructCustomization::OnCopy(FTransformField::Type Type, TWeakPtr<IPropertyHandle> PropertyHandlePtr)
@@ -461,4 +482,88 @@ bool FTransformStructCustomization::FlushValues( TWeakPtr<IPropertyHandle> Prope
 	}
 
 	return true;
-}
\ No newline at end of file
+}
+
+
+TSharedRef<IPropertyTypeCustomization> FQuatStructCustomization::MakeInstance()
+{
+	return MakeShareable(new FQuatStructCustomization);
+}
+
+
+void FQuatStructCustomization::MakeHeaderRow(TSharedRef<class IPropertyHandle>& InStructPropertyHandle, FDetailWidgetRow& Row)
+{
+	CustomizeRotation(InStructPropertyHandle, Row);
+}
+
+
+void FQuatStructCustomization::CustomizeChildren(TSharedRef<class IPropertyHandle> StructPropertyHandle, class IDetailChildrenBuilder& StructBuilder, IPropertyTypeCustomizationUtils& StructCustomizationUtils)
+{
+	FMathStructProxyCustomization::CustomizeChildren(StructPropertyHandle, StructBuilder, StructCustomizationUtils);
+}
+
+
+bool FQuatStructCustomization::CacheValues(TWeakPtr<IPropertyHandle> PropertyHandlePtr) const
+{
+	auto PropertyHandle = PropertyHandlePtr.Pin();
+
+	if (!PropertyHandle.IsValid())
+	{
+		return false;
+	}
+
+	TArray<void*> RawData;
+	PropertyHandle->AccessRawData(RawData);
+
+	if (RawData.Num() == 1)
+	{
+		FQuat* QuatValue = reinterpret_cast<FQuat*>(RawData[0]);
+		if (QuatValue != NULL)
+		{
+			CachedRotation->Set(QuatValue->Rotator());
+			return true;
+		}
+	}
+
+	return false;
+}
+
+bool FQuatStructCustomization::FlushValues(TWeakPtr<IPropertyHandle> PropertyHandlePtr) const
+{
+	auto PropertyHandle = PropertyHandlePtr.Pin();
+
+	if (!PropertyHandle.IsValid())
+	{
+		return false;
+	}
+
+	TArray<void*> RawData;
+	PropertyHandle->AccessRawData(RawData);
+
+	PropertyHandle->NotifyPreChange();
+	for (int32 ValueIndex = 0; ValueIndex < RawData.Num(); ValueIndex++)
+	{
+		FQuat* QuatValue = reinterpret_cast<FQuat*>(RawData[0]);
+		if (QuatValue != NULL)
+		{
+			const FRotator CurrentRotation = QuatValue->Rotator();
+
+			FRotator Rotation(
+				CachedRotationPitch->IsSet() ? CachedRotationPitch->Get() : CurrentRotation.Pitch,
+				CachedRotationYaw->IsSet() ? CachedRotationYaw->Get() : CurrentRotation.Yaw,
+				CachedRotationRoll->IsSet() ? CachedRotationRoll->Get() : CurrentRotation.Roll
+				);
+			*QuatValue = Rotation.Quaternion();
+		}
+	}
+	PropertyHandle->NotifyPostChange();
+
+	if (PropertyUtilities.IsValid() && !bIsUsingSlider)
+	{
+		FPropertyChangedEvent ChangeEvent(PropertyHandle->GetProperty(), EPropertyChangeType::ValueSet);
+		PropertyUtilities->NotifyFinishedChangingProperties(ChangeEvent);
+	}
+
+	return true;
+}
+
diff --git a/Engine/Source/Editor/DetailCustomizations/Private/MathStructProxyCustomizations.h b/Engine/Source/Editor/DetailCustomizations/Private/MathStructProxyCustomizations.h
--- a/Engine/Source/Editor/DetailCustomizations/Private/MathStructProxyCustomizations.h
+++ b/Engine/Source/Editor/DetailCustomizations/Private/MathStructProxyCustomizations.h
@@ -164,10 +164,11 @@ protected:
 	 * @param	StructPropertyHandle	Property handle to the containing struct
 	 * @param	ProxyValue				The value we will be editing in the proxy data.
 	 * @param	Label					A label to use for this value.
+	 * @param	bRotationInDegrees		Whether this is to be used for a rotation value (configures sliders appropriately).
 	 * @return the newly created widget.
 	 */
 	template<typename ProxyType, typename NumericType>
-	TSharedRef<SWidget> MakeNumericProxyWidget(TSharedRef<IPropertyHandle>& StructPropertyHandle, TSharedRef< TProxyProperty<ProxyType, NumericType> >& ProxyValue, const FText& Label, const FLinearColor& LabelColor = FCoreStyle::Get().GetColor("DefaultForeground"), const FLinearColor& LabelBackgroundColor = FCoreStyle::Get().GetColor("InvertedForeground"));
+	TSharedRef<SWidget> MakeNumericProxyWidget(TSharedRef<IPropertyHandle>& StructPropertyHandle, TSharedRef< TProxyProperty<ProxyType, NumericType> >& ProxyValue, const FText& Label, bool bRotationInDegrees = false, const FLinearColor& LabelColor = FCoreStyle::Get().GetColor("DefaultForeground"), const FLinearColor& LabelBackgroundColor = FCoreStyle::Get().GetColor("InvertedForeground"));
 
 private:
 	/**
@@ -240,6 +241,11 @@ public:
 	virtual void CustomizeChildren(TSharedRef<class IPropertyHandle> StructPropertyHandle, class IDetailChildrenBuilder& StructBuilder, IPropertyTypeCustomizationUtils& StructCustomizationUtils) override;
 
 protected:
+	/** Customization utility functions */
+	void CustomizeLocation(TSharedRef<class IPropertyHandle> StructPropertyHandle, FDetailWidgetRow& Row);
+	void CustomizeRotation(TSharedRef<class IPropertyHandle> StructPropertyHandle, FDetailWidgetRow& Row);
+	void CustomizeScale(TSharedRef<class IPropertyHandle> StructPropertyHandle, FDetailWidgetRow& Row);
+
 	/** FMathStructCustomization interface */
 	virtual void MakeHeaderRow(TSharedRef<class IPropertyHandle>& InStructPropertyHandle, FDetailWidgetRow& Row) override;
 
@@ -292,4 +298,24 @@ protected:
 	/** FMathStructProxyCustomization interface */
 	virtual bool CacheValues(TWeakPtr<IPropertyHandle> PropertyHandlePtr) const override;
 	virtual bool FlushValues(TWeakPtr<IPropertyHandle> PropertyHandlePtr) const override;
-};
\ No newline at end of file
+};
+
+/**
+ * Proxy struct customization that displays an FQuat as an euler rotation
+ */
+class FQuatStructCustomization : public FMatrixStructCustomization
+{
+public:
+	static TSharedRef<IPropertyTypeCustomization> MakeInstance();
+
+	/** IPropertyTypeCustomization interface */
+	virtual void CustomizeChildren(TSharedRef<class IPropertyHandle> StructPropertyHandle, class IDetailChildrenBuilder& StructBuilder, IPropertyTypeCustomizationUtils& StructCustomizationUtils) override;
+
+	/** FMathStructCustomization interface */
+	virtual void MakeHeaderRow(TSharedRef<class IPropertyHandle>& InStructPropertyHandle, FDetailWidgetRow& Row) override;
+
+protected:
+	/** FMathStructProxyCustomization interface */
+	virtual bool CacheValues(TWeakPtr<IPropertyHandle> PropertyHandlePtr) const override;
+	virtual bool FlushValues(TWeakPtr<IPropertyHandle> PropertyHandlePtr) const override;
+};
diff --git a/Engine/Source/Editor/EditorStyle/Private/SlateEditorStyle.cpp b/Engine/Source/Editor/EditorStyle/Private/SlateEditorStyle.cpp
--- a/Engine/Source/Editor/EditorStyle/Private/SlateEditorStyle.cpp
+++ b/Engine/Source/Editor/EditorStyle/Private/SlateEditorStyle.cpp
@@ -3914,6 +3914,7 @@ void FSlateEditorStyle::FStyle::SetupGraphEditorStyles()
 
 		{
 			Set( "Graph.AnimationResultNode.Body", new IMAGE_BRUSH( "/Graph/Animation/AnimationNode_Result_128x", FVector2D(128, 128) ) );
+			Set( "Graph.AnimationFastPathIndicator", new IMAGE_BRUSH( "/Graph/Animation/AnimationNode_FastPath", Icon32x32 ) );
 		}
 
 		// SoundCueEditor Graph Nodes
diff --git a/Engine/Source/Editor/GraphEditor/Private/AnimationNodes/SAnimationGraphNode.cpp b/Engine/Source/Editor/GraphEditor/Private/AnimationNodes/SAnimationGraphNode.cpp
new file mode 100644
--- /dev/null
+++ b/Engine/Source/Editor/GraphEditor/Private/AnimationNodes/SAnimationGraphNode.cpp
@@ -0,0 +1,41 @@
+// Copyright 1998-2016 Epic Games, Inc. All Rights Reserved.
+
+#include "GraphEditorCommon.h"
+#include "SAnimationGraphNode.h"
+#include "AnimGraphNode_Base.h"
+#include "IDocumentation.h"
+
+void SAnimationGraphNode::Construct(const FArguments& InArgs, UAnimGraphNode_Base* InNode)
+{
+	this->GraphNode = InNode;
+
+	this->SetCursor(EMouseCursor::CardinalCross);
+
+	this->UpdateGraphNode();
+
+	const FSlateBrush* ImageBrush = FEditorStyle::Get().GetBrush(TEXT("Graph.AnimationFastPathIndicator"));
+
+	IndicatorWidget =
+		SNew(SImage)
+		.Image(ImageBrush)
+		.ToolTip(IDocumentation::Get()->CreateToolTip(NSLOCTEXT("AnimationGraphNode", "AnimGraphNodeIndicatorTooltip", "Fast path enabled: This node is not using any Blueprint calls to update its data."), NULL, TEXT("Shared/GraphNodes/Animation"), TEXT("GraphNode_FastPathInfo")))
+		.Visibility(EVisibility::Visible);
+}
+
+TArray<FOverlayWidgetInfo> SAnimationGraphNode::GetOverlayWidgets(bool bSelected, const FVector2D& WidgetSize) const
+{
+	TArray<FOverlayWidgetInfo> Widgets;
+
+	if (GraphNode != nullptr && CastChecked<UAnimGraphNode_Base>(GraphNode)->BlueprintUsage == EBlueprintUsage::DoesNotUseBlueprint)
+	{
+		const FSlateBrush* ImageBrush = FEditorStyle::Get().GetBrush(TEXT("Graph.AnimationFastPathIndicator"));
+
+		FOverlayWidgetInfo Info;
+		Info.OverlayOffset = FVector2D(WidgetSize.X - (ImageBrush->ImageSize.X * 0.5f), -(ImageBrush->ImageSize.Y * 0.5f));
+		Info.Widget = IndicatorWidget;
+
+		Widgets.Add(Info);
+	}
+
+	return Widgets;
+}
\ No newline at end of file
diff --git a/Engine/Source/Editor/GraphEditor/Private/AnimationNodes/SAnimationGraphNode.h b/Engine/Source/Editor/GraphEditor/Private/AnimationNodes/SAnimationGraphNode.h
new file mode 100644
--- /dev/null
+++ b/Engine/Source/Editor/GraphEditor/Private/AnimationNodes/SAnimationGraphNode.h
@@ -0,0 +1,23 @@
+// Copyright 1998-2016 Epic Games, Inc. All Rights Reserved.
+
+#pragma once
+
+#include "SGraphNodeK2Base.h"
+
+class SAnimationGraphNode : public SGraphNodeK2Base
+{
+public:
+	SLATE_BEGIN_ARGS(SAnimationGraphNode) {}
+	SLATE_END_ARGS()
+
+	void Construct(const FArguments& InArgs, UAnimGraphNode_Base* InNode);
+
+protected:
+	// SGraphNode interface
+	virtual TArray<FOverlayWidgetInfo> GetOverlayWidgets(bool bSelected, const FVector2D& WidgetSize) const override;
+	// End of SGraphNode interface
+
+private:
+	/** Keep a reference to the indicator widget handing around */
+	TSharedPtr<SWidget> IndicatorWidget;
+};
\ No newline at end of file
diff --git a/Engine/Source/Editor/GraphEditor/Private/AnimationNodes/SGraphNodeAnimationResult.cpp b/Engine/Source/Editor/GraphEditor/Private/AnimationNodes/SGraphNodeAnimationResult.cpp
--- a/Engine/Source/Editor/GraphEditor/Private/AnimationNodes/SGraphNodeAnimationResult.cpp
+++ b/Engine/Source/Editor/GraphEditor/Private/AnimationNodes/SGraphNodeAnimationResult.cpp
@@ -16,6 +16,8 @@ void SGraphNodeAnimationResult::Construct(const FArguments& InArgs, UAnimGraphNo
 	this->SetCursor(EMouseCursor::CardinalCross);
 
 	this->UpdateGraphNode();
+
+	SAnimationGraphNode::Construct(SAnimationGraphNode::FArguments(), InNode);
 }
 
 TSharedRef<SWidget> SGraphNodeAnimationResult::CreateNodeContentArea()
diff --git a/Engine/Source/Editor/GraphEditor/Private/AnimationNodes/SGraphNodeAnimationResult.h b/Engine/Source/Editor/GraphEditor/Private/AnimationNodes/SGraphNodeAnimationResult.h
--- a/Engine/Source/Editor/GraphEditor/Private/AnimationNodes/SGraphNodeAnimationResult.h
+++ b/Engine/Source/Editor/GraphEditor/Private/AnimationNodes/SGraphNodeAnimationResult.h
@@ -2,7 +2,9 @@
 
 #pragma once
 
-class SGraphNodeAnimationResult : public SGraphNode
+#include "SAnimationGraphNode.h"
+
+class SGraphNodeAnimationResult : public SAnimationGraphNode
 {
 public:
 	SLATE_BEGIN_ARGS(SGraphNodeAnimationResult){}
diff --git a/Engine/Source/Editor/GraphEditor/Private/AnimationNodes/SGraphNodeLayeredBoneBlend.cpp b/Engine/Source/Editor/GraphEditor/Private/AnimationNodes/SGraphNodeLayeredBoneBlend.cpp
--- a/Engine/Source/Editor/GraphEditor/Private/AnimationNodes/SGraphNodeLayeredBoneBlend.cpp
+++ b/Engine/Source/Editor/GraphEditor/Private/AnimationNodes/SGraphNodeLayeredBoneBlend.cpp
@@ -15,6 +15,8 @@ void SGraphNodeLayeredBoneBlend::Construct(const FArguments& InArgs, UAnimGraphN
 	this->SetCursor(EMouseCursor::CardinalCross);
 
 	this->UpdateGraphNode();
+
+	SAnimationGraphNode::Construct(SAnimationGraphNode::FArguments(), InNode);
 }
 
 void SGraphNodeLayeredBoneBlend::CreateInputSideAddButton(TSharedPtr<SVerticalBox> InputBox)
diff --git a/Engine/Source/Editor/GraphEditor/Private/AnimationNodes/SGraphNodeLayeredBoneBlend.h b/Engine/Source/Editor/GraphEditor/Private/AnimationNodes/SGraphNodeLayeredBoneBlend.h
--- a/Engine/Source/Editor/GraphEditor/Private/AnimationNodes/SGraphNodeLayeredBoneBlend.h
+++ b/Engine/Source/Editor/GraphEditor/Private/AnimationNodes/SGraphNodeLayeredBoneBlend.h
@@ -2,9 +2,11 @@
 
 #pragma once
 
+#include "SAnimationGraphNode.h"
+
 class UAnimGraphNode_LayeredBoneBlend;
 
-class SGraphNodeLayeredBoneBlend : public SGraphNode
+class SGraphNodeLayeredBoneBlend : public SAnimationGraphNode
 {
 public:
 	SLATE_BEGIN_ARGS(SGraphNodeLayeredBoneBlend){}
diff --git a/Engine/Source/Editor/GraphEditor/Private/AnimationNodes/SGraphNodeSequencePlayer.cpp b/Engine/Source/Editor/GraphEditor/Private/AnimationNodes/SGraphNodeSequencePlayer.cpp
--- a/Engine/Source/Editor/GraphEditor/Private/AnimationNodes/SGraphNodeSequencePlayer.cpp
+++ b/Engine/Source/Editor/GraphEditor/Private/AnimationNodes/SGraphNodeSequencePlayer.cpp
@@ -11,13 +11,15 @@
 /////////////////////////////////////////////////////
 // SGraphNodeSequencePlayer
 
-void SGraphNodeSequencePlayer::Construct(const FArguments& InArgs, UK2Node* InNode)
+void SGraphNodeSequencePlayer::Construct(const FArguments& InArgs, UAnimGraphNode_Base* InNode)
 {
 	this->GraphNode = InNode;
 
 	this->SetCursor(EMouseCursor::CardinalCross);
 
 	this->UpdateGraphNode();
+
+	SAnimationGraphNode::Construct(SAnimationGraphNode::FArguments(), InNode);
 }
 
 void SGraphNodeSequencePlayer::GetNodeInfoPopups(FNodeInfoContext* Context, TArray<FGraphInformationPopupInfo>& Popups) const
diff --git a/Engine/Source/Editor/GraphEditor/Private/AnimationNodes/SGraphNodeSequencePlayer.h b/Engine/Source/Editor/GraphEditor/Private/AnimationNodes/SGraphNodeSequencePlayer.h
--- a/Engine/Source/Editor/GraphEditor/Private/AnimationNodes/SGraphNodeSequencePlayer.h
+++ b/Engine/Source/Editor/GraphEditor/Private/AnimationNodes/SGraphNodeSequencePlayer.h
@@ -2,13 +2,15 @@
 
 #pragma once
 
-class SGraphNodeSequencePlayer : public SGraphNode
+#include "SAnimationGraphNode.h"
+
+class SGraphNodeSequencePlayer : public SAnimationGraphNode
 {
 public:
 	SLATE_BEGIN_ARGS(SGraphNodeSequencePlayer){}
 	SLATE_END_ARGS()
 
-	void Construct(const FArguments& InArgs, UK2Node* InNode);
+	void Construct(const FArguments& InArgs, UAnimGraphNode_Base* InNode);
 
 	// SNodePanel::SNode interface
 	virtual void GetNodeInfoPopups(FNodeInfoContext* Context, TArray<FGraphInformationPopupInfo>& Popups) const override;
diff --git a/Engine/Source/Editor/GraphEditor/Private/NodeFactory.cpp b/Engine/Source/Editor/GraphEditor/Private/NodeFactory.cpp
--- a/Engine/Source/Editor/GraphEditor/Private/NodeFactory.cpp
+++ b/Engine/Source/Editor/GraphEditor/Private/NodeFactory.cpp
@@ -133,6 +133,10 @@ TSharedPtr<SGraphNode> FNodeFactory::CreateNodeWidget(UEdGraphNode* InNode)
 		{
 			return SNew(SGraphNodeLayeredBoneBlend, LayeredBlend);
 		}
+		else
+		{
+			return SNew(SAnimationGraphNode, BaseAnimNode);
+		}
 	}
 
 	if (USoundCueGraphNode_Base* BaseSoundNode = Cast<USoundCueGraphNode_Base>(InNode))
diff --git a/Engine/Source/Editor/HierarchicalLODOutliner/Private/LODActorItem.cpp b/Engine/Source/Editor/HierarchicalLODOutliner/Private/LODActorItem.cpp
--- a/Engine/Source/Editor/HierarchicalLODOutliner/Private/LODActorItem.cpp
+++ b/Engine/Source/Editor/HierarchicalLODOutliner/Private/LODActorItem.cpp
@@ -39,7 +39,7 @@ void HLODOutliner::FLODActorItem::GenerateContextMenu(FMenuBuilder& MenuBuilder,
 	else
 	{		
 		MenuBuilder.AddMenuEntry(LOCTEXT("ForceView", "ForceView"), FText(), FSlateIcon(), FUIAction(FExecuteAction::CreateRaw(&Outliner, &SHLODOutliner::ForceViewLODActor)));
-		MenuBuilder.AddMenuEntry(LOCTEXT("BuildLODActorMesh", "Rebuild Proxy Mesh"), FText(), FSlateIcon(), FUIAction(FExecuteAction::CreateRaw(&Outliner, &SHLODOutliner::RebuildLODActor)));
+		MenuBuilder.AddMenuEntry(LOCTEXT("RebuildLODActorMesh", "Rebuild Proxy Mesh"), FText(), FSlateIcon(), FUIAction(FExecuteAction::CreateRaw(&Outliner, &SHLODOutliner::RebuildLODActor)));
 	}
 
 	MenuBuilder.AddMenuEntry(LOCTEXT("CreateHLODVolume", "Create Containing Hierarchical Volume"), FText(), FSlateIcon(), FUIAction(FExecuteAction::CreateRaw(&Outliner, &SHLODOutliner::CreateHierarchicalVolumeForActor)));
diff --git a/Engine/Source/Editor/Kismet/Private/SSCSEditor.cpp b/Engine/Source/Editor/Kismet/Private/SSCSEditor.cpp
--- a/Engine/Source/Editor/Kismet/Private/SSCSEditor.cpp
+++ b/Engine/Source/Editor/Kismet/Private/SSCSEditor.cpp
@@ -5612,8 +5612,7 @@ void SSCSEditor::OnApplyChangesToBlueprint() const
 			Actor->Modify();
 
 			// Mark components that are either native or from the SCS as modified so they will be restored
-			const TArray<UActorComponent*> Components = Actor->GetComponents();
-			for (UActorComponent* ActorComponent : Components)
+			for (UActorComponent* ActorComponent : Actor->GetComponents())
 			{
 				if (ActorComponent->CreationMethod == EComponentCreationMethod::SimpleConstructionScript || ActorComponent->CreationMethod == EComponentCreationMethod::Native)
 				{
diff --git a/Engine/Source/Editor/KismetCompiler/Private/AnimBlueprintCompiler.cpp b/Engine/Source/Editor/KismetCompiler/Private/AnimBlueprintCompiler.cpp
--- a/Engine/Source/Editor/KismetCompiler/Private/AnimBlueprintCompiler.cpp
+++ b/Engine/Source/Editor/KismetCompiler/Private/AnimBlueprintCompiler.cpp
@@ -571,6 +571,9 @@ void FAnimBlueprintCompiler::ProcessAllAnimationNodes()
 	for (auto SourceNodeIt = AnimNodeList.CreateIterator(); SourceNodeIt; ++SourceNodeIt)
 	{
 		UAnimGraphNode_Base* SourceNode = *SourceNodeIt;
+		UAnimGraphNode_Base* TrueNode = MessageLog.FindSourceObjectTypeChecked<UAnimGraphNode_Base>(SourceNode);
+		TrueNode->BlueprintUsage = EBlueprintUsage::NoProperties;
+
 		if (UAnimGraphNode_Root* PossibleRoot = Cast<UAnimGraphNode_Root>(SourceNode))
 		{
 			if (UAnimGraphNode_Root* Root = ExactCast<UAnimGraphNode_Root>(PossibleRoot))
@@ -1342,6 +1345,29 @@ void FAnimBlueprintCompiler::PostCompile()
 			}
 		}
 	}
+
+	for (const FEffectiveConstantRecord& ConstantRecord : ValidAnimNodePinConstants)
+	{
+		UAnimGraphNode_Base* Node = CastChecked<UAnimGraphNode_Base>(ConstantRecord.LiteralSourcePin->GetOwningNode());
+		UAnimGraphNode_Base* TrueNode = MessageLog.FindSourceObjectTypeChecked<UAnimGraphNode_Base>(Node);
+		TrueNode->BlueprintUsage = EBlueprintUsage::DoesNotUseBlueprint;
+	}
+
+	for(const FEvaluationHandlerRecord& EvaluationHandler : ValidEvaluationHandlerList)
+	{
+		if(EvaluationHandler.ServicedProperties.Num() > 0)
+		{
+			const FAnimNodeSinglePropertyHandler& Handler = EvaluationHandler.ServicedProperties.CreateConstIterator()->Value;
+			UAnimGraphNode_Base* Node = CastChecked<UAnimGraphNode_Base>(Handler.ArrayPins.Num() > 0 ? Handler.ArrayPins[0]->GetOwningNode() : Handler.SinglePin->GetOwningNode());
+			UAnimGraphNode_Base* TrueNode = MessageLog.FindSourceObjectTypeChecked<UAnimGraphNode_Base>(Node);	
+			TrueNode->BlueprintUsage = EvaluationHandler.HandlerFunctionName != NAME_None ? EBlueprintUsage::UsesBlueprint : EBlueprintUsage::DoesNotUseBlueprint;
+
+			if(TrueNode->BlueprintUsage == EBlueprintUsage::UsesBlueprint && DefaultAnimInstance->bWarnAboutBlueprintUsage)
+			{
+				MessageLog.Warning(*LOCTEXT("BlueprintUsageWarning", "Node @@ uses Blueprint to update its values, access member variables directly or use a constant value for better performance.").ToString(), Node);
+			}
+		}
+	}
 }
 
 void FAnimBlueprintCompiler::CreateFunctionList()
diff --git a/Engine/Source/Editor/Localization/Private/LocalizationConfigurationScript.cpp b/Engine/Source/Editor/Localization/Private/LocalizationConfigurationScript.cpp
--- a/Engine/Source/Editor/Localization/Private/LocalizationConfigurationScript.cpp
+++ b/Engine/Source/Editor/Localization/Private/LocalizationConfigurationScript.cpp
@@ -409,7 +409,7 @@ namespace LocalizationConfigurationScript
 			// Do not use culture subdirectories if importing a single culture from a specific directory.
 			if (CultureName.IsSet() && ImportPathOverride.IsSet())
 			{
-				ConfigSection.Add( TEXT("bUseCultureDirectory"), "false" );
+				ConfigSection.Add( TEXT("bUseCultureDirectory"), TEXT("false") );
 			}
 
 			ConfigSection.Add( TEXT("ManifestName"), GetManifestFileName(Target) );
@@ -522,7 +522,7 @@ namespace LocalizationConfigurationScript
 			// Do not use culture subdirectories if exporting a single culture to a specific directory.
 			if (CultureName.IsSet() && ExportPathOverride.IsSet())
 			{
-				ConfigSection.Add( TEXT("bUseCultureDirectory"), "false" );
+				ConfigSection.Add( TEXT("bUseCultureDirectory"), TEXT("false") );
 			}
 
 			ConfigSection.Add( TEXT("ManifestName"), GetManifestFileName(Target) );
@@ -640,7 +640,7 @@ namespace LocalizationConfigurationScript
 			// Do not use culture subdirectories if importing a single culture from a specific directory.
 			if (CultureName.IsSet() && ImportPathOverride.IsSet())
 			{
-				ConfigSection.Add(TEXT("bUseCultureDirectory"), "false");
+				ConfigSection.Add(TEXT("bUseCultureDirectory"), TEXT("false"));
 			}
 
 			ConfigSection.Add(TEXT("ManifestName"), GetManifestFileName(Target));
@@ -751,7 +751,7 @@ namespace LocalizationConfigurationScript
 			// Do not use culture subdirectories if exporting a single culture to a specific directory.
 			if (CultureName.IsSet() && ExportPathOverride.IsSet())
 			{
-				ConfigSection.Add(TEXT("bUseCultureDirectory"), "false");
+				ConfigSection.Add(TEXT("bUseCultureDirectory"), TEXT("false"));
 			}
 
 			ConfigSection.Add(TEXT("ManifestName"), GetManifestFileName(Target));
diff --git a/Engine/Source/Editor/MovieSceneTools/Private/MatineeImportTools.cpp b/Engine/Source/Editor/MovieSceneTools/Private/MatineeImportTools.cpp
--- a/Engine/Source/Editor/MovieSceneTools/Private/MatineeImportTools.cpp
+++ b/Engine/Source/Editor/MovieSceneTools/Private/MatineeImportTools.cpp
@@ -81,7 +81,7 @@ void FMatineeImportTools::SetOrAddKey( FRichCurve& Curve, float Time, float Valu
 
 void FMatineeImportTools::CopyInterpFloatTrack( TSharedRef<ISequencer> Sequencer, UInterpTrackFloatBase* MatineeFloatTrack, UMovieSceneFloatTrack* FloatTrack )
 {
-	const FScopedTransaction Transaction( NSLOCTEXT( "Sequencer", "PasteMatineeFloatTrack", "Paste Matinee float track" ) );
+	const FScopedTransaction Transaction( NSLOCTEXT( "Sequencer", "PasteMatineeFloatTrack", "Paste Matinee Float Track" ) );
 	bool bSectionCreated = false;
 
 	FloatTrack->Modify();
diff --git a/Engine/Source/Editor/MovieSceneTools/Private/Sections/CameraCutSection.cpp b/Engine/Source/Editor/MovieSceneTools/Private/Sections/CameraCutSection.cpp
--- a/Engine/Source/Editor/MovieSceneTools/Private/Sections/CameraCutSection.cpp
+++ b/Engine/Source/Editor/MovieSceneTools/Private/Sections/CameraCutSection.cpp
@@ -107,7 +107,7 @@ void FCameraCutSection::BuildSectionContextMenu(FMenuBuilder& MenuBuilder, const
 
 FText FCameraCutSection::GetDisplayName() const
 {
-	return NSLOCTEXT("FCameraCutSection", "", "CameraCuts");
+	return NSLOCTEXT("FCameraCutSection", "CameraCuts", "CameraCuts");
 }
 
 
diff --git a/Engine/Source/Editor/MovieSceneTools/Private/Sections/CinematicShotSection.cpp b/Engine/Source/Editor/MovieSceneTools/Private/Sections/CinematicShotSection.cpp
--- a/Engine/Source/Editor/MovieSceneTools/Private/Sections/CinematicShotSection.cpp
+++ b/Engine/Source/Editor/MovieSceneTools/Private/Sections/CinematicShotSection.cpp
@@ -210,7 +210,7 @@ void FCinematicShotSection::AddTakesMenu(FMenuBuilder& MenuBuilder)
 
 FText FCinematicShotSection::GetDisplayName() const
 {
-	return NSLOCTEXT("FCinematicShotSection", "", "Shot");
+	return NSLOCTEXT("FCinematicShotSection", "Shot", "Shot");
 }
 
 /* FCinematicShotSection callbacks
diff --git a/Engine/Source/Editor/Persona/Private/AnimationEditorViewportClient.cpp b/Engine/Source/Editor/Persona/Private/AnimationEditorViewportClient.cpp
--- a/Engine/Source/Editor/Persona/Private/AnimationEditorViewportClient.cpp
+++ b/Engine/Source/Editor/Persona/Private/AnimationEditorViewportClient.cpp
@@ -134,7 +134,7 @@ FAnimationViewportClient::FAnimationViewportClient(FAnimationEditorPreviewScene&
 	bDrawUVs = false;
 	UVChannelToDraw = 0;
 
-	bAutoAlignFloor = true;
+	bAutoAlignFloor = ConfigOption->bAutoAlignFloorToMesh;
 
 	// Set audio mute option
 	UWorld* World = PreviewScene->GetWorld();
@@ -260,6 +260,8 @@ void FAnimationViewportClient::OnToggleAutoAlignFloor()
 {
 	bAutoAlignFloor = !bAutoAlignFloor;
 	UpdateCameraSetup();
+
+	ConfigOption->SetAutoAlignFloorToMesh(bAutoAlignFloor);
 }
 
 bool FAnimationViewportClient::IsAutoAlignFloor() const
diff --git a/Engine/Source/Editor/Persona/Private/Persona.cpp b/Engine/Source/Editor/Persona/Private/Persona.cpp
--- a/Engine/Source/Editor/Persona/Private/Persona.cpp
+++ b/Engine/Source/Editor/Persona/Private/Persona.cpp
@@ -2619,7 +2619,7 @@ FText FPersona::GetRecordMenuLabel() const
 		return LOCTEXT("Persona_StopRecordAnimationMenuLabel", "Stop Record Animation");
 	}
 
-	return LOCTEXT("Persona_StartRecordAnimationLabel", "Start Record Animation");
+	return LOCTEXT("Persona_StartRecordAnimationMenuLabel", "Start Record Animation");
 }
 
 FText FPersona::GetRecordStatusLabel() const
@@ -2974,7 +2974,7 @@ void FPersona::RemoveUnusedBones()
 		
 		if(FMessageDialog::Open( EAppMsgType::YesNo, TimeTakenMessage ) == EAppReturnType::Yes)
 		{
-			const FText StatusUpdate = FText::Format(LOCTEXT("RemoveUnusedBones_ProcessingAssets", "Processing Skeletal Meshes for {0}"), FText::FromString(TargetSkeleton->GetName()) );
+			const FText StatusUpdate = FText::Format(LOCTEXT("RemoveUnusedBones_ProcessingAssetsFor", "Processing Skeletal Meshes for {0}"), FText::FromString(TargetSkeleton->GetName()) );
 			GWarn->BeginSlowTask(StatusUpdate, true );
 
 			// Loop through all SkeletalMeshes and remove the bones they use from our list
diff --git a/Engine/Source/Editor/Persona/Private/SAnimMontagePanel.cpp b/Engine/Source/Editor/Persona/Private/SAnimMontagePanel.cpp
--- a/Engine/Source/Editor/Persona/Private/SAnimMontagePanel.cpp
+++ b/Engine/Source/Editor/Persona/Private/SAnimMontagePanel.cpp
@@ -339,9 +339,9 @@ void SAnimMontagePanel::SummonTrackContextMenu( FMenuBuilder& MenuBuilder, float
 
 void SAnimMontagePanel::FillElementSubMenuForTimes(FMenuBuilder& MenuBuilder)
 {
-	MenuBuilder.AddMenuEntry(LOCTEXT("SubLinkAbs", "Absolute"), LOCTEXT("SubLinkAbs", "Set all elements to absolute link"), FSlateIcon(), FUIAction(FExecuteAction::CreateSP(this, &SAnimMontagePanel::OnSetElementsToLinkMode, EAnimLinkMethod::Absolute)));
-	MenuBuilder.AddMenuEntry(LOCTEXT("SubLinkRel", "Relative"), LOCTEXT("SubLinkRel", "Set all elements to relative link"), FSlateIcon(), FUIAction(FExecuteAction::CreateSP(this, &SAnimMontagePanel::OnSetElementsToLinkMode, EAnimLinkMethod::Relative)));
-	MenuBuilder.AddMenuEntry(LOCTEXT("SubLinkPro", "Proportional"), LOCTEXT("SubLinkPro", "Set all elements to proportional link"), FSlateIcon(), FUIAction(FExecuteAction::CreateSP(this, &SAnimMontagePanel::OnSetElementsToLinkMode, EAnimLinkMethod::Proportional)));
+	MenuBuilder.AddMenuEntry(LOCTEXT("SubLinkAbs", "Absolute"), LOCTEXT("SubLinkAbsTooltip", "Set all elements to absolute link"), FSlateIcon(), FUIAction(FExecuteAction::CreateSP(this, &SAnimMontagePanel::OnSetElementsToLinkMode, EAnimLinkMethod::Absolute)));
+	MenuBuilder.AddMenuEntry(LOCTEXT("SubLinkRel", "Relative"), LOCTEXT("SubLinkRelTooltip", "Set all elements to relative link"), FSlateIcon(), FUIAction(FExecuteAction::CreateSP(this, &SAnimMontagePanel::OnSetElementsToLinkMode, EAnimLinkMethod::Relative)));
+	MenuBuilder.AddMenuEntry(LOCTEXT("SubLinkPro", "Proportional"), LOCTEXT("SubLinkProTooltip", "Set all elements to proportional link"), FSlateIcon(), FUIAction(FExecuteAction::CreateSP(this, &SAnimMontagePanel::OnSetElementsToLinkMode, EAnimLinkMethod::Proportional)));
 }
 
 void SAnimMontagePanel::FillSlotSubMenu(FMenuBuilder& Menubuilder)
diff --git a/Engine/Source/Editor/Persona/Private/SAnimNotifyPanel.cpp b/Engine/Source/Editor/Persona/Private/SAnimNotifyPanel.cpp
--- a/Engine/Source/Editor/Persona/Private/SAnimNotifyPanel.cpp
+++ b/Engine/Source/Editor/Persona/Private/SAnimNotifyPanel.cpp
@@ -2881,12 +2881,12 @@ TSharedPtr<SWidget> SAnimNotifyTrack::SummonContextMenu(const FGeometry& MyGeome
 					NewAction.ExecuteAction.BindRaw(
 						this, &SAnimNotifyTrack::OnPasteNotifyClicked, ENotifyPasteMode::MousePosition, ENotifyPasteMultipleMode::Relative);
 
-					MenuBuilder.AddMenuEntry(LOCTEXT("PasteMultRel", "Paste Multiple Relative"), LOCTEXT("PasteToolTip", "Paste multiple notifies beginning at the mouse cursor, maintaining the same relative spacing as the source."), FSlateIcon(), NewAction);
+					MenuBuilder.AddMenuEntry(LOCTEXT("PasteMultRel", "Paste Multiple Relative"), LOCTEXT("PasteMultRelToolTip", "Paste multiple notifies beginning at the mouse cursor, maintaining the same relative spacing as the source."), FSlateIcon(), NewAction);
 
 					NewAction.ExecuteAction.BindRaw(
 						this, &SAnimNotifyTrack::OnPasteNotifyClicked, ENotifyPasteMode::MousePosition, ENotifyPasteMultipleMode::Absolute);
 
-					MenuBuilder.AddMenuEntry(LOCTEXT("PasteMultAbs", "Paste Multiple Absolute"), LOCTEXT("PasteToolTip", "Paste multiple notifies beginning at the mouse cursor, maintaining absolute spacing."), FSlateIcon(), NewAction);
+					MenuBuilder.AddMenuEntry(LOCTEXT("PasteMultAbs", "Paste Multiple Absolute"), LOCTEXT("PasteMultAbsToolTip", "Paste multiple notifies beginning at the mouse cursor, maintaining absolute spacing."), FSlateIcon(), NewAction);
 				}
 
 				if(OriginalTime < Sequence->SequenceLength)
@@ -2916,7 +2916,7 @@ TSharedPtr<SWidget> SAnimNotifyTrack::SummonContextMenu(const FGeometry& MyGeome
 
 				NewAction.ExecuteAction.BindRaw(
 					this, &SAnimNotifyTrack::OnOpenNotifySource, Blueprint);
-				MenuBuilder.AddMenuEntry(LOCTEXT("OpenNotifyBlueprint", "Open Notify Blueprint"), LOCTEXT("OpenNotifyBlueprint", "Opens the source blueprint for this notify"), FSlateIcon(), NewAction);
+				MenuBuilder.AddMenuEntry(LOCTEXT("OpenNotifyBlueprint", "Open Notify Blueprint"), LOCTEXT("OpenNotifyBlueprintTooltip", "Opens the source blueprint for this notify"), FSlateIcon(), NewAction);
 
 				MenuBuilder.EndSection(); //ViewSource
 			}
@@ -4380,7 +4380,7 @@ void SAnimNotifyPanel::OnReplaceSelectedWithNotify(FString NewNotifyName, UClass
 	// Sort these since order is important for deletion
 	SelectedNodes.Sort();
 
-	const FScopedTransaction Transaction(LOCTEXT("AddNotifyEvent", "Replace Anim Notify"));
+	const FScopedTransaction Transaction(LOCTEXT("ReplaceAnimNotify", "Replace Anim Notify"));
 	Sequence->Modify(true);
 
 	for (INodeObjectInterface* NodeObject : SelectedNodes)
diff --git a/Engine/Source/Editor/Persona/Private/SAnimTimingPanel.cpp b/Engine/Source/Editor/Persona/Private/SAnimTimingPanel.cpp
--- a/Engine/Source/Editor/Persona/Private/SAnimTimingPanel.cpp
+++ b/Engine/Source/Editor/Persona/Private/SAnimTimingPanel.cpp
@@ -459,7 +459,7 @@ void FTimingRelevantElement_Section::GetDescriptionItems(TMap<FString, FText>& I
 	NumberOptions.MinimumFractionalDigits = 3;
 
 	Items.Add(LOCTEXT("SectionName", "Name").ToString(), FText::FromName(Section.SectionName));
-	Items.Add(LOCTEXT("SectionTriggerTime", "Trigger Time").ToString(), FText::Format(LOCTEXT("SectionTriggerTime", "{0}s"), FText::AsNumber(Section.GetTime(), &NumberOptions)));
+	Items.Add(LOCTEXT("SectionTriggerTime", "Trigger Time").ToString(), FText::Format(LOCTEXT("SectionTriggerTimeValue", "{0}s"), FText::AsNumber(Section.GetTime(), &NumberOptions)));
 }
 
 FName FTimingRelevantElement_Notify::GetTypeName()
diff --git a/Engine/Source/Editor/Persona/Private/SAnimationSequenceBrowser.cpp b/Engine/Source/Editor/Persona/Private/SAnimationSequenceBrowser.cpp
--- a/Engine/Source/Editor/Persona/Private/SAnimationSequenceBrowser.cpp
+++ b/Engine/Source/Editor/Persona/Private/SAnimationSequenceBrowser.cpp
@@ -945,11 +945,12 @@ bool SAnimationSequenceBrowser::OnVisualizeAssetToolTip(const TSharedPtr<SWidget
 			PreviewComponent->EnablePreview(true, Asset, nullptr);
 			PreviewComponent->PreviewInstance->PlayAnim(true);
 
+			FBoxSphereBounds MeshImportedBounds = MeshToUse->GetImportedBounds();
 			float HalfFov = FMath::DegreesToRadians(ViewportClient->ViewFOV) / 2.0f;
-			float TargetDist = MeshToUse->Bounds.SphereRadius / FMath::Tan(HalfFov);
+			float TargetDist = MeshImportedBounds.SphereRadius / FMath::Tan(HalfFov);
 
 			ViewportClient->SetViewRotation(FRotator(0.0f, -45.0f, 0.0f));
-			ViewportClient->SetViewLocationForOrbiting(FVector(0.0f, 0.0f, MeshToUse->Bounds.BoxExtent.Z / 2.0f), TargetDist);
+			ViewportClient->SetViewLocationForOrbiting(FVector(0.0f, 0.0f, MeshImportedBounds.BoxExtent.Z / 2.0f), TargetDist);
 
 			ViewportWidget->SetVisibility(EVisibility::Visible);
 			
diff --git a/Engine/Source/Editor/PhAT/Private/PhAT.cpp b/Engine/Source/Editor/PhAT/Private/PhAT.cpp
--- a/Engine/Source/Editor/PhAT/Private/PhAT.cpp
+++ b/Engine/Source/Editor/PhAT/Private/PhAT.cpp
@@ -1,5 +1,6 @@
 // Copyright 1998-2016 Epic Games, Inc. All Rights Reserved.
 
+#include "PhATPrivatePCH.h"
 #include "PhATModule.h"
 #include "AssetSelection.h"
 #include "ScopedTransaction.h"
@@ -34,6 +35,8 @@
 #include "PhysicalMaterials/PhysicalMaterial.h"
 #include "PersonaModule.h"
 
+const FName PhATAppIdentifier = FName(TEXT("PhATApp"));
+
 DEFINE_LOG_CATEGORY(LogPhAT);
 #define LOCTEXT_NAMESPACE "PhAT"
 
@@ -229,44 +232,43 @@ void FPhAT::InitPhAT(const EToolkitMode::Type Mode, const TSharedPtr< class IToo
 
 	CreateInternalWidgets();
 
-	const TSharedRef<FTabManager::FLayout> StandaloneDefaultLayout = FTabManager::NewLayout("Standalone_PhAT_Layout_v2")
+	const TSharedRef<FTabManager::FLayout> StandaloneDefaultLayout = FTabManager::NewLayout("Standalone_PhAT_Layout_v3")
 	->AddArea
 	(
 		FTabManager::NewPrimaryArea()
-		->SetOrientation(Orient_Vertical)
+		->SetOrientation(Orient_Horizontal)
 		->Split
 		(
 			FTabManager::NewStack()
-			->SetSizeCoefficient(0.1f)
-			->SetHideTabWell(true)
-			->AddTab(GetToolbarTabId(), ETabState::OpenedTab)
+			->SetSizeCoefficient(0.2f)
+			->AddTab(PhATHierarchyName, ETabState::OpenedTab)
 		)
 		->Split
 		(
 			FTabManager::NewSplitter()
-			->SetSizeCoefficient(0.9f)
-			->SetOrientation(Orient_Horizontal)
+			->SetOrientation(Orient_Vertical)
+			->SetSizeCoefficient(0.6f)
 			->Split
 			(
-				FTabManager::NewStack() 
-				->SetSizeCoefficient(0.8f)
-				->AddTab(PhATPreviewViewportName, ETabState::OpenedTab)
+				FTabManager::NewStack()
+				->SetSizeCoefficient(0.1f)
+				->SetHideTabWell(true)
+				->AddTab(GetToolbarTabId(), ETabState::OpenedTab)
 			)
 			->Split
 			(
-				FTabManager::NewSplitter()
-				->SetSizeCoefficient(0.2f)
-				->Split
-				(
-					FTabManager::NewStack() ->AddTab(PhATPropertiesName, ETabState::OpenedTab)
-				)
-				->Split
-				(
-					FTabManager::NewStack() ->AddTab(PhATHierarchyName, ETabState::OpenedTab)
-				)
-
+				FTabManager::NewStack()
+				->SetSizeCoefficient(0.9f)
+				->SetHideTabWell(true)
+				->AddTab(PhATPreviewViewportName, ETabState::OpenedTab)
 			)
 		)
+		->Split
+		(
+			FTabManager::NewStack()
+			->SetSizeCoefficient(0.2f)
+			->AddTab(PhATPropertiesName, ETabState::OpenedTab)
+		)
 	);
 
 	const bool bCreateDefaultStandaloneMenu = true;
@@ -636,7 +638,7 @@ void FPhAT::CreateInternalWidgets()
 	FPropertyEditorModule& PropertyModule = FModuleManager::LoadModuleChecked<FPropertyEditorModule>("PropertyEditor");
 	Properties = PropertyModule.CreateDetailView( Args );
 	Properties->SetObject(SharedData->EditorSimOptions);
-
+	Properties->OnFinishedChangingProperties().AddLambda([this](const FPropertyChangedEvent& Event) { RefreshPreviewViewport(); });
 
 	HierarchyControl = 
 	SNew(SBorder)
@@ -3231,7 +3233,7 @@ FText FPhAT::GetRecordMenuLabel() const
 		return LOCTEXT("Persona_StopRecordAnimationMenuLabel", "Stop Record Animation");
 	}
 
-	return LOCTEXT("Persona_StartRecordAnimationLabel", "Start Record Animation");
+	return LOCTEXT("Persona_StartRecordAnimationMenuLabel", "Start Record Animation");
 }
 
 FText FPhAT::GetRecordStatusLabel() const
diff --git a/Engine/Source/Editor/PhAT/Private/PhATActions.cpp b/Engine/Source/Editor/PhAT/Private/PhATActions.cpp
--- a/Engine/Source/Editor/PhAT/Private/PhATActions.cpp
+++ b/Engine/Source/Editor/PhAT/Private/PhATActions.cpp
@@ -1,6 +1,6 @@
 // Copyright 1998-2016 Epic Games, Inc. All Rights Reserved.
 
-#include "PhATModule.h"
+#include "PhATPrivatePCH.h"
 #include "PhATActions.h"
 #include "PhAT.h"
 
@@ -51,8 +51,8 @@ void FPhATCommands::RegisterCommands()
 	UI_COMMAND(ConvertToPrismatic, "To Prismatic", "Convert Selected Constraint To Prismatic", EUserInterfaceActionType::Button, FInputChord());
 	UI_COMMAND(ConvertToSkeletal, "To Skeletal", "Convert Selected Constraint To Skeletal", EUserInterfaceActionType::Button, FInputChord());
 	UI_COMMAND(DeleteConstraint, "Delete", "Delete Selected Constraint", EUserInterfaceActionType::Button, FInputChord());
-	UI_COMMAND(PlayAnimation, "Play", "Play Animation", EUserInterfaceActionType::ToggleButton, FInputChord());
-	UI_COMMAND(RecordAnimation, "Record", "Record Animation", EUserInterfaceActionType::ToggleButton, FInputChord());
+	UI_COMMAND(PlayAnimation, "PlayAnimation", "Play Animation", EUserInterfaceActionType::ToggleButton, FInputChord());
+	UI_COMMAND(RecordAnimation, "RecordAnimation", "Record Animation", EUserInterfaceActionType::ToggleButton, FInputChord());
 	UI_COMMAND(ShowSkeleton, "Skeleton", "Show Skeleton", EUserInterfaceActionType::ToggleButton, FInputChord());
 	UI_COMMAND(MakeBodyKinematic, "Kinematic", "Make Body Kinematic", EUserInterfaceActionType::RadioButton, FInputChord());
 	UI_COMMAND(MakeBodySimulated, "Simulated", "Make Body Simulated", EUserInterfaceActionType::RadioButton, FInputChord());
diff --git a/Engine/Source/Editor/PhAT/Private/PhATHitProxies.cpp b/Engine/Source/Editor/PhAT/Private/PhATHitProxies.cpp
--- a/Engine/Source/Editor/PhAT/Private/PhATHitProxies.cpp
+++ b/Engine/Source/Editor/PhAT/Private/PhATHitProxies.cpp
@@ -1,6 +1,6 @@
 // Copyright 1998-2016 Epic Games, Inc. All Rights Reserved.
 
-#include "PhATModule.h"
+#include "PhATPrivatePCH.h"
 #include "PhATHitProxies.h"
 
 IMPLEMENT_HIT_PROXY(HPhATEdBoneProxy,HHitProxy);
diff --git a/Engine/Source/Editor/PhAT/Private/PhATModule.cpp b/Engine/Source/Editor/PhAT/Private/PhATModule.cpp
--- a/Engine/Source/Editor/PhAT/Private/PhATModule.cpp
+++ b/Engine/Source/Editor/PhAT/Private/PhATModule.cpp
@@ -1,14 +1,11 @@
 // Copyright 1998-2016 Epic Games, Inc. All Rights Reserved.
 
-
+#include "PhATPrivatePCH.h"
 #include "PhATModule.h"
 #include "ModuleManager.h"
 #include "PhAT.h"
 #include "PhATSharedData.h"
 
-const FName PhATAppIdentifier = FName(TEXT("PhATApp"));
-
-
 /*-----------------------------------------------------------------------------
    FPhATModule
 -----------------------------------------------------------------------------*/
