Copying //UE4/Dev-Framework to //UE4/Dev-Main (Source: //UE4/Dev-Framework @ 3021292)

==========================
MAJOR FEATURES + CHANGES
==========================

Change 3016321 on 2016/06/16 by [REDACTED_PERSON]

	Refactor constraint instance so that we can easily swap parameters at runtime using profiles.

Change 3016367 on 2016/06/16 by [REDACTED_PERSON]

	Strip #lockdown from changelists as well

Change 3016380 on 2016/06/16 by [REDACTED_PERSON]

	Fixes for socket copying in Persona (mode to preserve bone attachment & fix for duplicating same socket multiple times)

Change 3016396 on 2016/06/16 by [REDACTED_PERSON]

	#ue4 - Fix bad call to RemoveAt() in FFinalPostProcessSettings::UpdateEntry() causing potential infinite loop.

	Github #2504

	Merging using //UE4/Dev-Framework_to_//UE4/Release-4.12 (reversed)

Change 3016477 on 2016/06/16 by [REDACTED_PERSON]

	Fix initialization order compile warning
	Actually initialize LastKnownScale

Change 3017624 on 2016/06/17 by [REDACTED_PERSON]

	- Remove DestructibleActor.h from Engine.h
	- Remove SkeletalMesh.h from EnginePrivate.h

Change 3017657 on 2016/06/17 by [REDACTED_PERSON]

	PR #2253: Fix UDestructibleComponent::AddImpulseAtLocation (Contributed by [REDACTED_PERSON])

Change 3017658 on 2016/06/17 by [REDACTED_PERSON]

	PR #1840 fixed generating patch with DestructibleMesh assets (Contributed by [REDACTED_PERSON])

Change 3017661 on 2016/06/17 by [REDACTED_PERSON]

	Dynamic space switching for anim dynamics, also exposed a bunch of properties as optional pins.

Change 3017914 on 2016/06/17 by [REDACTED_PERSON]

	- Add WITH_EDITORONLY_DATA around ProgressiveDrawingFraction
	- Add number of morph targets to SkelMesh asset registry tags

Change 3017915 on 2016/06/17 by [REDACTED_PERSON]

	Move util to find debugged node from FAnimationViewportClient to UAnimGraphNode_SkeletalControlBase

Change 3017916 on 2016/06/17 by [REDACTED_PERSON]

	Remove unused FMorphTargetMap struct

Change 3018038 on 2016/06/17 by [REDACTED_PERSON]

	UE-32124 Crash when Deleting Concatenator Node while Playing Sound Cue

Change 3018048 on 2016/06/17 by [REDACTED_PERSON]

	Made GameplayDebugger's HUD toggling persist over tool toggling #UE4

Change 3018068 on 2016/06/17 by [REDACTED_PERSON]

	Improved navigaiton-related condition in USceneComponent::PostUpdateNavigationData to avoid doing unnecessary work #Orion

	Thich change also fixes USceneComponent not working in client-side dynamic navmesh generation

Change 3018529 on 2016/06/17 by [REDACTED_PERSON]

	Passing in sendlist on CreateSource for xaudio2 in effort to avoid lag on AMD with CreateSourceVoice

	As per this thread:
	https://udn.unrealengine.com/questions/298497/xaudio2-setoutputvoices-lag.html

Change 3018907 on 2016/06/18 by [REDACTED_PERSON]

	UE-32004 Add EngineDefines.h to CameraTypes.h for WORLD_MAX

Change 3019372 on 2016/06/20 by [REDACTED_PERSON]

	- Initial add of OrientationDriver node
	  - UI for importing poses from PoseAset not yet working
	- Add FQuat::AngularDistance

Change 3019373 on 2016/06/20 by [REDACTED_PERSON]

	PSD test assets

Change 3019444 on 2016/06/20 by Jurre.deBaare

	Added advanced preview scene, this allows specific profiles to be set altering the Persona and Static mesh editor scene

Change 3019565 on 2016/06/20 by [REDACTED_PERSON]

	#anim : prioritize socket name over bone name

Change 3019712 on 2016/06/20 by Jurre.deBaare

	Changed the way the advanced preview tab is spawned for Persona (otherwise would cause a crash in anim graph editor)

Change 3020010 on 2016/06/20 by [REDACTED_PERSON]

	Implementing CL 3019996 in Dev-Framework

	Proper default LPF values to avoid situation when going to area with no reverb volume

Change 3020067 on 2016/06/20 by [REDACTED_PERSON]

	Add constraint profiles to physics assets.

Change 3020110 on 2016/06/20 by [REDACTED_PERSON]

	Github PR: mesh merge crash fix

Change 3020167 on 2016/06/20 by [REDACTED_PERSON]

	Added the ability to add forces and impulses to all bodies in a skeletal mesh component. Refactored iteration of bodies below in skeletal mesh component.

Change 3020324 on 2016/06/20 by [REDACTED_PERSON]

	Bring back cone constraint visualization, but with some improvements like colors and free vs locked.
	Turn on constraint visualization when drawing a skeletal mesh with constraint show flags.

Change 3020342 on 2016/06/20 by [REDACTED_PERSON]

	Fix constraint crash when trying to wake up an actor that hasn't been added into the scene yet.

Change 3020349 on 2016/06/20 by mason.seay

	Resaving material assets

Change 3020400 on 2016/06/20 by [REDACTED_PERSON]

	Disable check single influence on LOD check

	- this doesn't seem to appropriate anymore

	#code review:Alexis.Matte

Change 3020928 on 2016/06/21 by [REDACTED_PERSON]

	- Move PSD test assets out of Developer folder
	- Add PSD test AnimBP

Change 3021000 on 2016/06/21 by [REDACTED_PERSON]

	Remove ensure from body instance customization. This is a legit case where owner becomes stale (bp recompile for example)

Change 3021017 on 2016/06/21 by [REDACTED_PERSON]

	Help static analysis

Change 3021045 on 2016/06/21 by [REDACTED_PERSON]

	Fix regression with detachment in AActor::OnRep_AttachmentReplication

diff --git a/Engine/Config/BaseEditor.ini b/Engine/Config/BaseEditor.ini
--- a/Engine/Config/BaseEditor.ini
+++ b/Engine/Config/BaseEditor.ini
@@ -11,7 +11,7 @@ ShouldLoadLocalizedPropertyNames=True
 EnableStylusPressure=True
 
 [ImportSetting]
-CheckSingleInfluenceLOD=True
+CheckSingleInfluenceLOD=False
 
 [LogWindow]
 MaxNumberOfLogLines=2000
diff --git a/Engine/Config/BaseEngine.ini b/Engine/Config/BaseEngine.ini
--- a/Engine/Config/BaseEngine.ini
+++ b/Engine/Config/BaseEngine.ini
@@ -73,7 +73,6 @@ VertexColorViewModeMaterialName_RedOnly=/Engine/EngineDebugMaterials/VertexColor
 VertexColorViewModeMaterialName_GreenOnly=/Engine/EngineDebugMaterials/VertexColorViewMode_GreenOnly.VertexColorViewMode_GreenOnly
 VertexColorViewModeMaterialName_BlueOnly=/Engine/EngineDebugMaterials/VertexColorViewMode_BlueOnly.VertexColorViewMode_BlueOnly
 BoneWeightMaterialName=/Engine/EngineDebugMaterials/BoneWeightMaterial.BoneWeightMaterial
-ConstraintLimitMaterialName=/Engine/EditorMaterials/PhAT_JointLimitMaterial.PhAT_JointLimitMaterial
 InvalidLightmapSettingsMaterialName=/Engine/EngineMaterials/M_InvalidLightmapSettings.M_InvalidLightmapSettings
 PreviewShadowsIndicatorMaterialName=/Engine/EditorMaterials/PreviewShadowIndicatorMaterial.PreviewShadowIndicatorMaterial 
 EditorBrushMaterialName=/Engine/EngineMaterials/EditorBrushMaterial.EditorBrushMaterial
diff --git a/Engine/Source/Developer/GameplayDebugger/Private/GameplayDebuggerExtension_HUD.cpp b/Engine/Source/Developer/GameplayDebugger/Private/GameplayDebuggerExtension_HUD.cpp
--- a/Engine/Source/Developer/GameplayDebugger/Private/GameplayDebuggerExtension_HUD.cpp
+++ b/Engine/Source/Developer/GameplayDebugger/Private/GameplayDebuggerExtension_HUD.cpp
@@ -8,6 +8,7 @@
 
 FGameplayDebuggerExtension_HUD::FGameplayDebuggerExtension_HUD()
 {
+	bWantsHUDEnabled = false;
 	bIsGameHUDEnabled = false;
 	bAreDebugMessagesEnabled = false;
 	bPrevDebugMessagesEnabled = GEngine && GEngine->bEnableOnScreenDebugMessages;
@@ -33,7 +34,7 @@ TSharedRef<FGameplayDebuggerExtension> FGameplayDebuggerExtension_HUD::MakeInsta
 
 void FGameplayDebuggerExtension_HUD::OnActivated()
 {
-	SetGameHUDEnabled(false);
+	SetGameHUDEnabled(bWantsHUDEnabled);
 	SetDebugMessagesEnabled(false);
 }
 
@@ -103,7 +104,8 @@ void FGameplayDebuggerExtension_HUD::SetDebugMessagesEnabled(bool bEnable)
 
 void FGameplayDebuggerExtension_HUD::ToggleGameHUD()
 {
-	SetGameHUDEnabled(!bIsGameHUDEnabled);
+	const bool bEnable = bWantsHUDEnabled = !bIsGameHUDEnabled;
+	SetGameHUDEnabled(bEnable);
 }
 
 void FGameplayDebuggerExtension_HUD::ToggleDebugMessages()
diff --git a/Engine/Source/Editor/AnimGraph/AnimGraph.Build.cs b/Engine/Source/Editor/AnimGraph/AnimGraph.Build.cs
--- a/Engine/Source/Editor/AnimGraph/AnimGraph.Build.cs
+++ b/Engine/Source/Editor/AnimGraph/AnimGraph.Build.cs
@@ -27,7 +27,8 @@ public class AnimGraph : ModuleRules
 				"UnrealEd",
                 "GraphEditor",
 				"PropertyEditor",
-				"EditorStyle"
+				"EditorStyle",
+                "ContentBrowser"
 			}
 		);
 
diff --git a/Engine/Source/Editor/AnimGraph/Classes/AnimGraphNode_OrientationDriver.h b/Engine/Source/Editor/AnimGraph/Classes/AnimGraphNode_OrientationDriver.h
new file mode 100644
--- /dev/null
+++ b/Engine/Source/Editor/AnimGraph/Classes/AnimGraphNode_OrientationDriver.h
@@ -0,0 +1,54 @@
+// Copyright 1998-2016 Epic Games, Inc. All Rights Reserved.
+
+#pragma once
+
+class IDetailCategoryBuilder;
+class IPropertyHandle;
+
+#include "AnimGraphNode_SkeletalControlBase.h"
+#include "BoneControllers/AnimNode_OrientationDriver.h"
+#include "EdGraph/EdGraphNodeUtils.h" // for FNodeTitleTextTable
+#include "AnimGraphNode_OrientationDriver.generated.h"
+
+UCLASS()
+class ANIMGRAPH_API UAnimGraphNode_OrientationDriver : public UAnimGraphNode_SkeletalControlBase
+{
+	GENERATED_UCLASS_BODY()
+
+	UPROPERTY(EditAnywhere, Category = Settings)
+	FAnimNode_OrientationDriver Node;
+
+public:
+
+	// UEdGraphNode interface
+	virtual FText GetNodeTitle(ENodeTitleType::Type TitleType) const override;
+	virtual FText GetTooltipText() const override;
+	// End of UEdGraphNode interface
+
+	// UAnimGraphNode_Base interface
+	virtual void ValidateAnimNodeDuringCompilation(USkeleton* ForSkeleton, FCompilerResultsLog& MessageLog) override;
+	virtual void CustomizeDetails(IDetailLayoutBuilder& DetailBuilder) override;
+	// End of UAnimGraphNode_Base interface
+
+	// UAnimGraphNode_SkeletalControlBase interface
+	virtual void Draw(FPrimitiveDrawInterface* PDI, USkeletalMeshComponent* SkelMeshComp) const override;
+	virtual void DrawCanvas(FViewport& InViewport, FSceneView& View, FCanvas& Canvas, USkeletalMeshComponent * PreviewSkelMeshComp) const override;
+	// End of UAnimGraphNode_SkeletalControlBase interface
+
+protected:
+
+	// UAnimGraphNode_SkeletalControlBase protected interface
+	virtual FText GetControllerDescription() const override;
+	virtual const FAnimNode_SkeletalControlBase* GetNode() const override { return &Node; }
+	// End of UAnimGraphNode_SkeletalControlBase protected interface
+
+private:
+	/** Called when a pose is selected from asset picker */
+	void OnPoseAssetSelected(const FAssetData& AssetData);
+
+	/** Build pop-up asset picker for choosing PoseAsset to import info from */
+	TSharedRef<SWidget> BuildPoseAssetPicker();
+
+	/** Handle 'import poses from pose asset' button being pressed */
+	FReply ClickedOnImportPosesFromPoseAsset();
+};
diff --git a/Engine/Source/Editor/AnimGraph/Classes/AnimGraphNode_SkeletalControlBase.h b/Engine/Source/Editor/AnimGraph/Classes/AnimGraphNode_SkeletalControlBase.h
--- a/Engine/Source/Editor/AnimGraph/Classes/AnimGraphNode_SkeletalControlBase.h
+++ b/Engine/Source/Editor/AnimGraph/Classes/AnimGraphNode_SkeletalControlBase.h
@@ -80,6 +80,9 @@ public:
 	// called when the widget is dragged in scale mode
 	virtual void DoScale(const USkeletalMeshComponent* SkelComp, FVector& Scale, FAnimNode_Base* InOutAnimNode) {}
 
+	/** Try to find the preview control node instance for this anim graph node */
+	FAnimNode_SkeletalControlBase* FindDebugAnimNode(USkeletalMeshComponent * PreviewSkelMeshComp) const;
+
 protected:
 	// Returns the short descriptive name of the controller
 	virtual FText GetControllerDescription() const;
diff --git a/Engine/Source/Editor/AnimGraph/Private/AnimGraphNode_OrientationDriver.cpp b/Engine/Source/Editor/AnimGraph/Private/AnimGraphNode_OrientationDriver.cpp
new file mode 100644
--- /dev/null
+++ b/Engine/Source/Editor/AnimGraph/Private/AnimGraphNode_OrientationDriver.cpp
@@ -0,0 +1,148 @@
+// Copyright 1998-2016 Epic Games, Inc. All Rights Reserved.
+
+#include "AnimGraphPrivatePCH.h"
+#include "AnimGraphNode_OrientationDriver.h"
+#include "CompilerResultsLog.h"
+#include "PropertyEditing.h"
+#include "AnimationCustomVersion.h"
+#include "ContentBrowserModule.h"
+#include "Animation/PoseAsset.h"
+
+#define LOCTEXT_NAMESPACE "A3Nodes"
+
+UAnimGraphNode_OrientationDriver::UAnimGraphNode_OrientationDriver(const FObjectInitializer& ObjectInitializer)
+	:Super(ObjectInitializer)
+{
+
+}
+
+FText UAnimGraphNode_OrientationDriver::GetTooltipText() const
+{
+	return LOCTEXT("UAnimGraphNode_OrientationDriver_ToolTip", "Interpolate parameters between defined orientation poses.");
+}
+
+FText UAnimGraphNode_OrientationDriver::GetNodeTitle(ENodeTitleType::Type TitleType) const
+{
+	return GetControllerDescription();
+}
+
+FText UAnimGraphNode_OrientationDriver::GetControllerDescription() const
+{
+	return LOCTEXT("OrientationDriver", "Orientation Driver");
+}
+
+void UAnimGraphNode_OrientationDriver::Draw(FPrimitiveDrawInterface* PDI, USkeletalMeshComponent* SkelMeshComp) const
+{
+
+}
+
+void UAnimGraphNode_OrientationDriver::DrawCanvas(FViewport& InViewport, FSceneView& View, FCanvas& Canvas, USkeletalMeshComponent * PreviewSkelMeshComp) const
+{
+	// Get the node we are debugging
+	FAnimNode_OrientationDriver* DebugNode = (FAnimNode_OrientationDriver*)FindDebugAnimNode(PreviewSkelMeshComp);
+
+	if (DebugNode == nullptr)
+	{
+		return;
+	}
+
+	const int MarginX = 60;
+	const int XDelta = 10;
+
+	int PosY = 60;
+	const int YDelta = 20;
+
+	// Display bone we are watching
+	FText InterpItemText = FText::Format(LOCTEXT("BoneFormat", "Bone: {0}"), FText::FromName(DebugNode->SourceBone.BoneName));
+	FCanvasTextItem InterpItem(FVector2D(MarginX, PosY), InterpItemText, GEngine->GetSmallFont(), FLinearColor::White);
+	Canvas.DrawItem(InterpItem);
+	PosY += YDelta;
+
+	// List each pose with weight
+	int32 PoseIndex = 0;
+	for (const FOrientationDriverPose& Pose : DebugNode->Poses)
+	{
+		FString PoseItemString = FString::Printf(TEXT("%d : %f"), PoseIndex, Pose.PoseWeight);
+		FCanvasTextItem PoseItem(FVector2D(MarginX + XDelta, PosY), FText::FromString(PoseItemString), GEngine->GetSmallFont(), FLinearColor::White);
+		Canvas.DrawItem(PoseItem);
+		PosY += YDelta;
+
+		PoseIndex++;
+	}
+
+	PosY += YDelta;
+
+	// Parameters header
+	FCanvasTextItem ParamsHeaderItem(FVector2D(MarginX, PosY), LOCTEXT("Parameters", "Parameters"), GEngine->GetSmallFont(), FLinearColor::White);
+	Canvas.DrawItem(ParamsHeaderItem);
+	PosY += YDelta;
+
+	// Iterate over Parameters
+	for (const FOrientationDriverParam& Param : DebugNode->ResultParamSet.Params)
+	{
+		FString ParamItemString = FString::Printf(TEXT("%s : %f"), *Param.ParameterName.ToString(), Param.ParameterValue);
+		FCanvasTextItem ParamItem(FVector2D(MarginX + XDelta, PosY), FText::FromString(ParamItemString), GEngine->GetSmallFont(), FLinearColor::White);
+		Canvas.DrawItem(ParamItem);
+		PosY += YDelta;
+	}
+}
+
+
+void UAnimGraphNode_OrientationDriver::ValidateAnimNodeDuringCompilation(USkeleton* ForSkeleton, FCompilerResultsLog& MessageLog)
+{
+	if (ForSkeleton->GetReferenceSkeleton().FindBoneIndex(Node.SourceBone.BoneName) == INDEX_NONE)
+	{
+		MessageLog.Warning(*LOCTEXT("NoSourceBone", "@@ - You must pick a source bone as the Driver joint").ToString(), this);
+	}
+
+	Super::ValidateAnimNodeDuringCompilation(ForSkeleton, MessageLog);
+}
+
+
+void UAnimGraphNode_OrientationDriver::OnPoseAssetSelected(const FAssetData& AssetData)
+{
+
+}
+
+TSharedRef<SWidget> UAnimGraphNode_OrientationDriver::BuildPoseAssetPicker()
+{
+	FContentBrowserModule& ContentBrowserModule = FModuleManager::Get().LoadModuleChecked<FContentBrowserModule>(TEXT("ContentBrowser"));
+
+	FAssetPickerConfig AssetPickerConfig;
+	AssetPickerConfig.Filter.ClassNames.Add(UPoseAsset::StaticClass()->GetFName());
+	AssetPickerConfig.OnAssetSelected = FOnAssetSelected::CreateUObject(this, &UAnimGraphNode_OrientationDriver::OnPoseAssetSelected);
+	AssetPickerConfig.bAllowNullSelection = true;
+	AssetPickerConfig.InitialAssetViewType = EAssetViewType::List;
+	AssetPickerConfig.bFocusSearchBoxWhenOpened = true;
+	AssetPickerConfig.bShowBottomToolbar = false;
+	AssetPickerConfig.SelectionMode = ESelectionMode::Single;
+
+	return 
+		SNew(SBox)
+		.WidthOverride(384)
+		.HeightOverride(768)
+		[
+			ContentBrowserModule.Get().CreateAssetPicker(AssetPickerConfig)
+		];
+}
+
+void UAnimGraphNode_OrientationDriver::CustomizeDetails(IDetailLayoutBuilder& DetailBuilder)
+{
+	IDetailCategoryBuilder& ImportCategory = DetailBuilder.EditCategory("Import Driver Setup");
+
+	ImportCategory.AddCustomRow(LOCTEXT("ImportRow", "Import"))
+	[
+		SNew(SComboButton)
+		.ContentPadding(3)
+		.OnGetMenuContent(FOnGetContent::CreateUObject(this, &UAnimGraphNode_OrientationDriver::BuildPoseAssetPicker))
+		.ButtonContent()
+		[
+			SNew(STextBlock)
+			.Text(LOCTEXT("ImportPosesFromPoseAsset ", "Import Poses From PoseAsset"))
+			.ToolTipText(LOCTEXT("ImportPosesFromPoseAssetTooltip", "Import target poses from a PoseAsset"))
+		]
+	];
+}
+
+
+#undef LOCTEXT_NAMESPACE
diff --git a/Engine/Source/Editor/AnimGraph/Private/AnimGraphNode_SkeletalControlBase.cpp b/Engine/Source/Editor/AnimGraph/Private/AnimGraphNode_SkeletalControlBase.cpp
--- a/Engine/Source/Editor/AnimGraph/Private/AnimGraphNode_SkeletalControlBase.cpp
+++ b/Engine/Source/Editor/AnimGraph/Private/AnimGraphNode_SkeletalControlBase.cpp
@@ -4,6 +4,7 @@
 #include "AnimGraphNode_SkeletalControlBase.h"
 #include "AnimationGraphSchema.h"
 #include "Animation/AnimationSettings.h"
+#include "Animation/AnimInstance.h"
 #include "CompilerResultsLog.h"
 #include "AnimNode_SkeletalControlBase.h"
 /////////////////////////////////////////////////////
@@ -70,6 +71,34 @@ void UAnimGraphNode_SkeletalControlBase::CreateOutputPins()
 	CreatePin(EGPD_Output, Schema->PC_Struct, TEXT(""), FComponentSpacePoseLink::StaticStruct(), /*bIsArray=*/ false, /*bIsReference=*/ false, TEXT("Pose"));
 }
 
+FAnimNode_SkeletalControlBase* UAnimGraphNode_SkeletalControlBase::FindDebugAnimNode(USkeletalMeshComponent * PreviewSkelMeshComp) const
+{
+	FAnimNode_SkeletalControlBase* DebugNode = nullptr;
+
+	if (PreviewSkelMeshComp != nullptr && PreviewSkelMeshComp->GetAnimInstance() != nullptr)
+	{
+		// find an anim node index from debug data
+		UAnimBlueprintGeneratedClass* AnimBlueprintClass = Cast<UAnimBlueprintGeneratedClass>(PreviewSkelMeshComp->GetAnimInstance()->GetClass());
+		if (AnimBlueprintClass)
+		{
+			FAnimBlueprintDebugData& DebugData = AnimBlueprintClass->GetAnimBlueprintDebugData();
+			int32* IndexPtr = DebugData.NodePropertyToIndexMap.Find(this);
+
+			if (IndexPtr)
+			{
+				int32 AnimNodeIndex = *IndexPtr;
+				// reverse node index temporarily because of a bug in NodeGuidToIndexMap
+				AnimNodeIndex = AnimBlueprintClass->AnimNodeProperties.Num() - AnimNodeIndex - 1;
+
+				DebugNode = AnimBlueprintClass->AnimNodeProperties[AnimNodeIndex]->ContainerPtrToValuePtr<FAnimNode_SkeletalControlBase>(PreviewSkelMeshComp->GetAnimInstance());
+			}
+		}
+	}
+
+	return DebugNode;
+}
+
+
 void UAnimGraphNode_SkeletalControlBase::ConvertToComponentSpaceTransform(const USkeletalMeshComponent* SkelComp, const FTransform & InTransform, FTransform & OutCSTransform, int32 BoneIndex, EBoneControlSpace Space) const
 {
 	USkeleton * Skeleton = SkelComp->SkeletalMesh->Skeleton;
diff --git a/Engine/Source/Editor/DetailCustomizations/Private/BodyInstanceCustomization.cpp b/Engine/Source/Editor/DetailCustomizations/Private/BodyInstanceCustomization.cpp
--- a/Engine/Source/Editor/DetailCustomizations/Private/BodyInstanceCustomization.cpp
+++ b/Engine/Source/Editor/DetailCustomizations/Private/BodyInstanceCustomization.cpp
@@ -28,8 +28,6 @@ UStaticMeshComponent* FBodyInstanceCustomization::GetDefaultCollisionProvider(co
 	{
 		TWeakObjectPtr<UPrimitiveComponent> FoundComp = BodyInstanceToPrimComponent.FindRef(BI);
 		OwnerComp = FoundComp.Get();
-		
-		ensureMsgf(!FoundComp.IsStale(), TEXT("BodyInstance has no owner and the map data is stale"));
 	}
 	
 	UStaticMeshComponent* SMC = Cast<UStaticMeshComponent>(OwnerComp);
diff --git a/Engine/Source/Editor/DetailCustomizations/Private/DetailCustomizations.cpp b/Engine/Source/Editor/DetailCustomizations/Private/DetailCustomizations.cpp
--- a/Engine/Source/Editor/DetailCustomizations/Private/DetailCustomizations.cpp
+++ b/Engine/Source/Editor/DetailCustomizations/Private/DetailCustomizations.cpp
@@ -113,7 +113,7 @@
 #include "ColorStructCustomization.h"
 #include "RotatorStructCustomization.h"
 #include "VectorStructCustomization.h"
-
+#include "AssetViewerSettingsCustomization.h"
 
 IMPLEMENT_MODULE( FDetailCustomizationsModule, DetailCustomizations );
 
@@ -180,6 +180,7 @@ void FDetailCustomizationsModule::RegisterPropertyTypeCustomizations()
 	RegisterCustomPropertyTypeLayout("DialogueContext", FOnGetPropertyTypeCustomizationInstance::CreateStatic(&FDialogueContextStructCustomization::MakeInstance));
 	RegisterCustomPropertyTypeLayout("DialogueWaveParameter", FOnGetPropertyTypeCustomizationInstance::CreateStatic(&FDialogueWaveParameterStructCustomization::MakeInstance));
 	RegisterCustomPropertyTypeLayout("BodyInstance", FOnGetPropertyTypeCustomizationInstance::CreateStatic(&FBodyInstanceCustomization::MakeInstance));
+	RegisterCustomPropertyTypeLayout("PhysicsConstraintProfileHandle", FOnGetPropertyTypeCustomizationInstance::CreateStatic(&FPhysicsConstraintProfileHandleCustomization::MakeInstance));
 	RegisterCustomPropertyTypeLayout("VehicleTransmissionData", FOnGetPropertyTypeCustomizationInstance::CreateStatic(&FVehicleTransmissionDataCustomization::MakeInstance));
 	RegisterCustomPropertyTypeLayout("SlateBrush", FOnGetPropertyTypeCustomizationInstance::CreateStatic(&FSlateBrushStructCustomization::MakeInstance, true));
 	RegisterCustomPropertyTypeLayout("SlateSound", FOnGetPropertyTypeCustomizationInstance::CreateStatic(&FSlateSoundStructCustomization::MakeInstance));
@@ -232,7 +233,6 @@ void FDetailCustomizationsModule::RegisterPropertyTypeCustomizations()
 	RegisterCustomPropertyTypeLayout("LevelSequenceBurnInInitSettings", FOnGetPropertyTypeCustomizationInstance::CreateStatic(&FLevelSequenceBurnInInitSettingsCustomization::MakeInstance));
 }
 
-
 void FDetailCustomizationsModule::RegisterObjectCustomizations()
 {
 	// Note: By default properties are displayed in script defined order (i.e the order in the header).  These layout detail classes are called in the order seen here which will display properties
@@ -321,6 +321,8 @@ void FDetailCustomizationsModule::RegisterObjectCustomizations()
 
 	RegisterCustomClassLayout("AnalyticsPrivacySettings", FOnGetDetailCustomizationInstance::CreateStatic(&FImportantToggleSettingCustomization::MakeInstance));
 	RegisterCustomClassLayout("EndUserSettings", FOnGetDetailCustomizationInstance::CreateStatic(&FImportantToggleSettingCustomization::MakeInstance));
+
+	RegisterCustomClassLayout("AssetViewerSettings", FOnGetDetailCustomizationInstance::CreateStatic(&FAssetViewerSettingsCustomization::MakeInstance));
 }
 
 
diff --git a/Engine/Source/Editor/DetailCustomizations/Private/PhysicsConstraintComponentDetails.cpp b/Engine/Source/Editor/DetailCustomizations/Private/PhysicsConstraintComponentDetails.cpp
--- a/Engine/Source/Editor/DetailCustomizations/Private/PhysicsConstraintComponentDetails.cpp
+++ b/Engine/Source/Editor/DetailCustomizations/Private/PhysicsConstraintComponentDetails.cpp
@@ -5,76 +5,173 @@
 #include "PhysicsEngine/PhysicsConstraintTemplate.h"
 #include "PhysicsEngine/PhysicsConstraintComponent.h"
 #include "PhysicsEngine/PhysicsConstraintActor.h"
+#include "SNumericEntryBox.h"
+#include "ScopedTransaction.h"
 
 #define LOCTEXT_NAMESPACE "PhysicsConstraintComponentDetails"
+namespace ConstraintDetails
+{
+	bool GetBoolProperty(TSharedPtr<IPropertyHandle> Prop)
+	{
+		bool bIsEnabled = false;
+
+		if (Prop->GetValue(bIsEnabled))
+		{
+			return bIsEnabled;
+		}
+		return false;
+	}
+
+	TSharedRef<SWidget> JoinPropertyWidgets(TSharedPtr<IPropertyHandle> TargetProperty, FName TargetChildName, TSharedPtr<IPropertyHandle> ParentProperty, FName CheckPropertyName, TSharedPtr<IPropertyHandle>& StoreCheckProperty)
+	{
+		StoreCheckProperty = ParentProperty->GetChildHandle(CheckPropertyName);
+		StoreCheckProperty->MarkHiddenByCustomization();
+		TSharedRef<SWidget> TargetWidget = TargetProperty->GetChildHandle(TargetChildName)->CreatePropertyValueWidget();
+		TargetWidget->SetEnabled(TAttribute<bool>::Create(TAttribute<bool>::FGetter::CreateLambda([StoreCheckProperty]()
+		{
+			bool bSet;
+			if (StoreCheckProperty->GetValue(bSet))
+			{
+				return bSet;
+			}
+
+			return false;
+		})));
+		
+		return SNew(SHorizontalBox)
+			+ SHorizontalBox::Slot()
+			.AutoWidth()
+			.Padding(0, 0, 5, 0)
+			[
+				StoreCheckProperty->CreatePropertyValueWidget()
+			]
+		+ SHorizontalBox::Slot()
+			[
+				TargetWidget
+			];
+	}
+
+	TSharedRef<SWidget> CreateTriFloatWidget(TSharedPtr<IPropertyHandle> Prop1, TSharedPtr<IPropertyHandle> Prop2, TSharedPtr<IPropertyHandle> Prop3, const FText& TransactionName)
+	{
+		auto GetMultipleFloats = [Prop1, Prop2, Prop3]()
+		{
+			float Val1, Val2, Val3;
+			ensure(Prop1->GetValue(Val1) != FPropertyAccess::Fail);
+			ensure(Prop2->GetValue(Val2) != FPropertyAccess::Fail);
+			ensure(Prop3->GetValue(Val3) != FPropertyAccess::Fail);
+
+			if (Val1 == Val2 && Val2 == Val3)
+			{
+				return TOptional<float>(Val1);
+			}
+
+			return TOptional<float>();
+		};
+
+		auto SetMultipleFloatsCommitted = [Prop1, Prop2, Prop3, TransactionName, GetMultipleFloats](float NewValue, ETextCommit::Type)
+		{
+			TOptional<float> CommonFloat = GetMultipleFloats();
+			if(!CommonFloat.IsSet() || CommonFloat.GetValue() != NewValue)	//don't bother doing it twice
+			{
+				FScopedTransaction Transaction(TransactionName);
+
+				ensure(Prop1->SetValue(NewValue) == FPropertyAccess::Success);
+				ensure(Prop2->SetValue(NewValue) == FPropertyAccess::Success);
+				ensure(Prop3->SetValue(NewValue) == FPropertyAccess::Success);
+			}
+		};
+
+		return SNew(SHorizontalBox)
+			+ SHorizontalBox::Slot()
+			[
+				SNew(SNumericEntryBox<float>)
+				.OnValueCommitted_Lambda(SetMultipleFloatsCommitted)
+				.Value_Lambda(GetMultipleFloats)
+				.MinValue(0.f)
+			]
+			+ SHorizontalBox::Slot()
+			.VAlign(VAlign_Center)
+			.AutoWidth()
+			[
+				SNew(SButton)
+				.OnClicked_Lambda([Prop1, Prop2, Prop3, TransactionName]()
+				{
+					FScopedTransaction Transaction(TransactionName);
+					Prop1->ResetToDefault();
+					Prop2->ResetToDefault();
+					Prop3->ResetToDefault();
+					return FReply::Handled();
+				} )
+				.Visibility_Lambda([Prop1]() { return Prop1->DiffersFromDefault() ? EVisibility::Visible : EVisibility::Collapsed; })
+				.ContentPadding(FMargin(5.f, 0.f))
+				.ToolTipText(Prop1->GetResetToDefaultLabel())
+				.ButtonStyle(FEditorStyle::Get(), "NoBorder")
+				.Content()
+				[
+					SNew(SImage)
+					.Image(FEditorStyle::GetBrush("PropertyWindow.DiffersFromDefault"))
+				]
+			];
+	}
+
+
+	bool IsAngularPropertyEqual(TSharedPtr<IPropertyHandle> Prop, EAngularConstraintMotion CheckMotion)
+	{
+		uint8 Val;
+		if (Prop->GetValue(Val))
+		{
+			return Val == CheckMotion;
+		}
+		return false;
+	}
+
+}
 
 TSharedRef<IDetailCustomization> FPhysicsConstraintComponentDetails::MakeInstance()
 {
 	return MakeShareable(new FPhysicsConstraintComponentDetails());
 }
 
-void FPhysicsConstraintComponentDetails::AddConstraintPresets(IDetailLayoutBuilder& DetailBuilder, APhysicsConstraintActor* OwningConstraintActor)
+void FPhysicsConstraintComponentDetails::AddConstraintBehaviorProperties(IDetailLayoutBuilder& DetailBuilder, TSharedPtr<IPropertyHandle> ConstraintInstance, TSharedPtr<IPropertyHandle> ProfilePropertiesProperty)
 {
-	if (OwningConstraintActor != NULL)
+	IDetailCategoryBuilder& ConstraintCat = DetailBuilder.EditCategory("Constraint Behavior");
+
+	//hide the inner structs that we customize elsewhere
+	ProfilePropertiesProperty->GetChildHandle(GET_MEMBER_NAME_CHECKED(FConstraintProfileProperties, LinearLimit))->MarkHiddenByCustomization();
+	ProfilePropertiesProperty->GetChildHandle(GET_MEMBER_NAME_CHECKED(FConstraintProfileProperties, ConeLimit))->MarkHiddenByCustomization();
+	ProfilePropertiesProperty->GetChildHandle(GET_MEMBER_NAME_CHECKED(FConstraintProfileProperties, TwistLimit))->MarkHiddenByCustomization();
+	ProfilePropertiesProperty->GetChildHandle(GET_MEMBER_NAME_CHECKED(FConstraintProfileProperties, LinearDrive))->MarkHiddenByCustomization();
+	ProfilePropertiesProperty->GetChildHandle(GET_MEMBER_NAME_CHECKED(FConstraintProfileProperties, AngularDrive))->MarkHiddenByCustomization();
+	ProfilePropertiesProperty->MarkHiddenByCustomization();
+
+	//Add properties we want in specific order
+	ConstraintCat.AddProperty(ProfilePropertiesProperty->GetChildHandle(GET_MEMBER_NAME_CHECKED(FConstraintProfileProperties, bDisableCollision)));
+	ConstraintCat.AddProperty(ProfilePropertiesProperty->GetChildHandle(GET_MEMBER_NAME_CHECKED(FConstraintProfileProperties, bEnableProjection)));
+
+	
+	//Add the rest
+	uint32 NumProfileProperties = 0;
+	ProfilePropertiesProperty->GetNumChildren(NumProfileProperties);
+
+	for (uint32 ProfileChildIdx = 0; ProfileChildIdx < NumProfileProperties; ++ProfileChildIdx)
 	{
-		DetailBuilder.EditCategory("Joint Presets")
-			.AddCustomRow(FText::GetEmpty())
-			[
-				SNew(SVerticalBox)
-				+ SVerticalBox::Slot()
-				.Padding(0, 2.0f, 0, 0)
-				.FillHeight(1.0f)
-				.VAlign(VAlign_Center)
-				[
-					SNew(SHorizontalBox)
-					+ SHorizontalBox::Slot()
-					.AutoWidth()
-					.Padding(2.0f, 0.0f)
-					.VAlign(VAlign_Center)
-					.HAlign(HAlign_Left)
-					[
-						SNew(SButton)
-						.VAlign(VAlign_Center)
-						.OnClicked(this, &FPhysicsConstraintComponentDetails::OnHingeClicked)
-						.Text(LOCTEXT("HingePreset", "Hinge"))
-						.ToolTipText(LOCTEXT("HingePresetTooltip", "Setup joint as hinge. A hinge joint allows motion only in one plane."))
-					]
-					+ SHorizontalBox::Slot()
-						.AutoWidth()
-						.Padding(2.0f, 0.0f)
-						.VAlign(VAlign_Center)
-						.HAlign(HAlign_Left)
-						[
-							SNew(SButton)
-							.VAlign(VAlign_Center)
-							.OnClicked(this, &FPhysicsConstraintComponentDetails::OnPrismaticClicked)
-							.Text(LOCTEXT("PrismaticPreset", "Prismatic"))
-							.ToolTipText(LOCTEXT("PrismaticPresetTooltip", "Setup joint as prismatic. A prismatic joint allows only linear sliding movement."))
-						]
-					+ SHorizontalBox::Slot()
-						.AutoWidth()
-						.Padding(2.0f, 0.0f)
-						.VAlign(VAlign_Center)
-						.HAlign(HAlign_Left)
-						[
-							SNew(SButton)
-							.VAlign(VAlign_Center)
-							.OnClicked(this, &FPhysicsConstraintComponentDetails::OnBallSocketClicked)
-							.Text(LOCTEXT("BSPreset", "Ball and Socket"))
-							.ToolTipText(LOCTEXT("BSPresetTooltip", "Setup joint as ball and socket. A Ball and Socket joint allows motion around an indefinite number of axes, which have one common center"))
-						]
-				]
-			];
+		TSharedPtr<IPropertyHandle> ProfileChildProp = ProfilePropertiesProperty->GetChildHandle(ProfileChildIdx);
+		if (!ProfileChildProp->IsCustomized())
+		{
+			ConstraintCat.AddProperty(ProfileChildProp);
+		}
 	}
 }
 
-void FPhysicsConstraintComponentDetails::AddLinearLimits(IDetailLayoutBuilder& DetailBuilder, TSharedPtr<IPropertyHandle> ConstraintInstance)
+void FPhysicsConstraintComponentDetails::AddLinearLimits(IDetailLayoutBuilder& DetailBuilder, TSharedPtr<IPropertyHandle> ConstraintInstance, TSharedPtr<IPropertyHandle> ProfilePropertiesProperty)
 {
 	IDetailCategoryBuilder& LinearLimitCat = DetailBuilder.EditCategory("Linear Limits");
+	TSharedPtr<IPropertyHandle> LinearConstraintProperty = ProfilePropertiesProperty->GetChildHandle(GET_MEMBER_NAME_CHECKED(FConstraintProfileProperties, LinearLimit));
+
 
-	LinearXMotionProperty = ConstraintInstance->GetChildHandle("LinearXMotion");
-	LinearYMotionProperty = ConstraintInstance->GetChildHandle("LinearYMotion");
-	LinearZMotionProperty = ConstraintInstance->GetChildHandle("LinearZMotion");
+	TSharedPtr<IPropertyHandle> LinearXMotionProperty = LinearConstraintProperty->GetChildHandle(GET_MEMBER_NAME_CHECKED(FLinearConstraint, XMotion));
+	TSharedPtr<IPropertyHandle> LinearYMotionProperty = LinearConstraintProperty->GetChildHandle(GET_MEMBER_NAME_CHECKED(FLinearConstraint, YMotion));
+	TSharedPtr<IPropertyHandle> LinearZMotionProperty = LinearConstraintProperty->GetChildHandle(GET_MEMBER_NAME_CHECKED(FLinearConstraint, ZMotion));
 
 	TArray<TSharedPtr<FString>> LinearLimitOptionNames;
 	TArray<FText> LinearLimitOptionTooltips;
@@ -158,24 +255,48 @@ void FPhysicsConstraintComponentDetails::AddLinearLimits(IDetailLayoutBuilder& D
 					]
 			];
 	}
+	
+	auto IsLinearMotionLimited = [LinearXMotionProperty, LinearYMotionProperty, LinearZMotionProperty]()
+	{
+		uint8 XMotion, YMotion, ZMotion;
+		if (LinearXMotionProperty->GetValue(XMotion) && LinearYMotionProperty->GetValue(YMotion) && LinearZMotionProperty->GetValue(ZMotion))
+		{
+			return XMotion == LCM_Limited || YMotion == LCM_Limited || ZMotion == LCM_Limited;
+		}
 
+		return false;
+	};
 
-	LinearLimitCat.AddProperty(ConstraintInstance->GetChildHandle(GET_MEMBER_NAME_CHECKED(FConstraintInstance, LinearLimitSize)).ToSharedRef()).IsEnabled(TAttribute<bool>::Create(TAttribute<bool>::FGetter::CreateSP(this, &FPhysicsConstraintComponentDetails::IsPropertyEnabled, EPropertyType::LinearLimit)));
-	LinearLimitCat.AddProperty(ConstraintInstance->GetChildHandle(GET_MEMBER_NAME_CHECKED(FConstraintInstance, bScaleLinearLimits)).ToSharedRef()).IsEnabled(TAttribute<bool>::Create(TAttribute<bool>::FGetter::CreateSP(this, &FPhysicsConstraintComponentDetails::IsPropertyEnabled, EPropertyType::LinearLimit)));
-	LinearLimitCat.AddProperty(ConstraintInstance->GetChildHandle(GET_MEMBER_NAME_CHECKED(FConstraintInstance, bLinearLimitSoft)).ToSharedRef()).IsEnabled(TAttribute<bool>::Create(TAttribute<bool>::FGetter::CreateSP(this, &FPhysicsConstraintComponentDetails::IsPropertyEnabled, EPropertyType::LinearLimit)));
-	LinearLimitCat.AddProperty(ConstraintInstance->GetChildHandle(GET_MEMBER_NAME_CHECKED(FConstraintInstance, LinearLimitStiffness)).ToSharedRef()).IsEnabled(TAttribute<bool>::Create(TAttribute<bool>::FGetter::CreateSP(this, &FPhysicsConstraintComponentDetails::IsPropertyEnabled, EPropertyType::LinearLimit)));
-	LinearLimitCat.AddProperty(ConstraintInstance->GetChildHandle(GET_MEMBER_NAME_CHECKED(FConstraintInstance, LinearLimitDamping)).ToSharedRef()).IsEnabled(TAttribute<bool>::Create(TAttribute<bool>::FGetter::CreateSP(this, &FPhysicsConstraintComponentDetails::IsPropertyEnabled, EPropertyType::LinearLimit)));
-	LinearLimitCat.AddProperty(ConstraintInstance->GetChildHandle(GET_MEMBER_NAME_CHECKED(FConstraintInstance, bLinearBreakable)).ToSharedRef());
-	LinearLimitCat.AddProperty(ConstraintInstance->GetChildHandle(GET_MEMBER_NAME_CHECKED(FConstraintInstance, LinearBreakThreshold)).ToSharedRef());
+	TSharedPtr<IPropertyHandle> SoftProperty = LinearConstraintProperty->GetChildHandle(GET_MEMBER_NAME_CHECKED(FLinearConstraint, bSoftConstraint));
+
+	auto IsRestitutionEnabled = [IsLinearMotionLimited, SoftProperty]()
+	{
+		return !ConstraintDetails::GetBoolProperty(SoftProperty) && IsLinearMotionLimited();
+	};
+
+	LinearLimitCat.AddProperty(LinearConstraintProperty->GetChildHandle(GET_MEMBER_NAME_CHECKED(FLinearConstraint, Limit))).IsEnabled(TAttribute<bool>::Create(TAttribute<bool>::FGetter::CreateLambda(IsLinearMotionLimited)));
+	LinearLimitCat.AddProperty(ConstraintInstance->GetChildHandle(GET_MEMBER_NAME_CHECKED(FConstraintInstance, bScaleLinearLimits))).IsEnabled(TAttribute<bool>::Create(TAttribute<bool>::FGetter::CreateLambda(IsLinearMotionLimited)));
+	LinearLimitCat.AddProperty(LinearConstraintProperty->GetChildHandle(GET_MEMBER_NAME_CHECKED(FLinearConstraint, bSoftConstraint))).IsEnabled(TAttribute<bool>::Create(TAttribute<bool>::FGetter::CreateLambda(IsLinearMotionLimited)));
+	LinearLimitCat.AddProperty(LinearConstraintProperty->GetChildHandle(GET_MEMBER_NAME_CHECKED(FLinearConstraint, Stiffness))).IsEnabled(TAttribute<bool>::Create(TAttribute<bool>::FGetter::CreateLambda(IsLinearMotionLimited)));
+	LinearLimitCat.AddProperty(LinearConstraintProperty->GetChildHandle(GET_MEMBER_NAME_CHECKED(FLinearConstraint, Damping))).IsEnabled(TAttribute<bool>::Create(TAttribute<bool>::FGetter::CreateLambda(IsLinearMotionLimited)));
+	LinearLimitCat.AddProperty(LinearConstraintProperty->GetChildHandle(GET_MEMBER_NAME_CHECKED(FLinearConstraint, Restitution))).IsEnabled(TAttribute<bool>::Create(TAttribute<bool>::FGetter::CreateLambda(IsRestitutionEnabled)));
+	LinearLimitCat.AddProperty(LinearConstraintProperty->GetChildHandle(GET_MEMBER_NAME_CHECKED(FLinearConstraint, ContactDistance))).IsEnabled(TAttribute<bool>::Create(TAttribute<bool>::FGetter::CreateLambda(IsLinearMotionLimited)));
+	LinearLimitCat.AddProperty(ProfilePropertiesProperty->GetChildHandle(GET_MEMBER_NAME_CHECKED(FConstraintProfileProperties, bLinearBreakable)));
+	LinearLimitCat.AddProperty(ProfilePropertiesProperty->GetChildHandle(GET_MEMBER_NAME_CHECKED(FConstraintProfileProperties, LinearBreakThreshold)));
 }
 
-void FPhysicsConstraintComponentDetails::AddAngularLimits(IDetailLayoutBuilder& DetailBuilder, TSharedPtr<IPropertyHandle> ConstraintInstance)
+
+
+void FPhysicsConstraintComponentDetails::AddAngularLimits(IDetailLayoutBuilder& DetailBuilder, TSharedPtr<IPropertyHandle> ConstraintInstance, TSharedPtr<IPropertyHandle> ProfilePropertiesProperty)
 {
 	IDetailCategoryBuilder& AngularLimitCat = DetailBuilder.EditCategory("Angular Limits");
 
-	AngularSwing1MotionProperty = ConstraintInstance->GetChildHandle(GET_MEMBER_NAME_CHECKED(FConstraintInstance, AngularSwing1Motion));
-	AngularSwing2MotionProperty = ConstraintInstance->GetChildHandle(GET_MEMBER_NAME_CHECKED(FConstraintInstance, AngularSwing2Motion));
-	AngularTwistMotionProperty = ConstraintInstance->GetChildHandle(GET_MEMBER_NAME_CHECKED(FConstraintInstance, AngularTwistMotion));
+	TSharedPtr<IPropertyHandle> ConeConstraintProperty = ProfilePropertiesProperty->GetChildHandle(GET_MEMBER_NAME_CHECKED(FConstraintProfileProperties, ConeLimit));
+	TSharedPtr<IPropertyHandle> TwistConstraintProperty = ProfilePropertiesProperty->GetChildHandle(GET_MEMBER_NAME_CHECKED(FConstraintProfileProperties, TwistLimit));
+
+	AngularSwing1MotionProperty = ConeConstraintProperty->GetChildHandle(GET_MEMBER_NAME_CHECKED(FConeConstraint, Swing1Motion));
+	AngularSwing2MotionProperty = ConeConstraintProperty->GetChildHandle(GET_MEMBER_NAME_CHECKED(FConeConstraint, Swing2Motion));
+	AngularTwistMotionProperty = TwistConstraintProperty->GetChildHandle(GET_MEMBER_NAME_CHECKED(FTwistConstraint, TwistMotion));
 
 	TArray<TSharedPtr<FString>> AngularLimitOptionNames;
 	TArray<FText> AngularLimitOptionTooltips;
@@ -190,7 +311,7 @@ void FPhysicsConstraintComponentDetails::AddAngularLimits(IDetailLayoutBuilder&
 
 
 	uint8 AngularLimitEnum[LCM_MAX] = { ACM_Free, LCM_Limited, LCM_Locked };
-	TSharedPtr<IPropertyHandle> AngularLimitProperties[] = { AngularSwing1MotionProperty, AngularTwistMotionProperty, AngularSwing2MotionProperty };
+	TSharedPtr<IPropertyHandle> AngularLimitProperties[] = { AngularSwing1MotionProperty, AngularSwing2MotionProperty, AngularTwistMotionProperty };
 
 	for (int32 PropertyIdx = 0; PropertyIdx < 3; ++PropertyIdx)
 	{
@@ -259,26 +380,37 @@ void FPhysicsConstraintComponentDetails::AddAngularLimits(IDetailLayoutBuilder&
 			];
 	}
 
-	AngularLimitCat.AddProperty(ConstraintInstance->GetChildHandle(GET_MEMBER_NAME_CHECKED(FConstraintInstance, Swing1LimitAngle)).ToSharedRef()).IsEnabled(TAttribute<bool>::Create(TAttribute<bool>::FGetter::CreateSP(this, &FPhysicsConstraintComponentDetails::IsPropertyEnabled, EPropertyType::AngularSwing1Limit)));
-	AngularLimitCat.AddProperty(ConstraintInstance->GetChildHandle(GET_MEMBER_NAME_CHECKED(FConstraintInstance, TwistLimitAngle)).ToSharedRef()).IsEnabled(TAttribute<bool>::Create(TAttribute<bool>::FGetter::CreateSP(this, &FPhysicsConstraintComponentDetails::IsPropertyEnabled, EPropertyType::AngularTwistLimit)));
-	AngularLimitCat.AddProperty(ConstraintInstance->GetChildHandle(GET_MEMBER_NAME_CHECKED(FConstraintInstance, Swing2LimitAngle)).ToSharedRef()).IsEnabled(TAttribute<bool>::Create(TAttribute<bool>::FGetter::CreateSP(this, &FPhysicsConstraintComponentDetails::IsPropertyEnabled, EPropertyType::AngularSwing2Limit)));
+	AngularLimitCat.AddProperty(ConeConstraintProperty->GetChildHandle(GET_MEMBER_NAME_CHECKED(FConeConstraint, Swing1LimitDegrees)).ToSharedRef()).IsEnabled(TAttribute<bool>::Create(TAttribute<bool>::FGetter::CreateSP(this, &FPhysicsConstraintComponentDetails::IsPropertyEnabled, EPropertyType::AngularSwing1Limit)));
+	AngularLimitCat.AddProperty(ConeConstraintProperty->GetChildHandle(GET_MEMBER_NAME_CHECKED(FConeConstraint, Swing2LimitDegrees)).ToSharedRef()).IsEnabled(TAttribute<bool>::Create(TAttribute<bool>::FGetter::CreateSP(this, &FPhysicsConstraintComponentDetails::IsPropertyEnabled, EPropertyType::AngularSwing2Limit)));
+	AngularLimitCat.AddProperty(TwistConstraintProperty->GetChildHandle(GET_MEMBER_NAME_CHECKED(FTwistConstraint, TwistLimitDegrees)).ToSharedRef()).IsEnabled(TAttribute<bool>::Create(TAttribute<bool>::FGetter::CreateSP(this, &FPhysicsConstraintComponentDetails::IsPropertyEnabled, EPropertyType::AngularTwistLimit)));
 
+	TSharedPtr<IPropertyHandle> SoftSwingProperty = ConeConstraintProperty->GetChildHandle(GET_MEMBER_NAME_CHECKED(FConeConstraint, bSoftConstraint));
+	auto SwingRestitutionEnabled = [this, SoftSwingProperty]()
+	{
+		return !ConstraintDetails::GetBoolProperty(SoftSwingProperty) && (IsPropertyEnabled(EPropertyType::AngularSwing1Limit) || IsPropertyEnabled(EPropertyType::AngularSwing2Limit));
+	};
 
-	AngularLimitCat.AddProperty(ConstraintInstance->GetChildHandle(GET_MEMBER_NAME_CHECKED(FConstraintInstance, bSwingLimitSoft)).ToSharedRef())
-		.IsEnabled(TAttribute<bool>::Create(TAttribute<bool>::FGetter::CreateSP(this, &FPhysicsConstraintComponentDetails::IsPropertyEnabled, EPropertyType::AngularSwingLimit)));
-	AngularLimitCat.AddProperty(ConstraintInstance->GetChildHandle(GET_MEMBER_NAME_CHECKED(FConstraintInstance, SwingLimitStiffness)).ToSharedRef())
-		.IsEnabled(TAttribute<bool>::Create(TAttribute<bool>::FGetter::CreateSP(this, &FPhysicsConstraintComponentDetails::IsPropertyEnabled, EPropertyType::AngularSwingLimit)));
-	AngularLimitCat.AddProperty(ConstraintInstance->GetChildHandle(GET_MEMBER_NAME_CHECKED(FConstraintInstance, SwingLimitDamping)).ToSharedRef())
-		.IsEnabled(TAttribute<bool>::Create(TAttribute<bool>::FGetter::CreateSP(this, &FPhysicsConstraintComponentDetails::IsPropertyEnabled, EPropertyType::AngularSwingLimit)));
+	IDetailGroup& SwingGroup = AngularLimitCat.AddGroup("Swing Limits", LOCTEXT("SwingLimits", "Swing Limits"), true, true);
 
+	SwingGroup.AddPropertyRow(SoftSwingProperty.ToSharedRef()).IsEnabled(TAttribute<bool>::Create(TAttribute<bool>::FGetter::CreateSP(this, &FPhysicsConstraintComponentDetails::IsPropertyEnabled, EPropertyType::AngularSwingLimit)));
+	SwingGroup.AddPropertyRow(ConeConstraintProperty->GetChildHandle(GET_MEMBER_NAME_CHECKED(FConeConstraint, Stiffness)).ToSharedRef()).IsEnabled(TAttribute<bool>::Create(TAttribute<bool>::FGetter::CreateSP(this, &FPhysicsConstraintComponentDetails::IsPropertyEnabled, EPropertyType::AngularSwingLimit)));
+	SwingGroup.AddPropertyRow(ConeConstraintProperty->GetChildHandle(GET_MEMBER_NAME_CHECKED(FConeConstraint, Damping)).ToSharedRef()).IsEnabled(TAttribute<bool>::Create(TAttribute<bool>::FGetter::CreateSP(this, &FPhysicsConstraintComponentDetails::IsPropertyEnabled, EPropertyType::AngularSwingLimit)));
+	SwingGroup.AddPropertyRow(ConeConstraintProperty->GetChildHandle(GET_MEMBER_NAME_CHECKED(FConeConstraint, Restitution)).ToSharedRef()).IsEnabled(TAttribute<bool>::Create(TAttribute<bool>::FGetter::CreateLambda(SwingRestitutionEnabled)));
+	SwingGroup.AddPropertyRow(ConeConstraintProperty->GetChildHandle(GET_MEMBER_NAME_CHECKED(FConeConstraint, ContactDistance)).ToSharedRef()).IsEnabled(TAttribute<bool>::Create(TAttribute<bool>::FGetter::CreateSP(this, &FPhysicsConstraintComponentDetails::IsPropertyEnabled, EPropertyType::AngularSwingLimit)));
 
+	TSharedPtr<IPropertyHandle> SoftTwistProperty = TwistConstraintProperty->GetChildHandle(GET_MEMBER_NAME_CHECKED(FTwistConstraint, bSoftConstraint));
+	auto TwistRestitutionEnabled = [this, SoftTwistProperty]()
+	{
+		return !ConstraintDetails::GetBoolProperty(SoftTwistProperty) && IsPropertyEnabled(EPropertyType::AngularTwistLimit);
+	};
 
-	AngularLimitCat.AddProperty(ConstraintInstance->GetChildHandle(GET_MEMBER_NAME_CHECKED(FConstraintInstance, bTwistLimitSoft)).ToSharedRef())
-		.IsEnabled(TAttribute<bool>::Create(TAttribute<bool>::FGetter::CreateSP(this, &FPhysicsConstraintComponentDetails::IsPropertyEnabled, EPropertyType::AngularTwistLimit)));
-	AngularLimitCat.AddProperty(ConstraintInstance->GetChildHandle(GET_MEMBER_NAME_CHECKED(FConstraintInstance, TwistLimitStiffness)).ToSharedRef())
-		.IsEnabled(TAttribute<bool>::Create(TAttribute<bool>::FGetter::CreateSP(this, &FPhysicsConstraintComponentDetails::IsPropertyEnabled, EPropertyType::AngularTwistLimit)));
-	AngularLimitCat.AddProperty(ConstraintInstance->GetChildHandle(GET_MEMBER_NAME_CHECKED(FConstraintInstance, TwistLimitDamping)).ToSharedRef())
-		.IsEnabled(TAttribute<bool>::Create(TAttribute<bool>::FGetter::CreateSP(this, &FPhysicsConstraintComponentDetails::IsPropertyEnabled, EPropertyType::AngularTwistLimit)));
+	IDetailGroup& TwistGroup = AngularLimitCat.AddGroup("Twist Limits", LOCTEXT("TwistLimits", "Twist Limits"), true, true);
+
+	TwistGroup.AddPropertyRow(SoftTwistProperty.ToSharedRef()).IsEnabled(TAttribute<bool>::Create(TAttribute<bool>::FGetter::CreateSP(this, &FPhysicsConstraintComponentDetails::IsPropertyEnabled, EPropertyType::AngularTwistLimit)));
+	TwistGroup.AddPropertyRow(TwistConstraintProperty->GetChildHandle(GET_MEMBER_NAME_CHECKED(FTwistConstraint, Stiffness)).ToSharedRef()).IsEnabled(TAttribute<bool>::Create(TAttribute<bool>::FGetter::CreateSP(this, &FPhysicsConstraintComponentDetails::IsPropertyEnabled, EPropertyType::AngularTwistLimit)));
+	TwistGroup.AddPropertyRow(TwistConstraintProperty->GetChildHandle(GET_MEMBER_NAME_CHECKED(FTwistConstraint, Damping)).ToSharedRef()).IsEnabled(TAttribute<bool>::Create(TAttribute<bool>::FGetter::CreateSP(this, &FPhysicsConstraintComponentDetails::IsPropertyEnabled, EPropertyType::AngularTwistLimit)));
+	TwistGroup.AddPropertyRow(TwistConstraintProperty->GetChildHandle(GET_MEMBER_NAME_CHECKED(FTwistConstraint, Restitution)).ToSharedRef()).IsEnabled(TAttribute<bool>::Create(TAttribute<bool>::FGetter::CreateLambda(TwistRestitutionEnabled)));
+	TwistGroup.AddPropertyRow(TwistConstraintProperty->GetChildHandle(GET_MEMBER_NAME_CHECKED(FTwistConstraint, ContactDistance)).ToSharedRef()).IsEnabled(TAttribute<bool>::Create(TAttribute<bool>::FGetter::CreateSP(this, &FPhysicsConstraintComponentDetails::IsPropertyEnabled, EPropertyType::AngularTwistLimit)));
 
 	if (bInPhat == false)
 	{
@@ -290,54 +422,39 @@ void FPhysicsConstraintComponentDetails::AddAngularLimits(IDetailLayoutBuilder&
 			.Visibility(EVisibility::Collapsed);
 	}
 
-	AngularLimitCat.AddProperty(ConstraintInstance->GetChildHandle(GET_MEMBER_NAME_CHECKED(FConstraintInstance, bAngularBreakable)).ToSharedRef());
-	AngularLimitCat.AddProperty(ConstraintInstance->GetChildHandle(GET_MEMBER_NAME_CHECKED(FConstraintInstance, AngularBreakThreshold)).ToSharedRef());
+	AngularLimitCat.AddProperty(ProfilePropertiesProperty->GetChildHandle(GET_MEMBER_NAME_CHECKED(FConstraintProfileProperties, bAngularBreakable)).ToSharedRef());
+	AngularLimitCat.AddProperty(ProfilePropertiesProperty->GetChildHandle(GET_MEMBER_NAME_CHECKED(FConstraintProfileProperties, AngularBreakThreshold)).ToSharedRef());
 }
 
-TSharedRef<SWidget> FPhysicsConstraintComponentDetails::JoinPropertyWidgets(TSharedPtr<IPropertyHandle> TargetProperty, FName TargetChildName, TSharedPtr<IPropertyHandle> ConstraintInstance, FName CheckPropertyName, TSharedPtr<IPropertyHandle>& StoreCheckProperty, EPropertyType::Type DriveType)
-{
-	StoreCheckProperty = ConstraintInstance->GetChildHandle(CheckPropertyName);
-	StoreCheckProperty->MarkHiddenByCustomization();
-	TSharedRef<SWidget> TargetWidget = TargetProperty->GetChildHandle(TargetChildName)->CreatePropertyValueWidget();
-	TargetWidget->SetEnabled(TAttribute<bool>::Create(TAttribute<bool>::FGetter::CreateSP(this, &FPhysicsConstraintComponentDetails::IsPropertyEnabled, DriveType)));
-
-	return SNew(SHorizontalBox)
-		+ SHorizontalBox::Slot()
-		.AutoWidth()
-		.Padding(0, 0, 5, 0)
-		[
-			StoreCheckProperty->CreatePropertyValueWidget()
-		]
-	+ SHorizontalBox::Slot()
-		[
-			TargetWidget
-		];
-}
-
-
-void FPhysicsConstraintComponentDetails::AddLinearDrive(IDetailLayoutBuilder& DetailBuilder, TSharedPtr<IPropertyHandle> ConstraintInstance)
+void FPhysicsConstraintComponentDetails::AddLinearDrive(IDetailLayoutBuilder& DetailBuilder, TSharedPtr<IPropertyHandle> ConstraintInstance, TSharedPtr<IPropertyHandle> ProfilePropertiesProperty)
 {
 	IDetailCategoryBuilder& LinearMotorCat = DetailBuilder.EditCategory("LinearMotor");
-	LinearPositionDriveProperty = ConstraintInstance->GetChildHandle(GET_MEMBER_NAME_CHECKED(FConstraintInstance, bLinearPositionDrive));
 
-	IDetailGroup& PositionGroup = LinearMotorCat.AddGroup("Linear Position Drive", FText::GetEmpty());
-	PositionGroup.HeaderProperty(LinearPositionDriveProperty.ToSharedRef());
+	TSharedPtr<IPropertyHandle> LinearDriveProperty = ProfilePropertiesProperty->GetChildHandle(GET_MEMBER_NAME_CHECKED(FConstraintProfileProperties, LinearDrive));
 
-	TSharedRef<IPropertyHandle> LinearPositionTargetProperty = ConstraintInstance->GetChildHandle(GET_MEMBER_NAME_CHECKED(FConstraintInstance, LinearPositionTarget)).ToSharedRef();
+	IDetailGroup& PositionGroup = LinearMotorCat.AddGroup("Linear Position Drive", LOCTEXT("LinearPositionDrive", "Linear Position Drive"), false, true);
 	
+	TSharedRef<IPropertyHandle> LinearPositionTargetProperty = LinearDriveProperty->GetChildHandle(GET_MEMBER_NAME_CHECKED(FLinearDriveConstraint, PositionTarget)).ToSharedRef();
+	
+	TSharedPtr<IPropertyHandle> XDriveProperty = LinearDriveProperty->GetChildHandle(GET_MEMBER_NAME_CHECKED(FLinearDriveConstraint, XDrive));
+	TSharedPtr<IPropertyHandle> YDriveProperty = LinearDriveProperty->GetChildHandle(GET_MEMBER_NAME_CHECKED(FLinearDriveConstraint, YDrive));
+	TSharedPtr<IPropertyHandle> ZDriveProperty = LinearDriveProperty->GetChildHandle(GET_MEMBER_NAME_CHECKED(FLinearDriveConstraint, ZDrive));
+
+	LinearXPositionDriveProperty = XDriveProperty->GetChildHandle(GET_MEMBER_NAME_CHECKED(FConstraintDrive, bEnablePositionDrive));
+	LinearYPositionDriveProperty = YDriveProperty->GetChildHandle(GET_MEMBER_NAME_CHECKED(FConstraintDrive, bEnablePositionDrive));
+	LinearZPositionDriveProperty = ZDriveProperty->GetChildHandle(GET_MEMBER_NAME_CHECKED(FConstraintDrive, bEnablePositionDrive));
 
-	TSharedRef<SWidget> LinearPositionXWidget = JoinPropertyWidgets(LinearPositionTargetProperty, FName("X"), ConstraintInstance, FName("bLinearXPositionDrive"), LinearXPositionDriveProperty, EPropertyType::LinearXPositionDrive);
-	TSharedRef<SWidget> LinearPositionYWidget = JoinPropertyWidgets(LinearPositionTargetProperty, FName("Y"), ConstraintInstance, FName("bLinearYPositionDrive"), LinearYPositionDriveProperty, EPropertyType::LinearYPositionDrive);
-	TSharedRef<SWidget> LinearPositionZWidget = JoinPropertyWidgets(LinearPositionTargetProperty, FName("Z"), ConstraintInstance, FName("bLinearZPositionDrive"), LinearZPositionDriveProperty, EPropertyType::LinearZPositionDrive);
+	TSharedRef<SWidget> LinearPositionXWidget = ConstraintDetails::JoinPropertyWidgets(LinearPositionTargetProperty, FName("X"), XDriveProperty, GET_MEMBER_NAME_CHECKED(FConstraintDrive, bEnablePositionDrive), LinearXPositionDriveProperty);
+	TSharedRef<SWidget> LinearPositionYWidget = ConstraintDetails::JoinPropertyWidgets(LinearPositionTargetProperty, FName("Y"), YDriveProperty, GET_MEMBER_NAME_CHECKED(FConstraintDrive, bEnablePositionDrive), LinearYPositionDriveProperty);
+	TSharedRef<SWidget> LinearPositionZWidget = ConstraintDetails::JoinPropertyWidgets(LinearPositionTargetProperty, FName("Z"), ZDriveProperty, GET_MEMBER_NAME_CHECKED(FConstraintDrive, bEnablePositionDrive), LinearZPositionDriveProperty);
 
 
-	FDetailWidgetRow& LinearPositionTargetWidget = PositionGroup.AddPropertyRow(LinearPositionTargetProperty).CustomWidget(true);
-	LinearPositionTargetWidget.NameContent()
+	FDetailWidgetRow& LinearPositionTargetWidget = PositionGroup.HeaderProperty(LinearPositionTargetProperty).CustomWidget()
+	.NameContent()
 	[
 			LinearPositionTargetProperty->CreatePropertyNameWidget()
-	];
-
-	LinearPositionTargetWidget.ValueContent()
+	]
+	.ValueContent()
 	.MinDesiredWidth(125 * 3 + 18 * 3)
 	.MaxDesiredWidth(125 * 3 + 18 * 3)
 	[
@@ -359,82 +476,353 @@ void FPhysicsConstraintComponentDetails::AddLinearDrive(IDetailLayoutBuilder& De
 		]
 	];
 
-	PositionGroup.AddPropertyRow(ConstraintInstance->GetChildHandle(GET_MEMBER_NAME_CHECKED(FConstraintInstance, LinearDriveSpring)).ToSharedRef())
-		.IsEnabled(TAttribute<bool>::Create(TAttribute<bool>::FGetter::CreateSP(this, &FPhysicsConstraintComponentDetails::IsPropertyEnabled, EPropertyType::LinearPositionDrive)));
+	TSharedPtr<IPropertyHandle> StiffnessXProperty = XDriveProperty->GetChildHandle(GET_MEMBER_NAME_CHECKED(FConstraintDrive, Stiffness));
+	TSharedRef<SWidget> StiffnessWidget = ConstraintDetails::CreateTriFloatWidget(StiffnessXProperty, YDriveProperty->GetChildHandle(GET_MEMBER_NAME_CHECKED(FConstraintDrive, Stiffness)), ZDriveProperty->GetChildHandle(GET_MEMBER_NAME_CHECKED(FConstraintDrive, Stiffness)), LOCTEXT("EditStrength", "Edit Strength"));
+	StiffnessWidget->SetEnabled(TAttribute<bool>::Create(TAttribute<bool>::FGetter::CreateSP(this, &FPhysicsConstraintComponentDetails::IsPropertyEnabled, EPropertyType::LinearPositionDrive)));
+
+	PositionGroup.AddWidgetRow()
+	.NameContent()
+	[
+		SNew(STextBlock)
+		.Text(LOCTEXT("Strength", "Strength"))
+		.Font(IDetailLayoutBuilder::GetDetailFont())
+	]
+	.ValueContent()
+	[
+		StiffnessWidget
+	];
+	
+	// VELOCITY
 
-	//velocity
-	LinearVelocityDriveProperty = ConstraintInstance->GetChildHandle(GET_MEMBER_NAME_CHECKED(FConstraintInstance, bLinearVelocityDrive));
-	TSharedRef<IPropertyHandle> LinearVelocityTargetProperty = ConstraintInstance->GetChildHandle(GET_MEMBER_NAME_CHECKED(FConstraintInstance, LinearVelocityTarget)).ToSharedRef();
+	IDetailGroup& VelocityGroup = LinearMotorCat.AddGroup("Linear Velocity Drive", LOCTEXT("LinearVelocityDrive", "Linear Velocity Drive"), false, true);
 
-	IDetailGroup& VelocityGroup = LinearMotorCat.AddGroup("Linear Velocity Drive", FText::GetEmpty());
-	VelocityGroup.HeaderProperty(LinearVelocityDriveProperty.ToSharedRef());
+	TSharedRef<IPropertyHandle> LinearVelocityTargetProperty = LinearDriveProperty->GetChildHandle(GET_MEMBER_NAME_CHECKED(FLinearDriveConstraint, VelocityTarget)).ToSharedRef();
 
-	TSharedRef<SWidget> LinearVelocityXWidget = JoinPropertyWidgets(LinearVelocityTargetProperty, FName("X"), ConstraintInstance, FName("bLinearXVelocityDrive"), LinearXVelocityDriveProperty, EPropertyType::LinearXVelocityDrive);
-	TSharedRef<SWidget> LinearVelocityYWidget = JoinPropertyWidgets(LinearVelocityTargetProperty, FName("Y"), ConstraintInstance, FName("bLinearYVelocityDrive"), LinearYVelocityDriveProperty, EPropertyType::LinearYVelocityDrive);
-	TSharedRef<SWidget> LinearVelocityZWidget = JoinPropertyWidgets(LinearVelocityTargetProperty, FName("Z"), ConstraintInstance, FName("bLinearZVelocityDrive"), LinearZVelocityDriveProperty, EPropertyType::LinearZVelocityDrive);
+	LinearXVelocityDriveProperty = XDriveProperty->GetChildHandle(GET_MEMBER_NAME_CHECKED(FConstraintDrive, bEnableVelocityDrive));
+	LinearYVelocityDriveProperty = YDriveProperty->GetChildHandle(GET_MEMBER_NAME_CHECKED(FConstraintDrive, bEnableVelocityDrive));
+	LinearZVelocityDriveProperty = ZDriveProperty->GetChildHandle(GET_MEMBER_NAME_CHECKED(FConstraintDrive, bEnableVelocityDrive));
 
+	TSharedRef<SWidget> LinearVelocityXWidget = ConstraintDetails::JoinPropertyWidgets(LinearVelocityTargetProperty, FName("X"), XDriveProperty, GET_MEMBER_NAME_CHECKED(FConstraintDrive, bEnableVelocityDrive), LinearXVelocityDriveProperty);
+	TSharedRef<SWidget> LinearVelocityYWidget = ConstraintDetails::JoinPropertyWidgets(LinearVelocityTargetProperty, FName("Y"), YDriveProperty, GET_MEMBER_NAME_CHECKED(FConstraintDrive, bEnableVelocityDrive), LinearYVelocityDriveProperty);
+	TSharedRef<SWidget> LinearVelocityZWidget = ConstraintDetails::JoinPropertyWidgets(LinearVelocityTargetProperty, FName("Z"), ZDriveProperty, GET_MEMBER_NAME_CHECKED(FConstraintDrive, bEnableVelocityDrive), LinearZVelocityDriveProperty);
 
-	FDetailWidgetRow& LinearVelocityTargetWidget = VelocityGroup.AddPropertyRow(LinearVelocityTargetProperty).CustomWidget(true);
+	FDetailWidgetRow& LinearVelocityTargetWidget = VelocityGroup.HeaderProperty(LinearVelocityTargetProperty).CustomWidget(true);
 	LinearVelocityTargetWidget.NameContent()
+		[
+			LinearVelocityTargetProperty->CreatePropertyNameWidget()
+		];
+
+	LinearVelocityTargetWidget.ValueContent()
+		.MinDesiredWidth(125 * 3 + 18 * 3)
+		.MaxDesiredWidth(125 * 3 + 18 * 3)
+		[
+			SNew(SHorizontalBox)
+			+ SHorizontalBox::Slot()
+		[
+			LinearVelocityXWidget
+		]
+	+ SHorizontalBox::Slot()
+		.Padding(5, 0, 0, 0)
+		[
+			LinearVelocityYWidget
+		]
+
+	+ SHorizontalBox::Slot()
+		.Padding(5, 0, 0, 0)
+		[
+			LinearVelocityZWidget
+		]
+		];
+
+	TSharedPtr<IPropertyHandle> XDampingProperty = XDriveProperty->GetChildHandle(GET_MEMBER_NAME_CHECKED(FConstraintDrive, Damping));
+	TSharedRef<SWidget> DampingWidget = ConstraintDetails::CreateTriFloatWidget(XDampingProperty, YDriveProperty->GetChildHandle(GET_MEMBER_NAME_CHECKED(FConstraintDrive, Damping)), ZDriveProperty->GetChildHandle(GET_MEMBER_NAME_CHECKED(FConstraintDrive, Damping)), LOCTEXT("EditStrength", "Edit Strength"));
+	DampingWidget->SetEnabled(TAttribute<bool>::Create(TAttribute<bool>::FGetter::CreateSP(this, &FPhysicsConstraintComponentDetails::IsPropertyEnabled, EPropertyType::LinearVelocityDrive)));
+
+	VelocityGroup.AddWidgetRow()
+	.NameContent()
+	[
+		SNew(STextBlock)
+		.Text(LOCTEXT("Strength", "Strength"))
+		.Font(IDetailLayoutBuilder::GetDetailFont())
+	]
+	.ValueContent()
 	[
-		LinearVelocityTargetProperty->CreatePropertyNameWidget()
+		DampingWidget
 	];
 
-	LinearVelocityTargetWidget.ValueContent()
+	// max force limit
+	TSharedPtr<IPropertyHandle> MaxForceProperty = XDriveProperty->GetChildHandle(GET_MEMBER_NAME_CHECKED(FConstraintDrive, MaxForce));
+	TSharedRef<SWidget> MaxForceWidget = ConstraintDetails::CreateTriFloatWidget(MaxForceProperty, YDriveProperty->GetChildHandle(GET_MEMBER_NAME_CHECKED(FConstraintDrive, MaxForce)), ZDriveProperty->GetChildHandle(GET_MEMBER_NAME_CHECKED(FConstraintDrive, MaxForce)), LOCTEXT("EditMaxForce", "Edit Max Force"));
+	MaxForceWidget->SetEnabled(TAttribute<bool>::Create(TAttribute<bool>::FGetter::CreateSP(this, &FPhysicsConstraintComponentDetails::IsPropertyEnabled, EPropertyType::LinearDrive)));
+
+	LinearMotorCat.AddCustomRow(LOCTEXT("MaxForce", "Max Force"), true)
+	.NameContent()
+	[
+		MaxForceProperty->CreatePropertyNameWidget()
+	]
+	.ValueContent()
+	[
+		MaxForceWidget
+	];
+
+}
+
+void FPhysicsConstraintComponentDetails::AddAngularDrive(IDetailLayoutBuilder& DetailBuilder, TSharedPtr<IPropertyHandle> ConstraintInstance, TSharedPtr<IPropertyHandle> ProfilePropertiesProperty)
+{
+	IDetailCategoryBuilder& AngularMotorCat = DetailBuilder.EditCategory("AngularMotor");
+
+	TSharedPtr<IPropertyHandle> AngularDriveProperty = ProfilePropertiesProperty->GetChildHandle(GET_MEMBER_NAME_CHECKED(FConstraintProfileProperties, AngularDrive));
+	TSharedPtr<IPropertyHandle> AngularDriveModeProperty = AngularDriveProperty->GetChildHandle(GET_MEMBER_NAME_CHECKED(FAngularDriveConstraint, AngularDriveMode));
+
+	TSharedPtr<IPropertyHandle> SlerpDriveProperty = AngularDriveProperty->GetChildHandle(GET_MEMBER_NAME_CHECKED(FAngularDriveConstraint, SlerpDrive));
+	TSharedPtr<IPropertyHandle> SwingDriveProperty = AngularDriveProperty->GetChildHandle(GET_MEMBER_NAME_CHECKED(FAngularDriveConstraint, SwingDrive));
+	TSharedPtr<IPropertyHandle> TwistDriveProperty = AngularDriveProperty->GetChildHandle(GET_MEMBER_NAME_CHECKED(FAngularDriveConstraint, TwistDrive));
+
+	TSharedPtr<IPropertyHandle> SlerpPositionDriveProperty = SlerpDriveProperty->GetChildHandle(GET_MEMBER_NAME_CHECKED(FConstraintDrive, bEnablePositionDrive));
+	TSharedPtr<IPropertyHandle> SlerpVelocityDriveProperty = SlerpDriveProperty->GetChildHandle(GET_MEMBER_NAME_CHECKED(FConstraintDrive, bEnableVelocityDrive));
+	TSharedPtr<IPropertyHandle> SwingPositionDriveProperty = SwingDriveProperty->GetChildHandle(GET_MEMBER_NAME_CHECKED(FConstraintDrive, bEnablePositionDrive));
+	TSharedPtr<IPropertyHandle> SwingVelocityDriveProperty = SwingDriveProperty->GetChildHandle(GET_MEMBER_NAME_CHECKED(FConstraintDrive, bEnableVelocityDrive));
+	TSharedPtr<IPropertyHandle> TwistPositionDriveProperty = TwistDriveProperty->GetChildHandle(GET_MEMBER_NAME_CHECKED(FConstraintDrive, bEnablePositionDrive));
+	TSharedPtr<IPropertyHandle> TwistVelocityDriveProperty = TwistDriveProperty->GetChildHandle(GET_MEMBER_NAME_CHECKED(FConstraintDrive, bEnableVelocityDrive));
+		
+	auto IsAngularMode = [AngularDriveModeProperty](EAngularDriveMode::Type CheckMode)
+	{
+		uint8 DriveMode;
+		if (AngularDriveModeProperty->GetValue(DriveMode))
+		{
+			return DriveMode == CheckMode;
+		}
+		return false;
+	};
+
+	auto EligibleForSLERP = [this, IsAngularMode]()
+	{
+		return IsAngularMode(EAngularDriveMode::SLERP) && !ConstraintDetails::IsAngularPropertyEqual(AngularSwing1MotionProperty, ACM_Locked) && !ConstraintDetails::IsAngularPropertyEqual(AngularSwing2MotionProperty, ACM_Locked) && !ConstraintDetails::IsAngularPropertyEqual(AngularTwistMotionProperty, ACM_Locked);
+	};
+
+	auto EligibleForTwistAndSwing = [IsAngularMode]()
+	{
+		return IsAngularMode(EAngularDriveMode::TwistAndSwing);
+	};
+
+	auto OrientationEnabled = [EligibleForSLERP, EligibleForTwistAndSwing, TwistPositionDriveProperty, SwingPositionDriveProperty, SlerpPositionDriveProperty]()
+	{
+		if(EligibleForSLERP())
+		{
+			return ConstraintDetails::GetBoolProperty(SlerpPositionDriveProperty);
+		} else if(EligibleForTwistAndSwing())
+		{
+			return ConstraintDetails::GetBoolProperty(TwistPositionDriveProperty) || ConstraintDetails::GetBoolProperty(SwingPositionDriveProperty);
+		}
+
+		return false;
+	};
+
+	auto VelocityEnabled = [EligibleForSLERP, EligibleForTwistAndSwing, TwistVelocityDriveProperty, SwingVelocityDriveProperty, SlerpVelocityDriveProperty]()
+	{
+		if (EligibleForSLERP())
+		{
+			return ConstraintDetails::GetBoolProperty(SlerpVelocityDriveProperty);
+		}
+		else if (EligibleForTwistAndSwing())
+		{
+			return ConstraintDetails::GetBoolProperty(TwistVelocityDriveProperty) || ConstraintDetails::GetBoolProperty(SwingVelocityDriveProperty);
+		}
+
+		return false;
+	};
+
+	auto VelocityOrOrientationEnabled = [VelocityEnabled, OrientationEnabled]()
+	{
+		return VelocityEnabled() || OrientationEnabled();
+	};
+	
+	AngularMotorCat.AddProperty(AngularDriveModeProperty);
+
+	IDetailGroup& OrientationGroup = AngularMotorCat.AddGroup("Orientation Drive", LOCTEXT("OrientrationDrive", "Orientation Drive"), false, true);
+	OrientationGroup.HeaderProperty(AngularDriveProperty->GetChildHandle(GET_MEMBER_NAME_CHECKED(FAngularDriveConstraint, OrientationTarget)).ToSharedRef()).DisplayName(LOCTEXT("TargetOrientation", "Target Orientation")).IsEnabled(TAttribute<bool>::Create(TAttribute<bool>::FGetter::CreateLambda(OrientationEnabled)));
+
+
+	TSharedRef<SWidget> SlerpPositionWidget = SlerpPositionDriveProperty->CreatePropertyValueWidget();
+	TSharedRef<SWidget> SlerpVelocityWidget = SlerpVelocityDriveProperty->CreatePropertyValueWidget();
+	SlerpPositionWidget->SetEnabled(TAttribute<bool>::Create(TAttribute<bool>::FGetter::CreateLambda(EligibleForSLERP)));
+	SlerpVelocityWidget->SetEnabled(TAttribute<bool>::Create(TAttribute<bool>::FGetter::CreateLambda(EligibleForSLERP)));
+
+	TSharedRef<SWidget> TwistPositionWidget = TwistPositionDriveProperty->CreatePropertyValueWidget();
+	TSharedRef<SWidget> TwistVelocityWidget = TwistVelocityDriveProperty->CreatePropertyValueWidget();
+	TwistPositionWidget->SetEnabled(TAttribute<bool>::Create(TAttribute<bool>::FGetter::CreateLambda(EligibleForTwistAndSwing)));
+	TwistVelocityWidget->SetEnabled(TAttribute<bool>::Create(TAttribute<bool>::FGetter::CreateLambda(EligibleForTwistAndSwing)));
+
+	TSharedRef<SWidget> SwingPositionWidget = SwingPositionDriveProperty->CreatePropertyValueWidget();
+	TSharedRef<SWidget> SwingVelocityWidget = SwingVelocityDriveProperty->CreatePropertyValueWidget();
+	SwingPositionWidget->SetEnabled(TAttribute<bool>::Create(TAttribute<bool>::FGetter::CreateLambda(EligibleForTwistAndSwing)));
+	SwingVelocityWidget->SetEnabled(TAttribute<bool>::Create(TAttribute<bool>::FGetter::CreateLambda(EligibleForTwistAndSwing)));
+
+	OrientationGroup.AddWidgetRow()
+	.NameContent()
+	[
+		SNew(STextBlock)
+		.Text(LOCTEXT("TwistSwingSlerpDrive", "Drives"))
+		.Font(IDetailLayoutBuilder::GetDetailFont())
+	]
+	.ValueContent()
 	.MinDesiredWidth(125 * 3 + 18 * 3)
 	.MaxDesiredWidth(125 * 3 + 18 * 3)
 	[
 		SNew(SHorizontalBox)
 		+ SHorizontalBox::Slot()
 		[
-			LinearVelocityXWidget
+			SNew(SHorizontalBox)
+			+ SHorizontalBox::Slot()
+			[
+				SlerpDriveProperty->CreatePropertyNameWidget()
+			]
+			+ SHorizontalBox::Slot()
+			[
+				SlerpPositionWidget
+			]
 		]
 		+ SHorizontalBox::Slot()
-		.Padding(5, 0, 0, 0)
 		[
-			LinearVelocityYWidget
+			SNew(SHorizontalBox)
+			+ SHorizontalBox::Slot()
+			[
+				TwistDriveProperty->CreatePropertyNameWidget()
+			]
+			+ SHorizontalBox::Slot()
+			[
+				TwistPositionWidget
+			]
+		]
+		+ SHorizontalBox::Slot()
+		[
+			SNew(SHorizontalBox)
+			+ SHorizontalBox::Slot()
+			[
+				SwingDriveProperty->CreatePropertyNameWidget()
+			]
+			+ SHorizontalBox::Slot()
+			[
+				SwingPositionWidget
+			]
 		]
+	];
+
+	TSharedPtr<IPropertyHandle> StiffnessSlerpProperty = SlerpDriveProperty->GetChildHandle(GET_MEMBER_NAME_CHECKED(FConstraintDrive, Stiffness));
+	TSharedRef<SWidget> OrientationStrengthWidget = ConstraintDetails::CreateTriFloatWidget(StiffnessSlerpProperty, TwistDriveProperty->GetChildHandle(GET_MEMBER_NAME_CHECKED(FConstraintDrive, Stiffness)), SwingDriveProperty->GetChildHandle(GET_MEMBER_NAME_CHECKED(FConstraintDrive, Stiffness)), LOCTEXT("EditStrength", "Edit Strength"));
+	OrientationStrengthWidget->SetEnabled(TAttribute<bool>::Create(TAttribute<bool>::FGetter::CreateLambda(OrientationEnabled)));
+
+	OrientationGroup.AddWidgetRow()
+	.NameContent()
+	[
+		SNew(STextBlock)
+		.Text(LOCTEXT("Strength", "Strength"))
+		.Font(IDetailLayoutBuilder::GetDetailFont())
+	]
+	.ValueContent()
+	[
+		OrientationStrengthWidget
+	];
 
+	IDetailGroup& AngularVelocityGroup = AngularMotorCat.AddGroup("Velocity Drive", LOCTEXT("VelocityDrive", "Velocity Drive"), false, true);
+	AngularVelocityGroup.HeaderProperty(AngularDriveProperty->GetChildHandle(GET_MEMBER_NAME_CHECKED(FAngularDriveConstraint, AngularVelocityTarget)).ToSharedRef()).DisplayName(LOCTEXT("TargetVelocity", "Target Velocity")).IsEnabled(TAttribute<bool>::Create(TAttribute<bool>::FGetter::CreateLambda(VelocityEnabled)));
+
+	AngularVelocityGroup.AddWidgetRow()
+	.NameContent()
+	[
+		SNew(STextBlock)
+		.Text(LOCTEXT("TwistSwingSlerpDrive", "Drives"))
+		.Font(IDetailLayoutBuilder::GetDetailFont())
+	]
+	.ValueContent()
+	.MinDesiredWidth(125 * 3 + 18 * 3)
+	.MaxDesiredWidth(125 * 3 + 18 * 3)
+	[
+		SNew(SHorizontalBox)
 		+ SHorizontalBox::Slot()
-		.Padding(5, 0, 0, 0)
 		[
-			LinearVelocityZWidget
+			SNew(SHorizontalBox)
+			+ SHorizontalBox::Slot()
+			[
+				SlerpDriveProperty->CreatePropertyNameWidget()
+			]
+			+ SHorizontalBox::Slot()
+			[
+				SlerpVelocityWidget
+			]
+		]
+		+ SHorizontalBox::Slot()
+		[
+			SNew(SHorizontalBox)
+			+ SHorizontalBox::Slot()
+			[
+				TwistDriveProperty->CreatePropertyNameWidget()
+			]
+			+ SHorizontalBox::Slot()
+			[
+				TwistVelocityWidget
+			]
+		]
+		+ SHorizontalBox::Slot()
+		[
+			SNew(SHorizontalBox)
+			+ SHorizontalBox::Slot()
+			[
+				SwingDriveProperty->CreatePropertyNameWidget()
+			]
+			+ SHorizontalBox::Slot()
+			[
+				SwingVelocityWidget
+			]
 		]
 	];
 
-	VelocityGroup.AddPropertyRow(ConstraintInstance->GetChildHandle(GET_MEMBER_NAME_CHECKED(FConstraintInstance, LinearDriveDamping)).ToSharedRef())
-		.IsEnabled(TAttribute<bool>::Create(TAttribute<bool>::FGetter::CreateSP(this, &FPhysicsConstraintComponentDetails::IsPropertyEnabled, EPropertyType::LinearVelocityDrive)));
+	TSharedPtr<IPropertyHandle> DampingSlerpProperty = SlerpDriveProperty->GetChildHandle(GET_MEMBER_NAME_CHECKED(FConstraintDrive, Damping));
+	TSharedRef<SWidget> DampingSlerpWidget = ConstraintDetails::CreateTriFloatWidget(DampingSlerpProperty, TwistDriveProperty->GetChildHandle(GET_MEMBER_NAME_CHECKED(FConstraintDrive, Damping)), SwingDriveProperty->GetChildHandle(GET_MEMBER_NAME_CHECKED(FConstraintDrive, Damping)), LOCTEXT("EditStrength", "Edit Strength"));
+	DampingSlerpWidget->SetEnabled(TAttribute<bool>::Create(TAttribute<bool>::FGetter::CreateLambda(VelocityEnabled)));
+	AngularVelocityGroup.AddWidgetRow()
+	.NameContent()
+	[
+		SNew(STextBlock)
+		.Text(LOCTEXT("Strength", "Strength"))
+		.Font(IDetailLayoutBuilder::GetDetailFont())
+	]
+	.ValueContent()
+	[
+		DampingSlerpWidget
+	];
+	
+	// max force limit
+	TSharedPtr<IPropertyHandle> MaxForcePropertySlerp = SlerpDriveProperty->GetChildHandle(GET_MEMBER_NAME_CHECKED(FConstraintDrive, MaxForce));
+	TSharedRef<SWidget> MaxForceWidget = ConstraintDetails::CreateTriFloatWidget(MaxForcePropertySlerp, TwistDriveProperty->GetChildHandle(GET_MEMBER_NAME_CHECKED(FConstraintDrive, MaxForce)), SwingDriveProperty->GetChildHandle(GET_MEMBER_NAME_CHECKED(FConstraintDrive, MaxForce)), LOCTEXT("EditMaxForce", "Edit Max Force"));
+	MaxForceWidget->SetEnabled(TAttribute<bool>::Create(TAttribute<bool>::FGetter::CreateLambda(VelocityOrOrientationEnabled)));
+
+	AngularMotorCat.AddCustomRow(LOCTEXT("MaxForce", "Max Force"), true)
+	.NameContent()
+	[
+		MaxForcePropertySlerp->CreatePropertyNameWidget()
+	]
+	.ValueContent()
+	[
+		MaxForceWidget
+	];
 
-	LinearMotorCat.AddProperty(ConstraintInstance->GetChildHandle(GET_MEMBER_NAME_CHECKED(FConstraintInstance, LinearDriveForceLimit)).ToSharedRef())
-		.IsEnabled(TAttribute<bool>::Create(TAttribute<bool>::FGetter::CreateSP(this, &FPhysicsConstraintComponentDetails::IsPropertyEnabled, EPropertyType::LinearAnyDrive)));
+	
 }
 
-void FPhysicsConstraintComponentDetails::AddAngularDrive(IDetailLayoutBuilder& DetailBuilder, TSharedPtr<IPropertyHandle> ConstraintInstance)
+void FPhysicsConstraintComponentDetails::AddProfiles(IDetailLayoutBuilder& DetailBuilder)
 {
-	AngularPositionDriveProperty = ConstraintInstance->GetChildHandle(GET_MEMBER_NAME_CHECKED(FConstraintInstance, bAngularOrientationDrive));
-	AngularVelocityDriveProperty = ConstraintInstance->GetChildHandle(GET_MEMBER_NAME_CHECKED(FConstraintInstance, bAngularVelocityDrive));
+	IDetailCategoryBuilder& ConstraintBehavCat = DetailBuilder.EditCategory("Constraint Behavior");	//Create this category first so it's at the top
 
-	IDetailCategoryBuilder& AngularMotorCat = DetailBuilder.EditCategory("AngularMotor");
+	TSharedPtr<IPropertyHandle> DefaultInstance = DetailBuilder.GetProperty(GET_MEMBER_NAME_CHECKED(UPhysicsConstraintTemplate, DefaultInstance));
+	TSharedPtr<IPropertyHandle> DefaultProfileName = DetailBuilder.GetProperty(GET_MEMBER_NAME_CHECKED(UPhysicsConstraintTemplate, DefaultProfileName));
+	TSharedPtr<IPropertyHandle> ProfileHandles = DetailBuilder.GetProperty(GET_MEMBER_NAME_CHECKED(UPhysicsConstraintTemplate, ProfileHandles));
+
+	ConstraintBehavCat.AddProperty(DefaultProfileName);
+	ConstraintBehavCat.AddProperty(ProfileHandles);
 
-	IDetailGroup& PositionGroup = AngularMotorCat.AddGroup("Angular Orientation Drive", FText::GetEmpty());
-	PositionGroup.HeaderProperty(AngularPositionDriveProperty.ToSharedRef());
-
-	PositionGroup.AddPropertyRow(ConstraintInstance->GetChildHandle(GET_MEMBER_NAME_CHECKED(FConstraintInstance, AngularOrientationTarget)).ToSharedRef())
-		.IsEnabled(TAttribute<bool>::Create(TAttribute<bool>::FGetter::CreateSP(this, &FPhysicsConstraintComponentDetails::IsPropertyEnabled, EPropertyType::AngularPositionDrive)));
-
-	PositionGroup.AddPropertyRow(ConstraintInstance->GetChildHandle(GET_MEMBER_NAME_CHECKED(FConstraintInstance, AngularDriveSpring)).ToSharedRef())
-		.IsEnabled(TAttribute<bool>::Create(TAttribute<bool>::FGetter::CreateSP(this, &FPhysicsConstraintComponentDetails::IsPropertyEnabled, EPropertyType::AngularAnyDrive)));
-
-	IDetailGroup& VelocityGroup = AngularMotorCat.AddGroup("Angular Velocity Drive", FText::GetEmpty());
-	VelocityGroup.HeaderProperty(AngularVelocityDriveProperty.ToSharedRef());
-	VelocityGroup.AddPropertyRow(ConstraintInstance->GetChildHandle(GET_MEMBER_NAME_CHECKED(FConstraintInstance, AngularVelocityTarget)).ToSharedRef())
-		.IsEnabled(TAttribute<bool>::Create(TAttribute<bool>::FGetter::CreateSP(this, &FPhysicsConstraintComponentDetails::IsPropertyEnabled, EPropertyType::AngularVelocityDrive)));
-	VelocityGroup.AddPropertyRow(ConstraintInstance->GetChildHandle(GET_MEMBER_NAME_CHECKED(FConstraintInstance, AngularDriveDamping)).ToSharedRef())
-		.IsEnabled(TAttribute<bool>::Create(TAttribute<bool>::FGetter::CreateSP(this, &FPhysicsConstraintComponentDetails::IsPropertyEnabled, EPropertyType::AngularAnyDrive)));
-	AngularMotorCat.AddProperty(ConstraintInstance->GetChildHandle(GET_MEMBER_NAME_CHECKED(FConstraintInstance, AngularDriveForceLimit)).ToSharedRef())
-		.IsEnabled(TAttribute<bool>::Create(TAttribute<bool>::FGetter::CreateSP(this, &FPhysicsConstraintComponentDetails::IsPropertyEnabled, EPropertyType::AngularAnyDrive)));
-	AngularMotorCat.AddProperty(ConstraintInstance->GetChildHandle(GET_MEMBER_NAME_CHECKED(FConstraintInstance, AngularDriveMode)).ToSharedRef())
-		.IsEnabled(TAttribute<bool>::Create(TAttribute<bool>::FGetter::CreateSP(this, &FPhysicsConstraintComponentDetails::IsPropertyEnabled, EPropertyType::AngularAnyDrive)));
+	DefaultInstance->GetChildHandle(GET_MEMBER_NAME_CHECKED(FConstraintInstance, ProfileInstance))->MarkHiddenByCustomization();
+	
 }
 
 void FPhysicsConstraintComponentDetails::CustomizeDetails( IDetailLayoutBuilder& DetailBuilder )
@@ -467,110 +855,50 @@ void FPhysicsConstraintComponentDetails::CustomizeDetails( IDetailLayoutBuilder&
 		}
 	}
 
-	AddConstraintPresets(DetailBuilder, OwningConstraintActor);
-	AddLinearLimits(DetailBuilder, ConstraintInstance);
-	AddAngularLimits(DetailBuilder, ConstraintInstance);
-	AddLinearDrive(DetailBuilder, ConstraintInstance);
-	AddAngularDrive(DetailBuilder, ConstraintInstance);
+	DetailBuilder.EditCategory("Constraint");	//Create this category first so it's at the top
+	DetailBuilder.EditCategory("Constraint Behavior");	//Create this category first so it's at the top
+	
+	TSharedPtr<IPropertyHandle> ProfileInstance = ConstraintInstance->GetChildHandle(GET_MEMBER_NAME_CHECKED(FConstraintInstance, ProfileInstance));
+	AddLinearLimits(DetailBuilder, ConstraintInstance, ProfileInstance);
+	AddAngularLimits(DetailBuilder, ConstraintInstance, ProfileInstance);
+	AddLinearDrive(DetailBuilder, ConstraintInstance, ProfileInstance);
+	AddAngularDrive(DetailBuilder, ConstraintInstance, ProfileInstance);
+	
+	AddConstraintBehaviorProperties(DetailBuilder, ConstraintInstance, ProfileInstance);	//Now we've added all the complex UI, just dump the rest into Constraint category
+
+	if (bInPhat)
+	{
+		AddProfiles(DetailBuilder);
+	}
 }
 
+
 bool FPhysicsConstraintComponentDetails::IsPropertyEnabled( EPropertyType::Type Type ) const
 {
 	bool bIsVisible = false;
 	switch (Type)
 	{
-		case EPropertyType::LinearXPositionDrive:	return GetBoolProperty(LinearXPositionDriveProperty);
-		case EPropertyType::LinearYPositionDrive:	return GetBoolProperty(LinearYPositionDriveProperty);
-		case EPropertyType::LinearZPositionDrive:	return GetBoolProperty(LinearZPositionDriveProperty);
-
-		case EPropertyType::LinearXVelocityDrive:	return GetBoolProperty(LinearXVelocityDriveProperty);
-		case EPropertyType::LinearYVelocityDrive:	return GetBoolProperty(LinearYVelocityDriveProperty);
-		case EPropertyType::LinearZVelocityDrive:	return GetBoolProperty(LinearZVelocityDriveProperty);
-
-		case EPropertyType::LinearPositionDrive:	return GetBoolProperty(LinearPositionDriveProperty);
-		case EPropertyType::LinearVelocityDrive:	return GetBoolProperty(LinearVelocityDriveProperty);
-		case EPropertyType::LinearAnyDrive:			return GetBoolProperty(LinearPositionDriveProperty) || GetBoolProperty(LinearVelocityDriveProperty);
-		case EPropertyType::AngularPositionDrive:	return GetBoolProperty(AngularPositionDriveProperty);
-		case EPropertyType::AngularVelocityDrive:	return GetBoolProperty(AngularVelocityDriveProperty);
-		case EPropertyType::AngularAnyDrive:		return GetBoolProperty(AngularPositionDriveProperty) || GetBoolProperty(AngularVelocityDriveProperty);
-		case EPropertyType::LinearLimit:			return IsLinearMotionLimited();
-		case EPropertyType::AngularSwing1Limit:		return IsAngularPropertyLimited(AngularSwing1MotionProperty);
-		case EPropertyType::AngularSwing2Limit:		return IsAngularPropertyLimited(AngularSwing2MotionProperty);
-		case EPropertyType::AngularSwingLimit:		return IsAngularPropertyLimited(AngularSwing1MotionProperty) || IsAngularPropertyLimited(AngularSwing2MotionProperty);
-		case EPropertyType::AngularTwistLimit:		return IsAngularPropertyLimited(AngularTwistMotionProperty);
-		case EPropertyType::AngularAnyLimit:		return IsAngularPropertyLimited(AngularSwing1MotionProperty) || IsAngularPropertyLimited(AngularSwing2MotionProperty) || 
-																 IsAngularPropertyLimited(AngularTwistMotionProperty);
+		case EPropertyType::LinearXPositionDrive:	return ConstraintDetails::GetBoolProperty(LinearXPositionDriveProperty);
+		case EPropertyType::LinearYPositionDrive:	return ConstraintDetails::GetBoolProperty(LinearYPositionDriveProperty);
+		case EPropertyType::LinearZPositionDrive:	return ConstraintDetails::GetBoolProperty(LinearZPositionDriveProperty);
+
+		case EPropertyType::LinearXVelocityDrive:	return ConstraintDetails::GetBoolProperty(LinearXVelocityDriveProperty);
+		case EPropertyType::LinearYVelocityDrive:	return ConstraintDetails::GetBoolProperty(LinearYVelocityDriveProperty);
+		case EPropertyType::LinearZVelocityDrive:	return ConstraintDetails::GetBoolProperty(LinearZVelocityDriveProperty);
+		case EPropertyType::LinearPositionDrive:	return ConstraintDetails::GetBoolProperty(LinearXPositionDriveProperty) || ConstraintDetails::GetBoolProperty(LinearYPositionDriveProperty) || ConstraintDetails::GetBoolProperty(LinearZPositionDriveProperty);
+		case EPropertyType::LinearVelocityDrive:	return ConstraintDetails::GetBoolProperty(LinearXVelocityDriveProperty) || ConstraintDetails::GetBoolProperty(LinearYVelocityDriveProperty) || ConstraintDetails::GetBoolProperty(LinearZVelocityDriveProperty);
+		case EPropertyType::LinearDrive:			return ConstraintDetails::GetBoolProperty(LinearXPositionDriveProperty) || ConstraintDetails::GetBoolProperty(LinearYPositionDriveProperty) || ConstraintDetails::GetBoolProperty(LinearZPositionDriveProperty)
+															|| ConstraintDetails::GetBoolProperty(LinearXVelocityDriveProperty) || ConstraintDetails::GetBoolProperty(LinearYVelocityDriveProperty) || ConstraintDetails::GetBoolProperty(LinearZVelocityDriveProperty);
+		case EPropertyType::AngularSwing1Limit:		return ConstraintDetails::IsAngularPropertyEqual(AngularSwing1MotionProperty, ACM_Limited);
+		case EPropertyType::AngularSwing2Limit:		return ConstraintDetails::IsAngularPropertyEqual(AngularSwing2MotionProperty, ACM_Limited);
+		case EPropertyType::AngularSwingLimit:		return ConstraintDetails::IsAngularPropertyEqual(AngularSwing1MotionProperty, ACM_Limited) || ConstraintDetails::IsAngularPropertyEqual(AngularSwing2MotionProperty, ACM_Limited);
+		case EPropertyType::AngularTwistLimit:		return ConstraintDetails::IsAngularPropertyEqual(AngularTwistMotionProperty, ACM_Limited);
+		case EPropertyType::AngularAnyLimit:		return ConstraintDetails::IsAngularPropertyEqual(AngularSwing1MotionProperty, ACM_Limited) || ConstraintDetails::IsAngularPropertyEqual(AngularSwing2MotionProperty, ACM_Limited) || ConstraintDetails::IsAngularPropertyEqual(AngularTwistMotionProperty, ACM_Limited);
 	}	
 
 	return bIsVisible;
 }
 
-bool FPhysicsConstraintComponentDetails::GetBoolProperty(TSharedPtr<IPropertyHandle> Prop) const
-{
-	bool bIsEnabled = false;
-
-	if (Prop->GetValue(bIsEnabled))
-	{
-		return bIsEnabled;
-	}
-	return false;
-}
-
-bool FPhysicsConstraintComponentDetails::IsAngularPropertyLimited(TSharedPtr<IPropertyHandle> Prop) const
-{
-	uint8 Val;
-	if (Prop->GetValue(Val))
-	{
-		return Val == ACM_Limited;
-	}
-	return false;
-}
-
-bool FPhysicsConstraintComponentDetails::IsLinearMotionLimited() const
-{
-	uint8 XMotion, YMotion, ZMotion;
-	if (LinearXMotionProperty->GetValue(XMotion) && 
-		LinearYMotionProperty->GetValue(YMotion) && 
-		LinearZMotionProperty->GetValue(ZMotion))
-	{
-		return XMotion == LCM_Limited || YMotion == LCM_Limited || ZMotion == LCM_Limited;
-	}
-	return false;
-}
-
-FReply FPhysicsConstraintComponentDetails::OnHingeClicked()
-{
-	if (ConstraintComponent.IsValid())
-	{
-		UPhysicsConstraintComponent* Comp = Cast<UPhysicsConstraintComponent>(ConstraintComponent.Get());
-		Comp->ConstraintInstance.ConfigureAsHinge();
-		Comp->UpdateSpriteTexture();
-	}
-	return FReply::Handled();
-}
-
-FReply FPhysicsConstraintComponentDetails::OnPrismaticClicked()
-{
-	if (ConstraintComponent.IsValid())
-	{
-		UPhysicsConstraintComponent* Comp = Cast<UPhysicsConstraintComponent>(ConstraintComponent.Get());
-		Comp->ConstraintInstance.ConfigureAsPrismatic();
-		Comp->UpdateSpriteTexture();
-	}
-	return FReply::Handled();
-}
-
-FReply FPhysicsConstraintComponentDetails::OnBallSocketClicked()
-{
-	if (ConstraintComponent.IsValid())
-	{
-		UPhysicsConstraintComponent* Comp = Cast<UPhysicsConstraintComponent>(ConstraintComponent.Get());
-		Comp->ConstraintInstance.ConfigureAsBS();
-		Comp->UpdateSpriteTexture();
-	}
-	return FReply::Handled();
-}
-
 ECheckBoxState FPhysicsConstraintComponentDetails::IsLimitRadioChecked( TSharedPtr<IPropertyHandle> Property, uint8 Value ) const
 {
 	uint8 PropertyEnumValue = 0;
@@ -589,4 +917,17 @@ void FPhysicsConstraintComponentDetails::OnLimitRadioChanged( ECheckBoxState Che
 	}
 }
 
-#undef LOCTEXT_NAMESPACE
+TSharedRef<IPropertyTypeCustomization> FPhysicsConstraintProfileHandleCustomization::MakeInstance()
+{
+	return MakeShareable(new FPhysicsConstraintProfileHandleCustomization());
+}
+
+void FPhysicsConstraintProfileHandleCustomization::CustomizeChildren(TSharedRef<IPropertyHandle> PropertyHandle, IDetailChildrenBuilder& ChildBuilder, IPropertyTypeCustomizationUtils& CustomizationUtils)
+{
+	ChildBuilder.AddChildProperty(PropertyHandle->GetChildHandle(GET_MEMBER_NAME_CHECKED(FPhysicsConstraintProfileHandle, ProfileName)).ToSharedRef());
+}
+
+void FPhysicsConstraintProfileHandleCustomization::CustomizeHeader(TSharedRef<IPropertyHandle> PropertyHandle, FDetailWidgetRow& HeaderRow, IPropertyTypeCustomizationUtils& CustomizationUtils)
+{
+}
+#undef LOCTEXT_NAMESPACE
\ No newline at end of file
diff --git a/Engine/Source/Editor/DetailCustomizations/Private/PhysicsConstraintComponentDetails.h b/Engine/Source/Editor/DetailCustomizations/Private/PhysicsConstraintComponentDetails.h
--- a/Engine/Source/Editor/DetailCustomizations/Private/PhysicsConstraintComponentDetails.h
+++ b/Engine/Source/Editor/DetailCustomizations/Private/PhysicsConstraintComponentDetails.h
@@ -27,12 +27,8 @@ private:
 			LinearYVelocityDrive,
 			LinearZVelocityDrive,
 			LinearVelocityDrive,
-			LinearAnyDrive,
-			LinearLimit,
-			
-			AngularPositionDrive,
-			AngularVelocityDrive,
-			AngularAnyDrive,
+			LinearDrive,
+		
 			AngularSwingLimit,
 			AngularSwing1Limit,
 			AngularSwing2Limit,
@@ -42,44 +38,26 @@ private:
 	};
 
 	bool IsPropertyEnabled(EPropertyType::Type Type) const;
-	
-	bool IsLinearMotionLimited() const;
-	
-	bool GetBoolProperty(TSharedPtr<IPropertyHandle> Prop) const;
-	bool IsAngularPropertyLimited(TSharedPtr<IPropertyHandle> Prop) const;
-
-	FReply OnHingeClicked();
-	FReply OnPrismaticClicked();
-	FReply OnBallSocketClicked();
 
 	ECheckBoxState IsLimitRadioChecked(TSharedPtr<IPropertyHandle> Property, uint8 Value) const ;
 	void OnLimitRadioChanged(ECheckBoxState CheckType, TSharedPtr<IPropertyHandle> Property, uint8 Value);
 
-	void AddConstraintPresets(IDetailLayoutBuilder& DetailBuilder, class APhysicsConstraintActor* OwningConstraintActor);
-	void AddLinearLimits(IDetailLayoutBuilder& DetailBuilder, TSharedPtr<IPropertyHandle> ConstraintInstance);
-	void AddAngularLimits(IDetailLayoutBuilder& DetailBuilder, TSharedPtr<IPropertyHandle> ConstraintInstance);
-	void AddLinearDrive(IDetailLayoutBuilder& DetailBuilder, TSharedPtr<IPropertyHandle> ConstraintInstance);
-	void AddAngularDrive(IDetailLayoutBuilder& DetailBuilder, TSharedPtr<IPropertyHandle> ConstraintInstance);
-
-	TSharedRef<SWidget> JoinPropertyWidgets(TSharedPtr<IPropertyHandle> TargetProperty, FName TargetChildName, TSharedPtr<IPropertyHandle> ConstraintInstance, FName CheckPropertyName, TSharedPtr<IPropertyHandle>& StoreCheckProperty, EPropertyType::Type DriveType);
+	void AddConstraintBehaviorProperties(IDetailLayoutBuilder& DetailBuilder, TSharedPtr<IPropertyHandle> ConstraintInstance, TSharedPtr<IPropertyHandle> ProfileInstance);
+	void AddLinearLimits(IDetailLayoutBuilder& DetailBuilder, TSharedPtr<IPropertyHandle> ConstraintInstance, TSharedPtr<IPropertyHandle> ProfileInstance);
+	void AddAngularLimits(IDetailLayoutBuilder& DetailBuilder, TSharedPtr<IPropertyHandle> ConstraintInstance, TSharedPtr<IPropertyHandle> ProfileInstance);
+	void AddLinearDrive(IDetailLayoutBuilder& DetailBuilder, TSharedPtr<IPropertyHandle> ConstraintInstance, TSharedPtr<IPropertyHandle> ProfileInstance);
+	void AddAngularDrive(IDetailLayoutBuilder& DetailBuilder, TSharedPtr<IPropertyHandle> ConstraintInstance, TSharedPtr<IPropertyHandle> ProfileInstance);
+	void AddProfiles(IDetailLayoutBuilder& DetailBuilder);
 private:
 
-	TSharedPtr<IPropertyHandle> LinearPositionDriveProperty;
 	TSharedPtr<IPropertyHandle> LinearXPositionDriveProperty;
 	TSharedPtr<IPropertyHandle> LinearYPositionDriveProperty;
 	TSharedPtr<IPropertyHandle> LinearZPositionDriveProperty;
 
-	TSharedPtr<IPropertyHandle> LinearVelocityDriveProperty;
 	TSharedPtr<IPropertyHandle> LinearXVelocityDriveProperty;
 	TSharedPtr<IPropertyHandle> LinearYVelocityDriveProperty;
 	TSharedPtr<IPropertyHandle> LinearZVelocityDriveProperty;
-	TSharedPtr<IPropertyHandle> AngularPositionDriveProperty;
-	TSharedPtr<IPropertyHandle> AngularVelocityDriveProperty;
 
-	TSharedPtr<IPropertyHandle> LinearXMotionProperty;
-	TSharedPtr<IPropertyHandle> LinearYMotionProperty;
-	TSharedPtr<IPropertyHandle> LinearZMotionProperty;
-	
 	TSharedPtr<IPropertyHandle> AngularSwing1MotionProperty;
 	TSharedPtr<IPropertyHandle> AngularSwing2MotionProperty;
 	TSharedPtr<IPropertyHandle> AngularTwistMotionProperty;
@@ -88,3 +66,14 @@ private:
 	
 	bool bInPhat;
 };
+
+class FPhysicsConstraintProfileHandleCustomization : public IPropertyTypeCustomization
+{
+public:
+	/** Makes a new instance of this detail layout class for a specific detail view requesting it */
+	static TSharedRef<IPropertyTypeCustomization> MakeInstance();
+
+	virtual void CustomizeChildren( TSharedRef<IPropertyHandle> PropertyHandle, IDetailChildrenBuilder& ChildBuilder, IPropertyTypeCustomizationUtils& CustomizationUtils ) override;
+
+	virtual void CustomizeHeader( TSharedRef<IPropertyHandle> PropertyHandle, FDetailWidgetRow& HeaderRow, IPropertyTypeCustomizationUtils& CustomizationUtils ) override;
+};
\ No newline at end of file
diff --git a/Engine/Source/Editor/Persona/Private/AnimationEditorViewportClient.cpp b/Engine/Source/Editor/Persona/Private/AnimationEditorViewportClient.cpp
--- a/Engine/Source/Editor/Persona/Private/AnimationEditorViewportClient.cpp
+++ b/Engine/Source/Editor/Persona/Private/AnimationEditorViewportClient.cpp
@@ -2,7 +2,6 @@
 
 #include "PersonaPrivatePCH.h"
 
-#include "AnimationEditorPreviewScene.h"
 #include "SAnimationEditorViewport.h"
 #include "Runtime/Engine/Public/Slate/SceneViewport.h"
 #include "SAnimViewportToolBar.h"
@@ -30,6 +29,8 @@
 #include "Engine/StaticMesh.h"
 #include "SAnimationEditorViewport.h"
 
+#include "AssetViewerSettings.h"
+
 namespace {
 	// Value from UE3
 	static const float AnimationEditorViewport_RotateSpeed = 0.02f;
@@ -79,7 +80,7 @@ IMPLEMENT_HIT_PROXY( HPersonaBoneProxy, HHitProxy );
 /////////////////////////////////////////////////////////////////////////
 // FAnimationViewportClient
 
-FAnimationViewportClient::FAnimationViewportClient(FAnimationEditorPreviewScene& InPreviewScene, TWeakPtr<FPersona> InPersonaPtr, const TSharedRef<SAnimationEditorViewport>& InAnimationEditorViewport)
+FAnimationViewportClient::FAnimationViewportClient(FAdvancedPreviewScene& InPreviewScene, TWeakPtr<FPersona> InPersonaPtr, const TSharedRef<SAnimationEditorViewport>& InAnimationEditorViewport)
 	: FEditorViewportClient(nullptr, &InPreviewScene, StaticCastSharedRef<SEditorViewport>(InAnimationEditorViewport))
 	, PersonaPtr( InPersonaPtr )
 	, bManipulating(false)
@@ -117,15 +118,9 @@ FAnimationViewportClient::FAnimationViewportClient(FAnimationEditorPreviewScene&
 	{
 		SetRealtime(false,true); // We are PIE, don't start in realtime mode
 	}
-
-	// set visibility
-	GetAnimPreviewScene()->EditorSkyComp->SetVisibility(ConfigOption->bShowSky);
-	GetAnimPreviewScene()->EditorHeightFogComponent->SetVisibility(ConfigOption->bShowSky);
-	GetAnimPreviewScene()->EditorFloorComp->SetVisibility(ConfigOption->bShowFloor);
-
+	
 	ViewFOV = FMath::Clamp<float>(ConfigOption->ViewFOV, FOVMin, FOVMax);
 
-	EngineShowFlags.DisableAdvancedFeatures();
 	EngineShowFlags.SetSeparateTranslucency(true);
 	EngineShowFlags.SetCompositeEditorPrimitives(true);
 
@@ -173,17 +168,6 @@ FLinearColor FAnimationViewportClient::GetBackgroundColor() const
 	return SelectedHSVColor.HSVToLinearRGB();
 }
 
-FSceneView* FAnimationViewportClient::CalcSceneView(FSceneViewFamily* ViewFamily, const EStereoscopicPass StereoPass)
-{
-	float AmbientCubemapIntensity = 0.4f;
-
-	FSceneView* SceneView = FEditorViewportClient::CalcSceneView(ViewFamily, StereoPass);
-	FFinalPostProcessSettings::FCubemapEntry& CubemapEntry = *new(SceneView->FinalPostProcessSettings.ContributingCubemaps) FFinalPostProcessSettings::FCubemapEntry;
-	CubemapEntry.AmbientCubemap = GUnrealEd->GetThumbnailManager()->AmbientCubemap;
-	CubemapEntry.AmbientCubemapTintMulScaleValue = FLinearColor::White * AmbientCubemapIntensity;
-	return SceneView;
-}
-
 void FAnimationViewportClient::SetSelectedBackgroundColor(const FLinearColor& RGBColor, bool bSave/*=true*/)
 {
 	SelectedHSVColor = RGBColor.LinearRGBToHSV(); 
@@ -241,53 +225,6 @@ bool FAnimationViewportClient::IsShowingGrid() const
 	return FEditorViewportClient::IsSetShowGridChecked();
 }
 
-void FAnimationViewportClient::OnToggleShowFloor()
-{
-	if (GetAnimPreviewScene()->EditorFloorComp)
-	{
-		GetAnimPreviewScene()->EditorFloorComp->SetVisibility(!GetAnimPreviewScene()->EditorFloorComp->bVisible);
-		Invalidate();
-	}
-
-	ConfigOption->SetShowFloor(GetAnimPreviewScene()->EditorFloorComp->bVisible);
-}
-
-bool FAnimationViewportClient::IsShowingFloor() const
-{
-	return (GetAnimPreviewScene()->EditorFloorComp) ? GetAnimPreviewScene()->EditorFloorComp->bVisible : false;
-}
-
-void FAnimationViewportClient::OnToggleAutoAlignFloor()
-{
-	bAutoAlignFloor = !bAutoAlignFloor;
-	UpdateCameraSetup();
-
-	ConfigOption->SetAutoAlignFloorToMesh(bAutoAlignFloor);
-}
-
-bool FAnimationViewportClient::IsAutoAlignFloor() const
-{
-	return bAutoAlignFloor;
-}
-
-void FAnimationViewportClient::OnToggleShowSky()
-{
-	if (GetAnimPreviewScene()->EditorSkyComp)
-	{
-		bool bNewVisibility = !GetAnimPreviewScene()->EditorSkyComp->bVisible;
-		GetAnimPreviewScene()->EditorSkyComp->SetVisibility(bNewVisibility);
-		GetAnimPreviewScene()->EditorHeightFogComponent->SetVisibility(bNewVisibility);
-		Invalidate();
-	}
-
-	ConfigOption->SetShowSky(GetAnimPreviewScene()->EditorSkyComp->bVisible);
-}
-
-bool FAnimationViewportClient::IsShowingSky() const
-{
-	return (GetAnimPreviewScene()->EditorSkyComp) ? GetAnimPreviewScene()->EditorSkyComp->bVisible : false;
-}
-
 void FAnimationViewportClient::OnToggleMuteAudio()
 {
 	UWorld* World = PreviewScene->GetWorld();
@@ -593,29 +530,9 @@ FAnimNode_SkeletalControlBase* FAnimationViewportClient::FindSkeletalControlAnim
 {
 	FAnimNode_SkeletalControlBase* AnimNode = NULL;
 
-	if (!PreviewSkelMeshComp.IsValid() || !PreviewSkelMeshComp->GetAnimInstance())
-	{
-		return NULL;
-	}
-
-	if (AnimGraphNode.IsValid())
+	if (PreviewSkelMeshComp.IsValid() && PreviewSkelMeshComp->GetAnimInstance() && AnimGraphNode.IsValid())
 	{
-		// find an anim node index from debug data
-		UAnimBlueprintGeneratedClass* AnimBlueprintClass = Cast<UAnimBlueprintGeneratedClass>(PreviewSkelMeshComp->GetAnimInstance()->GetClass());
-		if (AnimBlueprintClass)
-		{
-			FAnimBlueprintDebugData& DebugData = AnimBlueprintClass->GetAnimBlueprintDebugData();
-			int32* IndexPtr = DebugData.NodePropertyToIndexMap.Find(AnimGraphNode);
-
-			if (IndexPtr)
-			{
-				int32 AnimNodeIndex = *IndexPtr;
-				// reverse node index temporarily because of a bug in NodeGuidToIndexMap
-				AnimNodeIndex = AnimBlueprintClass->AnimNodeProperties.Num() - AnimNodeIndex - 1;
-
-				AnimNode = AnimBlueprintClass->AnimNodeProperties[AnimNodeIndex]->ContainerPtrToValuePtr<FAnimNode_SkeletalControlBase>(PreviewSkelMeshComp->GetAnimInstance());
-			}
-		}
+		AnimNode = AnimGraphNode.Get()->FindDebugAnimNode(PreviewSkelMeshComp.Get());
 	}
 
 	return AnimNode;
@@ -777,8 +694,9 @@ void FAnimationViewportClient::Tick(float DeltaSeconds)
 	UDebugSkelMeshComponent* PreviewComp = PreviewSkelMeshComp.Get();
 	if (PreviewComp)
 	{
-		// Handle updating the preview component to represent the effects of root motion	
-		FBoxSphereBounds Bounds = GetAnimPreviewScene()->EditorFloorComp->CalcBounds(GetAnimPreviewScene()->EditorFloorComp->GetRelativeTransform());
+		// Handle updating the preview component to represent the effects of root motion
+		const UStaticMeshComponent* FloorMeshComponent = GetAdvancedPreviewScene()->GetFloorMeshComponent();
+		FBoxSphereBounds Bounds = FloorMeshComponent->CalcBounds(FloorMeshComponent->GetRelativeTransform());
 		PreviewComp->ConsumeRootMotion(Bounds.GetBox().Min, Bounds.GetBox().Max);
 	}	
 
@@ -858,8 +776,12 @@ void FAnimationViewportClient::DisplayInfo(FCanvas* Canvas, FSceneView* View, bo
 	StringSize( GEngine->GetSmallFont(),  XL, YL, TEXT("L") );
 	FString InfoString;
 
+	const UAssetViewerSettings* Settings = UAssetViewerSettings::Get();	
+	const UEditorPerProjectUserSettings* PerProjectUserSettings = GetDefault<UEditorPerProjectUserSettings>();
+	const int32 ProfileIndex = Settings->Profiles.IsValidIndex(PerProjectUserSettings->AssetViewerProfileIndex) ? PerProjectUserSettings->AssetViewerProfileIndex : 0;
+
 	// it is weird, but unless it's completely black, it's too bright, so just making it white if only black
-	const FLinearColor TextColor = (SelectedHSVColor.B < 0.3f) ? FLinearColor::White : FLinearColor::Black;
+	const FLinearColor TextColor = ((SelectedHSVColor.B < 0.3f) || (Settings->Profiles[ProfileIndex].bShowEnvironment)) ? FLinearColor::White : FLinearColor::Black;
 	const FColor HeadlineColour(255, 83, 0);
 	const FColor SubHeadlineColour(202, 66, 0);
 
@@ -1725,6 +1647,27 @@ bool FAnimationViewportClient::InputKey( FViewport* InViewport, int32 Controller
 		: FEditorViewportClient::InputKey(InViewport,  ControllerId,  Key,  Event,  AmountDepressed,  bGamepad);
 }
 
+bool FAnimationViewportClient::InputAxis(FViewport* InViewport, int32 ControllerId, FKey Key, float Delta, float DeltaTime, int32 NumSamples /*= 1*/, bool bGamepad /*= false*/)
+{
+	bool bResult = true;
+
+	if (!bDisableInput)
+	{
+		FAdvancedPreviewScene* AdvancedScene = (FAdvancedPreviewScene*)PreviewScene;
+		bResult = AdvancedScene->HandleViewportInput(InViewport, ControllerId, Key, Delta, DeltaTime, NumSamples, bGamepad);
+		if (bResult)
+		{
+			Invalidate();
+		}
+		else
+		{
+			bResult = FEditorViewportClient::InputAxis(InViewport, ControllerId, Key, Delta, DeltaTime, NumSamples, bGamepad);
+		}
+	}
+
+	return bResult;
+}
+
 void FAnimationViewportClient::SetWidgetMode(FWidget::EWidgetMode InMode)
 {
 	bool bAnimBPMode = PersonaPtr.Pin()->IsModeCurrent(FPersonaModes::AnimBlueprintEditMode);
@@ -2150,19 +2093,6 @@ void FAnimationViewportClient::UpdateCameraSetup()
 		FVector CustomOrbitLookAt = BoundSphere.Center;
 
 		SetCameraSetup(CustomOrbitLookAt, CustomOrbitRotation, CustomOrbitZoom, CustomOrbitLookAt, GetViewLocation(), GetViewRotation() );
-
-		// Move the floor to the bottom of the bounding box of the mesh, rather than on the origin
-		FVector Bottom = PreviewSkelMeshComp->Bounds.GetBoxExtrema(0);
-
-		if (GetAnimPreviewScene()->EditorFloorComp)
-		{
-			FVector FloorPos(0.f, 0.f, 0.f);
-			if (bAutoAlignFloor)
-			{
-				FloorPos.Z = GetFloorOffset() + Bottom.Z;
-			}
-			GetAnimPreviewScene()->EditorFloorComp->SetWorldTransform(FTransform(FQuat::Identity, FloorPos, FVector(3.0f, 3.0f, 1.0f)));
-		}
 	}
 }
 
@@ -2436,9 +2366,9 @@ int32 FAnimationViewportClient::GetShowMeshStats() const
 	return ConfigOption->ShowMeshStats;
 }
 
-FAnimationEditorPreviewScene* FAnimationViewportClient::GetAnimPreviewScene() const
+FAdvancedPreviewScene* FAnimationViewportClient::GetAdvancedPreviewScene() const
 {
-	return static_cast<FAnimationEditorPreviewScene*>(PreviewScene);
+	return static_cast<FAdvancedPreviewScene*>(PreviewScene);
 }
 
 #undef LOCTEXT_NAMESPACE
diff --git a/Engine/Source/Editor/Persona/Private/AnimationEditorViewportClient.h b/Engine/Source/Editor/Persona/Private/AnimationEditorViewportClient.h
--- a/Engine/Source/Editor/Persona/Private/AnimationEditorViewportClient.h
+++ b/Engine/Source/Editor/Persona/Private/AnimationEditorViewportClient.h
@@ -5,7 +5,7 @@
 #include "EditorViewportClient.h"
 #include "Engine/WindDirectionalSource.h"
 
-class FAnimationEditorPreviewScene;
+class FAdvancedPreviewScene;
 
 //////////////////////////////////////////////////////////////////////////
 // ELocalAxesMode
@@ -70,7 +70,7 @@ protected:
 	void DrawNodeDebugLines(TArray<FText>& Lines, FCanvas* Canvas, FSceneView* View);
 
 public:
-	FAnimationViewportClient(FAnimationEditorPreviewScene& InPreviewScene, TWeakPtr<FPersona> InPersonaPtr, const TSharedRef<SAnimationEditorViewport>& InAnimationEditorViewport);
+	FAnimationViewportClient(FAdvancedPreviewScene& InPreviewScene, TWeakPtr<FPersona> InPersonaPtr, const TSharedRef<SAnimationEditorViewport>& InAnimationEditorViewport);
 	virtual ~FAnimationViewportClient();
 
 	// FEditorViewportClient interface
@@ -78,8 +78,8 @@ public:
 	virtual void Tick(float DeltaSeconds) override;
 	virtual void Draw(const FSceneView* View,FPrimitiveDrawInterface* PDI) override;
 	virtual void DrawCanvas( FViewport& InViewport, FSceneView& View, FCanvas& Canvas ) override;
-	virtual FSceneView* CalcSceneView(FSceneViewFamily* ViewFamily, const EStereoscopicPass StereoPass) override;
 	virtual bool InputKey(FViewport* Viewport, int32 ControllerId, FKey Key, EInputEvent Event, float AmountDepressed = 1.f, bool bGamepad=false) override;
+	virtual bool InputAxis(FViewport* InViewport, int32 ControllerId, FKey Key, float Delta, float DeltaTime, int32 NumSamples = 1, bool bGamepad = false) override;
 	virtual void ProcessClick(class FSceneView& View, class HHitProxy* HitProxy, FKey Key, EInputEvent Event, uint32 HitX, uint32 HitY) override;
 	virtual bool InputWidgetDelta( FViewport* Viewport, EAxisList::Type CurrentAxis, FVector& Drag, FRotator& Rot, FVector& Scale ) override;
 	virtual void TrackingStarted( const struct FInputEventState& InInputState, bool bIsDragging, bool bNudge ) override;
@@ -116,24 +116,6 @@ public:
 	/** Function to check whether grid is displayed or not */
 	bool IsShowingGrid() const;
 
-	/** Function to show/hide floor in the viewport */
-	void OnToggleShowFloor();
-
-	/** Function to check whether grid is displayed or not */
-	bool IsShowingFloor() const;
-
-	/** Function to enable/disable floor auto align */
-	void OnToggleAutoAlignFloor();
-
-	/** Function to check whether floor is auto align or not */
-	bool IsAutoAlignFloor() const;
-
-	/** Function to show/hide Sky in the viewport */
-	void OnToggleShowSky();
-
-	/** Function to check whether grid is displayed or not */
-	bool IsShowingSky() const;
-
 	/** Function to mute/unmute audio in the viewport */
 	void OnToggleMuteAudio();
 
@@ -365,5 +347,5 @@ private:
 	// to check whether we should update literal values in selected AnimGraphNode
 	bool bShouldUpdateDefaultValues;
 
-	FAnimationEditorPreviewScene* GetAnimPreviewScene() const;
+	FAdvancedPreviewScene* GetAdvancedPreviewScene() const;
 };
diff --git a/Engine/Source/Editor/Persona/Private/BlueprintMode/AnimBlueprintMode.cpp b/Engine/Source/Editor/Persona/Private/BlueprintMode/AnimBlueprintMode.cpp
--- a/Engine/Source/Editor/Persona/Private/BlueprintMode/AnimBlueprintMode.cpp
+++ b/Engine/Source/Editor/Persona/Private/BlueprintMode/AnimBlueprintMode.cpp
@@ -116,6 +116,7 @@ FAnimBlueprintEditAppMode::FAnimBlueprintEditAppMode(TSharedPtr<FPersona> InPers
 	PersonaTabFactories.RegisterFactory(MakeShareable(new FAnimBlueprintPreviewEditorSummoner(InPersona)));
 	PersonaTabFactories.RegisterFactory(MakeShareable(new FAnimBlueprintParentPlayerEditorSummoner(InPersona)));
 	PersonaTabFactories.RegisterFactory(MakeShareable(new FSkeletonSlotNamesSummoner(InPersona)));
+	PersonaTabFactories.RegisterFactory(MakeShareable(new FAdvancedPreviewSceneTabSummoner(InPersona)));
 
 	// setup toolbar - clear existing toolbar extender from the BP mode
 	//@TODO: Keep this in sync with BlueprintEditorModes.cpp
diff --git a/Engine/Source/Editor/Persona/Private/Persona.cpp b/Engine/Source/Editor/Persona/Private/Persona.cpp
--- a/Engine/Source/Editor/Persona/Private/Persona.cpp
+++ b/Engine/Source/Editor/Persona/Private/Persona.cpp
@@ -78,9 +78,12 @@
 #include "Animation/PoseAsset.h"
 
 #include "MessageLog.h"
+#include "SAdvancedPreviewDetailsTab.h"
 
 #define LOCTEXT_NAMESPACE "FPersona"
 
+const FName FPersona::PreviewSceneSettingsTabId(TEXT("Persona_PreviewScene"));
+
 /////////////////////////////////////////////////////
 // FLocalCharEditorCallbacks
 
diff --git a/Engine/Source/Editor/Persona/Private/Persona.h b/Engine/Source/Editor/Persona/Private/Persona.h
--- a/Engine/Source/Editor/Persona/Private/Persona.h
+++ b/Engine/Source/Editor/Persona/Private/Persona.h
@@ -4,7 +4,7 @@
 
 #include "Editor/Kismet/Public/BlueprintEditor.h"
 
-#include "AnimationEditorPreviewScene.h"
+#include "AdvancedPreviewScene.h"
 
 #include "PersonaModule.h"
 #include "PersonaCommands.h"
@@ -162,7 +162,7 @@ public:
 	void CleanupComponent(USceneComponent* Component);
 
 	/** Returns the editors preview scene */
-	FAnimationEditorPreviewScene& GetPreviewScene() { return PreviewScene; }
+	FAdvancedPreviewScene& GetPreviewScene() { return PreviewScene; }
 
 	/** Gets called when it's clicked via mode tab - Reinitialize the mode **/
 	void ReinitMode();
@@ -193,7 +193,7 @@ public:
 	virtual FText GetToolkitName() const override;
 	virtual FText GetToolkitToolTipText() const override;
 	virtual FString GetWorldCentricTabPrefix() const override;
-	virtual FLinearColor GetWorldCentricTabColorScale() const override;
+	virtual FLinearColor GetWorldCentricTabColorScale() const override;	
 	//~ End IToolkit Interface
 
 	/** Saves all animation assets related to a skeleton */
@@ -346,7 +346,6 @@ private:
 	// Called when the preview viewport is created
 	DECLARE_MULTICAST_DELEGATE_OneParam( FOnCreateViewport, TWeakPtr<SAnimationEditorViewportTabBody> )
 public:
-
 	// Called when Persona refreshes
 	typedef FSimpleMulticastDelegate::FDelegate FOnPersonaRefresh;
 
@@ -719,7 +718,7 @@ private:
 	TSharedPtr<FExtender> ToolbarExtender;
 
 	/** Preview scene for the editor */
-	FAnimationEditorPreviewScene PreviewScene;
+	FAdvancedPreviewScene PreviewScene;
 
 	/** Brush to use as a dirty marker for assets */
 	const FSlateBrush* AssetDirtyBrush;
@@ -739,4 +738,6 @@ private:
 	/** Handle additive anim scale validation */
 	bool bDoesAdditiveRefPoseHaveZeroScale;
 	FGuid RefPoseGuid;
+
+	static const FName PreviewSceneSettingsTabId;
 };
