Copying //UE4/Dev-Editor to //UE4/Dev-Main (Source: //UE4/Dev-Editor @ 3228984)

diff --git a/Engine/Source/Editor/DetailCustomizations/Private/ColorGradingVectorCustomization.cpp b/Engine/Source/Editor/DetailCustomizations/Private/ColorGradingVectorCustomization.cpp
new file mode 100644
--- /dev/null
+++ b/Engine/Source/Editor/DetailCustomizations/Private/ColorGradingVectorCustomization.cpp
@@ -0,0 +1,277 @@
+// Copyright 1998-2015 Epic Games, Inc. All Rights Reserved.
+#include "ColorGradingVectorCustomization.h"
+#include "IPropertyUtilities.h"
+#include "SNumericEntryBox.h"
+#include "SColorGradingPicker.h"
+#include "IDetailChildrenBuilder.h"
+#include "DetailLayoutBuilder.h"
+#include "DetailWidgetRow.h"
+
+FColorGradingVectorCustomization::FColorGradingVectorCustomization()
+{
+	CurrentData.Empty();
+	CustomColorGradingBuilder = nullptr;
+}
+
+void FColorGradingVectorCustomization::CustomizeChildren(TSharedRef<IPropertyHandle> StructPropertyHandle, IDetailChildrenBuilder& StructBuilder, IPropertyTypeCustomizationUtils& StructCustomizationUtils, TArray<TWeakPtr<IPropertyHandle>> WeakChildArray)
+{
+	FColorGradingCustomBuilderDelegates ColorGradingPickerDelegates;
+	ColorGradingPickerDelegates.OnGetColorGradingData = FOnGetColorGradingData::CreateSP(this, &FColorGradingVectorCustomization::OnGetData);
+	ColorGradingPickerDelegates.OnColorGradingChanged = FOnColorGradingChanged::CreateSP(this, &FColorGradingVectorCustomization::OnVector4DataChanged);
+	TWeakPtr<IPropertyHandle> PropertyHandleWeak = StructPropertyHandle;
+	CustomColorGradingBuilder = MakeShareable(new FColorGradingCustomBuilder(ColorGradingPickerDelegates, PropertyHandleWeak, WeakChildArray));
+	// Add the individual properties as children as well so the vector can be expanded for more room
+	StructBuilder.AddChildCustomBuilder(CustomColorGradingBuilder.ToSharedRef());
+}
+
+void FColorGradingVectorCustomization::OnGetData(TArray<FVector4*> &OutColorGradingData, TWeakPtr<IPropertyHandle> StructPropertyHandle)
+{
+	//We need a value array to kept the FVector4 pointer alive during the color grading Picker operation
+	TArray<FVector4> PropertyValues;
+	CurrentData.Empty();
+	TArray<FString> StrPropertyValues;
+	if (StructPropertyHandle.IsValid() && StructPropertyHandle.Pin()->GetPerObjectValues(StrPropertyValues) == FPropertyAccess::Success)
+	{
+		// Loop through each object and scale based on the new ratio for each object individually
+		PropertyValues.AddZeroed(StrPropertyValues.Num());
+		for (int32 PropertyValueIndex = 0; PropertyValueIndex < StrPropertyValues.Num(); ++PropertyValueIndex)
+		{
+			PropertyValues[PropertyValueIndex].InitFromString(StrPropertyValues[PropertyValueIndex]);
+			FVector4 *NewVector4 = new FVector4(PropertyValues[PropertyValueIndex]);
+			CurrentData.Add(NewVector4);
+			OutColorGradingData.Add(NewVector4);
+		}
+	}
+}
+
+void FColorGradingVectorCustomization::OnVector4DataChanged(FVector4 ColorGradingData, TWeakPtr<IPropertyHandle> StructPropertyHandle)
+{
+	TArray<FString> ObjectChildValues;
+	if (StructPropertyHandle.IsValid() && StructPropertyHandle.Pin()->GetPerObjectValues (ObjectChildValues) == FPropertyAccess::Success)
+	{
+		for (int32 PropertyValueIndex = 0; PropertyValueIndex < ObjectChildValues.Num(); ++PropertyValueIndex)
+		{
+			ObjectChildValues[PropertyValueIndex] = FString::Printf(TEXT("(X=%3.3f,Y=%3.3f,Z=%3.3f,W=%3.3f)"), ColorGradingData.X, ColorGradingData.Y, ColorGradingData.Z, ColorGradingData.W);
+		}
+
+		StructPropertyHandle.Pin()->SetPerObjectValues(ObjectChildValues);
+	}
+}
+
+
+//////////////////////////////////////////////////////////////////////////
+// Color Gradient custom builder implementation
+
+FColorGradingCustomBuilder::FColorGradingCustomBuilder(FColorGradingCustomBuilderDelegates& ColorGradingCustomBuilderDelegates, TWeakPtr<IPropertyHandle> InColorGradingPropertyHandle, TArray< TWeakPtr<IPropertyHandle> >& InSortedChildArray)
+{
+	ColorGradingDelegates = ColorGradingCustomBuilderDelegates;
+	ColorGradingPropertyHandle = InColorGradingPropertyHandle;
+	for (TWeakPtr<IPropertyHandle> ChildProperty : InSortedChildArray)
+	{
+		SortedChildArray.Add(ChildProperty);
+	}
+}
+
+void FColorGradingCustomBuilder::Tick(float DeltaTime)
+{
+}
+
+void FColorGradingCustomBuilder::GenerateHeaderRowContent(FDetailWidgetRow& NodeRow)
+{
+	if (!ColorGradingDelegates.OnGetColorGradingData.IsBound())
+	{
+		return;
+	}
+	//Query all meta data we need
+	UProperty* Property = ColorGradingPropertyHandle.Pin()->GetProperty();
+	const FString& MetaUIMinString = Property->GetMetaData(TEXT("UIMin"));
+	const FString& MetaUIMaxString = Property->GetMetaData(TEXT("UIMax"));
+	const FString& DeltaString = Property->GetMetaData(TEXT("Delta"));
+	const FString& ShiftMouseMovePixelPerDeltaString = Property->GetMetaData(TEXT("ShiftMouseMovePixelPerDelta"));
+	const FString& ColorGradingModeString = Property->GetMetaData(TEXT("ColorGradingMode"));
+	float UIMin = TNumericLimits<float>::Lowest();
+	float UIMax = TNumericLimits<float>::Max();
+	TTypeFromString<float>::FromString(UIMin, *MetaUIMinString);
+	TTypeFromString<float>::FromString(UIMax, *MetaUIMaxString);
+	float Delta = float(0);
+	if (DeltaString.Len())
+	{
+		TTypeFromString<float>::FromString(Delta, *DeltaString);
+	}
+	int32 ShiftMouseMovePixelPerDelta = 1;
+	if (ShiftMouseMovePixelPerDeltaString.Len())
+	{
+		TTypeFromString<int32>::FromString(ShiftMouseMovePixelPerDelta, *ShiftMouseMovePixelPerDeltaString);
+		//The value should be greater or equal to 1
+		// 1 is neutral since it is a multiplier of the mouse drag pixel
+		if (ShiftMouseMovePixelPerDelta < 1)
+		{
+			ShiftMouseMovePixelPerDelta = 1;
+		}
+	}
+
+	EColorGradingModes ColorGradingMode = EColorGradingModes::Contrast;
+	if (ColorGradingModeString.Len())
+	{
+		if (ColorGradingModeString.Compare(TEXT("saturation")) == 0)
+		{
+			ColorGradingMode = EColorGradingModes::Saturation;
+		}
+		else if (ColorGradingModeString.Compare(TEXT("contrast")) == 0)
+		{
+			ColorGradingMode = EColorGradingModes::Contrast;
+		}
+		else if (ColorGradingModeString.Compare(TEXT("gamma")) == 0)
+		{
+			ColorGradingMode = EColorGradingModes::Gamma;
+		}
+		else if (ColorGradingModeString.Compare(TEXT("gain")) == 0)
+		{
+			ColorGradingMode = EColorGradingModes::Gain;
+		}
+		else if (ColorGradingModeString.Compare(TEXT("offset")) == 0)
+		{
+			ColorGradingMode = EColorGradingModes::Offset;
+		}
+	}
+
+	NodeRow.NameContent()
+		.HAlign(HAlign_Fill)
+		[
+			SNew(SBox)
+			.HAlign(HAlign_Center)
+			.VAlign(VAlign_Center)
+			.WidthOverride(125.0f)
+			//.HeightOverride(150.0f)
+			.MinDesiredWidth(125.0f)
+			.MaxDesiredWidth(125.0f)
+			.Padding(FMargin(2.0f, 2.0f, 2.0f, 2.0f))
+			[
+				SNew(SColorGradingPicker)
+					.ValuesMin(UIMin)
+					.ValuesMax(UIMax)
+					.MainDelta(Delta)
+					.MainShiftMouseMovePixelPerDelta(ShiftMouseMovePixelPerDelta)
+					.ColorGradingModes(ColorGradingMode)
+					.OnColorCommitted(FOnVector4ValueChanged::CreateSP(this, &FColorGradingCustomBuilder::OnColorGradingPickerChanged))
+					.OnQueryCurrentColor(FOnGetCurrentVector4Value::CreateSP(this, &FColorGradingCustomBuilder::GetCurrentColorGradingValue))
+			]
+		];
+	
+	TSharedPtr<SVerticalBox> VerticalBox = SNew(SVerticalBox);
+	VerticalBox->AddSlot()
+		.Padding(FMargin(0.0f, 2.0f, 0.0f, 0.0f))
+		.FillHeight(1.0)
+		[
+			SNew(SBox)
+		];
+	int32 ColorIndex = 0;
+	for (TWeakPtr<IPropertyHandle> ChildProperty : SortedChildArray)
+	{
+		TSharedPtr<SWidget> LabelWidget = SNew(STextBlock)
+			.Font(IDetailLayoutBuilder::GetDetailFont())
+			.Text(ChildProperty.Pin()->GetPropertyDisplayName());
+
+
+		TSharedPtr<SNumericEntryBox<float>> NumericEntry;
+		SAssignNew(NumericEntry, SNumericEntryBox<float>)
+			.EditableTextBoxStyle(&FCoreStyle::Get().GetWidgetStyle<FEditableTextBoxStyle>("DarkEditableTextBox"))
+			.Value(this, &FColorGradingCustomBuilder::OnSliderGetValue, ChildProperty)
+			.Font(IDetailLayoutBuilder::GetDetailFont())
+			.OnValueChanged(this, &FColorGradingCustomBuilder::OnSliderValueChanged, ChildProperty)
+			.LabelVAlign(VAlign_Center)
+			// Only allow spin on handles with one object.  Otherwise it is not clear what value to spin
+			.AllowSpin(ChildProperty.Pin()->GetNumOuterObjects() == 1)
+			.UseDarkStyle(true)
+			.ColorIndex(ColorIndex)
+			.ShiftMouseMovePixelPerDelta(ShiftMouseMovePixelPerDelta)
+			.MinValue(UIMin)
+			.MaxValue(UIMax)
+			.MinSliderValue(UIMin)
+			.MaxSliderValue(UIMax)
+			.SliderExponent(3.0f)
+			.SliderExponentNeutralValue(UIMin + (UIMax - UIMin) / 2.0f)
+			.Delta(Delta)
+			.Label()
+			[
+				LabelWidget.ToSharedRef()
+			];
+
+		VerticalBox->AddSlot()
+			.Padding(FMargin(3.0f, 3.0f, 3.0f, 3.0f))
+			.AutoHeight()
+			.HAlign(HAlign_Fill)
+			.VAlign(VAlign_Center)
+			[
+				NumericEntry.ToSharedRef()
+			];
+		ColorIndex++;
+	}
+	//Add a filler to take all the rest of the space
+	VerticalBox->AddSlot()
+		.Padding(FMargin(0.0f, 0.0f, 0.0f, 2.0f))
+		.FillHeight(1.0)
+		[
+			SNew(SBox)
+		];
+	
+	NodeRow.ValueContent()
+		.HAlign(HAlign_Fill)
+	[
+		VerticalBox.ToSharedRef()
+	];
+}
+
+void FColorGradingCustomBuilder::GenerateChildContent(IDetailChildrenBuilder& ChildrenBuilder)
+{
+}
+
+void FColorGradingCustomBuilder::OnSliderValueChanged(float NewValue, TWeakPtr<IPropertyHandle> WeakHandlePtr)
+{
+	if (!ColorGradingDelegates.OnColorGradingChanged.IsBound())
+	{
+		return;
+	}
+	int32 ChildIndex = 0;
+	if (SortedChildArray.Find(WeakHandlePtr, ChildIndex))
+	{
+		FVector4 ValueVector;
+		GetCurrentColorGradingValue(ValueVector);
+		ValueVector[ChildIndex] = NewValue;
+		ColorGradingDelegates.OnColorGradingChanged.Execute(ValueVector, ColorGradingPropertyHandle);
+	}
+}
+
+TOptional<float> FColorGradingCustomBuilder::OnSliderGetValue(TWeakPtr<IPropertyHandle> WeakHandlePtr) const
+{
+	float CurrentValue = 0.0f;
+	if (WeakHandlePtr.Pin()->GetValue(CurrentValue) == FPropertyAccess::Success)
+	{
+		return CurrentValue;
+	}
+
+	// Value couldn't be accessed.  Return an unset value
+	return 0.0f;
+}
+
+void FColorGradingCustomBuilder::OnColorGradingPickerChanged(FVector4 &NewValue)
+{
+	if (!ColorGradingDelegates.OnColorGradingChanged.IsBound())
+	{
+		return;
+	}
+	ColorGradingDelegates.OnColorGradingChanged.Execute(NewValue, ColorGradingPropertyHandle);
+}
+
+void FColorGradingCustomBuilder::GetCurrentColorGradingValue(FVector4 &OutCurrentValue)
+{
+	if (!ColorGradingDelegates.OnGetColorGradingData.IsBound())
+	{
+		OutCurrentValue = FVector4(0.0f);
+		return;
+	}
+	TArray<FVector4*> ColorGradingArrayPtr;
+	ColorGradingDelegates.OnGetColorGradingData.Execute(ColorGradingArrayPtr, ColorGradingPropertyHandle);
+
+	OutCurrentValue = ColorGradingArrayPtr.Num() > 0 ? *ColorGradingArrayPtr[0] : FVector4(0.0f);
+}
diff --git a/Engine/Source/Editor/DetailCustomizations/Private/ColorGradingVectorCustomization.h b/Engine/Source/Editor/DetailCustomizations/Private/ColorGradingVectorCustomization.h
new file mode 100644
--- /dev/null
+++ b/Engine/Source/Editor/DetailCustomizations/Private/ColorGradingVectorCustomization.h
@@ -0,0 +1,70 @@
+// Copyright 1998-2015 Epic Games, Inc. All Rights Reserved.
+
+#pragma once
+
+#include "MathStructCustomizations.h"
+#include "IDetailCustomNodeBuilder.h"
+
+class FColorGradingVectorCustomization : public TSharedFromThis<FColorGradingVectorCustomization>
+{
+public:
+	FColorGradingVectorCustomization();
+	void CustomizeChildren(TSharedRef<IPropertyHandle> StructPropertyHandle, IDetailChildrenBuilder& StructBuilder, IPropertyTypeCustomizationUtils& StructCustomizationUtils, TArray<TWeakPtr<IPropertyHandle>> WeakChildArray);
+	void OnGetData(TArray<FVector4*> &OutVector4Data, TWeakPtr<IPropertyHandle> StructPropertyHandle);
+	void OnVector4DataChanged(FVector4 ColorGradingData, TWeakPtr<IPropertyHandle> StructPropertyHandle);
+	TArray<FVector4*> CurrentData;
+	TSharedPtr<class FColorGradingCustomBuilder> CustomColorGradingBuilder;
+};
+
+// Delegate called when we need to get the current Color Data
+DECLARE_DELEGATE_TwoParams(FOnGetColorGradingData, TArray<FVector4*>&, TWeakPtr<IPropertyHandle>);
+// Delegate called when the widget Color Data changed
+DECLARE_DELEGATE_TwoParams(FOnColorGradingChanged, FVector4, TWeakPtr<IPropertyHandle>);
+
+struct FColorGradingCustomBuilderDelegates
+{
+	FColorGradingCustomBuilderDelegates()
+		: OnGetColorGradingData()
+		, OnColorGradingChanged()
+	{}
+	// Delegate called when we need to get the current Color Data
+	FOnGetColorGradingData OnGetColorGradingData;
+	// Delegate called when the widget Color Data changed
+	FOnColorGradingChanged OnColorGradingChanged;
+};
+
+class FColorGradingCustomBuilder : public IDetailCustomNodeBuilder, public TSharedFromThis<FColorGradingCustomBuilder>
+{
+public:
+	FColorGradingCustomBuilder(FColorGradingCustomBuilderDelegates& ColorGradingCustomBuilderDelegates, TWeakPtr<IPropertyHandle> InColorGradingPropertyHandle, TArray< TWeakPtr<IPropertyHandle> >& InSortedChildArray);
+	~FColorGradingCustomBuilder()
+	{
+	}
+
+private:
+
+	/** IDetailCustomNodeBuilder interface */
+	virtual void SetOnRebuildChildren(FSimpleDelegate InOnRebuildChildren) override { OnRebuildChildren = InOnRebuildChildren; }
+	virtual bool RequiresTick() const override { return false; }
+	virtual void Tick(float DeltaTime) override;
+	virtual void GenerateHeaderRowContent(FDetailWidgetRow& NodeRow) override;
+	virtual void GenerateChildContent(IDetailChildrenBuilder& ChildrenBuilder) override;
+	virtual FName GetName() const override { return NAME_None; }
+	virtual bool InitiallyCollapsed() const override { return false; }
+
+	/* Local UI Handlers*/
+	void OnColorGradingPickerChanged(FVector4 &NewValue);
+	void GetCurrentColorGradingValue(FVector4 &OutCurrentValue);
+
+	void OnSliderValueChanged(float NewValue, TWeakPtr<IPropertyHandle> WeakHandlePtr);
+	TOptional<float> OnSliderGetValue(TWeakPtr<IPropertyHandle> WeakHandlePtr) const;
+
+	/** Delegates for the Section list */
+	FColorGradingCustomBuilderDelegates ColorGradingDelegates;
+	/** Called to rebuild the children of the detail tree */
+	FSimpleDelegate OnRebuildChildren;
+	/* The property we want to change */
+	TWeakPtr<IPropertyHandle> ColorGradingPropertyHandle;
+	// The sorted child array
+	TArray< TWeakPtr<IPropertyHandle> > SortedChildArray;
+};
\ No newline at end of file
diff --git a/Engine/Source/Editor/DetailCustomizations/Private/CurveColorCustomization.cpp b/Engine/Source/Editor/DetailCustomizations/Private/CurveColorCustomization.cpp
--- a/Engine/Source/Editor/DetailCustomizations/Private/CurveColorCustomization.cpp
+++ b/Engine/Source/Editor/DetailCustomizations/Private/CurveColorCustomization.cpp
@@ -210,6 +210,17 @@ void FCurveColorCustomization::ModifyOwner()
 	}
 }
 
+TArray<const UObject*> FCurveColorCustomization::GetOwners() const
+{
+	TArray<const UObject*> Owners;
+	if (Owner)
+	{
+		Owners.Add(Owner);
+	}
+
+	return Owners;
+}
+
 void FCurveColorCustomization::MakeTransactional()
 {
 	if (Owner)
diff --git a/Engine/Source/Editor/DetailCustomizations/Private/CurveColorCustomization.h b/Engine/Source/Editor/DetailCustomizations/Private/CurveColorCustomization.h
--- a/Engine/Source/Editor/DetailCustomizations/Private/CurveColorCustomization.h
+++ b/Engine/Source/Editor/DetailCustomizations/Private/CurveColorCustomization.h
@@ -34,6 +34,7 @@ public:
 	virtual TArray<FRichCurveEditInfoConst> GetCurves() const override;
 	virtual TArray<FRichCurveEditInfo> GetCurves() override;
 	virtual void ModifyOwner() override;
+	virtual TArray<const UObject*> GetOwners() const override;
 	virtual void MakeTransactional() override;
 	virtual void OnCurveChanged(const TArray<FRichCurveEditInfo>& ChangedCurveEditInfos) override;
 	virtual bool IsLinearColorCurve() const override { return true; }
diff --git a/Engine/Source/Editor/DetailCustomizations/Private/CurveStructCustomization.cpp b/Engine/Source/Editor/DetailCustomizations/Private/CurveStructCustomization.cpp
--- a/Engine/Source/Editor/DetailCustomizations/Private/CurveStructCustomization.cpp
+++ b/Engine/Source/Editor/DetailCustomizations/Private/CurveStructCustomization.cpp
@@ -225,6 +225,17 @@ void FCurveStructCustomization::ModifyOwner()
 	}
 }
 
+TArray<const UObject*> FCurveStructCustomization::GetOwners() const
+{
+	TArray<const UObject*> Owners;
+	if (Owner)
+	{
+		Owners.Add(Owner);
+	}
+
+	return Owners;
+}
+
 void FCurveStructCustomization::MakeTransactional()
 {
 	if (Owner)
diff --git a/Engine/Source/Editor/DetailCustomizations/Private/DetailCustomizations.cpp b/Engine/Source/Editor/DetailCustomizations/Private/DetailCustomizations.cpp
--- a/Engine/Source/Editor/DetailCustomizations/Private/DetailCustomizations.cpp
+++ b/Engine/Source/Editor/DetailCustomizations/Private/DetailCustomizations.cpp
@@ -113,6 +113,7 @@
 #include "CameraFocusSettingsCustomization.h"
 #include "RotatorStructCustomization.h"
 #include "VectorStructCustomization.h"
+#include "Vector4StructCustomization.h"
 #include "AssetViewerSettingsCustomization.h"
 #include "MeshMergingSettingsCustomization.h"
 #include "MaterialAttributePropertyDetails.h"
@@ -168,7 +169,7 @@ void FDetailCustomizationsModule::RegisterPropertyTypeCustomizations()
 	RegisterCustomPropertyTypeLayout("DataTableCategoryHandle", FOnGetPropertyTypeCustomizationInstance::CreateStatic(&FDataTableCategoryCustomizationLayout::MakeInstance));
 	RegisterCustomPropertyTypeLayout("CurveTableRowHandle", FOnGetPropertyTypeCustomizationInstance::CreateStatic(&FCurveTableCustomizationLayout::MakeInstance));
 	RegisterCustomPropertyTypeLayout(NAME_Vector, FOnGetPropertyTypeCustomizationInstance::CreateStatic(&FVectorStructCustomization::MakeInstance));
-	RegisterCustomPropertyTypeLayout(NAME_Vector4, FOnGetPropertyTypeCustomizationInstance::CreateStatic(&FMathStructCustomization::MakeInstance));
+	RegisterCustomPropertyTypeLayout(NAME_Vector4, FOnGetPropertyTypeCustomizationInstance::CreateStatic(&FVector4StructCustomization::MakeInstance));
 	RegisterCustomPropertyTypeLayout(NAME_Vector2D, FOnGetPropertyTypeCustomizationInstance::CreateStatic(&FMathStructCustomization::MakeInstance));
 	RegisterCustomPropertyTypeLayout(NAME_IntPoint, FOnGetPropertyTypeCustomizationInstance::CreateStatic(&FMathStructCustomization::MakeInstance));
 	RegisterCustomPropertyTypeLayout(NAME_Rotator, FOnGetPropertyTypeCustomizationInstance::CreateStatic(&FRotatorStructCustomization::MakeInstance));
diff --git a/Engine/Source/Editor/DetailCustomizations/Private/DirectoryPathStructCustomization.cpp b/Engine/Source/Editor/DetailCustomizations/Private/DirectoryPathStructCustomization.cpp
--- a/Engine/Source/Editor/DetailCustomizations/Private/DirectoryPathStructCustomization.cpp
+++ b/Engine/Source/Editor/DetailCustomizations/Private/DirectoryPathStructCustomization.cpp
@@ -40,27 +40,19 @@ void FDirectoryPathStructCustomization::CustomizeHeader( TSharedRef<IPropertyHan
 
 		if(bContentDir)
 		{
-			FPathPickerConfig PathPickerConfig;
-			PathPickerConfig.bAllowContextMenu = false;
-			PathPickerConfig.OnPathSelected = FOnPathSelected::CreateSP(this, &FDirectoryPathStructCustomization::OnPathPicked, PathProperty.ToSharedRef());
-
-			FContentBrowserModule& ContentBrowserModule = FModuleManager::LoadModuleChecked<FContentBrowserModule>("ContentBrowser");
-
-			PickerWidget = SAssignNew(PickerButton, SComboButton)
+			PickerWidget = SAssignNew(PickerButton, SButton)
 			.ButtonStyle( FEditorStyle::Get(), "HoverHintOnly" )
 			.ToolTipText( LOCTEXT( "FolderComboToolTipText", "Choose a content directory") )
-			.ContentPadding( 2.0f )
+			.OnClicked(FOnClicked::CreateSP(this, &FDirectoryPathStructCustomization::OnPickContent, PathProperty.ToSharedRef()))
+			.ContentPadding(2.0f)
 			.ForegroundColor( FSlateColor::UseForeground() )
-			.IsFocusable( false )
-			.MenuContent()
+			.IsFocusable(false)
 			[
-				SNew(SBox)
-				.WidthOverride(300.0f)
-				.HeightOverride(300.0f)
-				[
-					ContentBrowserModule.Get().CreatePathPicker(PathPickerConfig)
-				]
-			];			
+				SNew(SImage)
+				.Image(FEditorStyle::GetBrush("PropertyWindow.Button_Ellipsis"))
+				.ColorAndOpacity(FSlateColor::UseForeground())
+			];
+
 		}
 		else
 		{
@@ -108,6 +100,32 @@ void FDirectoryPathStructCustomization::CustomizeChildren( TSharedRef<IPropertyH
 {
 }
 
+FReply FDirectoryPathStructCustomization::OnPickContent(TSharedRef<IPropertyHandle> PropertyHandle) 
+{
+	FContentBrowserModule& ContentBrowserModule = FModuleManager::LoadModuleChecked<FContentBrowserModule>("ContentBrowser");
+	FPathPickerConfig PathPickerConfig;
+	PathPickerConfig.bAllowContextMenu = false;
+	PathPickerConfig.OnPathSelected = FOnPathSelected::CreateSP(this, &FDirectoryPathStructCustomization::OnPathPicked, PropertyHandle);
+
+	FMenuBuilder MenuBuilder(true, NULL);
+	MenuBuilder.AddWidget(SNew(SBox)
+		.WidthOverride(300.0f)
+		.HeightOverride(300.0f)
+		[
+			ContentBrowserModule.Get().CreatePathPicker(PathPickerConfig)
+		], FText());
+
+
+	PickerMenu = FSlateApplication::Get().PushMenu(PickerButton.ToSharedRef(),
+		FWidgetPath(),
+		MenuBuilder.MakeWidget(),
+		FSlateApplication::Get().GetCursorPos(),
+		FPopupTransitionEffect(FPopupTransitionEffect::ContextMenu)
+		);
+
+	return FReply::Handled();
+}
+
 FReply FDirectoryPathStructCustomization::OnPickDirectory(TSharedRef<IPropertyHandle> PropertyHandle, const bool bRelativeToGameContentDir, const bool bUseRelativePath, const bool bLongPackageName) const
 {
 	FString Directory;
@@ -192,7 +210,11 @@ bool FDirectoryPathStructCustomization::IsValidPath(const FString& AbsolutePath,
 
 void FDirectoryPathStructCustomization::OnPathPicked(const FString& Path, TSharedRef<IPropertyHandle> PropertyHandle)
 {
-	PickerButton->SetIsOpen(false);
+	if (PickerMenu.IsValid())
+	{
+		PickerMenu->Dismiss();
+		PickerMenu.Reset();
+	}
 
 	PropertyHandle->SetValue(Path);
 }
diff --git a/Engine/Source/Editor/DetailCustomizations/Private/DirectoryPathStructCustomization.h b/Engine/Source/Editor/DetailCustomizations/Private/DirectoryPathStructCustomization.h
--- a/Engine/Source/Editor/DetailCustomizations/Private/DirectoryPathStructCustomization.h
+++ b/Engine/Source/Editor/DetailCustomizations/Private/DirectoryPathStructCustomization.h
@@ -25,6 +25,9 @@ private:
 	/** Delegate for displaying text value of path */
 	FText GetDisplayedText(TSharedRef<IPropertyHandle> PropertyHandle) const;
 
+	/** Delegate used to display a directory picker */
+	FReply OnPickContent(TSharedRef<IPropertyHandle> PropertyHandle) ;
+
 	/** Delegate used to display a directory picker */
 	FReply OnPickDirectory(TSharedRef<IPropertyHandle> PropertyHandle, const bool bRelativeToGameContentDir, const bool bUseRelativePaths, const bool bLongPackageName) const;
 
@@ -38,7 +41,10 @@ private:
 	TSharedPtr<SButton> BrowseButton;
 
 	/** The pick button widget */
-	TSharedPtr<SComboButton> PickerButton;
+	TSharedPtr<SButton> PickerButton;
+
+	/** The pick button popup menu*/
+	TSharedPtr<IMenu> PickerMenu;
 	
 	/** Absolute path to the game content directory */
 	FString AbsoluteGameContentDir;
diff --git a/Engine/Source/Editor/DetailCustomizations/Private/MathStructCustomizations.cpp b/Engine/Source/Editor/DetailCustomizations/Private/MathStructCustomizations.cpp
--- a/Engine/Source/Editor/DetailCustomizations/Private/MathStructCustomizations.cpp
+++ b/Engine/Source/Editor/DetailCustomizations/Private/MathStructCustomizations.cpp
@@ -145,7 +145,7 @@ void FMathStructCustomization::GetSortedChildren(TSharedRef<IPropertyHandle> Str
 
 
 template<typename NumericType>
-void ExtractNumericMetadata(TSharedRef<IPropertyHandle>& PropertyHandle, TOptional<NumericType>& MinValue, TOptional<NumericType>& MaxValue, TOptional<NumericType>& SliderMinValue, TOptional<NumericType>& SliderMaxValue, NumericType& SliderExponent, NumericType& Delta)
+void ExtractNumericMetadata(TSharedRef<IPropertyHandle>& PropertyHandle, TOptional<NumericType>& MinValue, TOptional<NumericType>& MaxValue, TOptional<NumericType>& SliderMinValue, TOptional<NumericType>& SliderMaxValue, NumericType& SliderExponent, NumericType& Delta, int32 &ShiftMouseMovePixelPerDelta, bool &DisplayChannelColor)
 {
 	UProperty* Property = PropertyHandle->GetProperty();
 
@@ -153,8 +153,10 @@ void ExtractNumericMetadata(TSharedRef<IPropertyHandle>& PropertyHandle, TOption
 	const FString& MetaUIMaxString = Property->GetMetaData(TEXT("UIMax"));
 	const FString& SliderExponentString = Property->GetMetaData(TEXT("SliderExponent"));
 	const FString& DeltaString = Property->GetMetaData(TEXT("Delta"));
+	const FString& ShiftMouseMovePixelPerDeltaString = Property->GetMetaData(TEXT("ShiftMouseMovePixelPerDelta"));
 	const FString& ClampMinString = Property->GetMetaData(TEXT("ClampMin"));
 	const FString& ClampMaxString = Property->GetMetaData(TEXT("ClampMax"));
+	const FString& ColorGradingModeString = Property->GetMetaData(TEXT("ColorGradingMode"));
 
 	// If no UIMin/Max was specified then use the clamp string
 	const FString& UIMinString = MetaUIMinString.Len() ? MetaUIMinString : ClampMinString;
@@ -192,6 +194,18 @@ void ExtractNumericMetadata(TSharedRef<IPropertyHandle>& PropertyHandle, TOption
 		TTypeFromString<NumericType>::FromString(Delta, *DeltaString);
 	}
 
+	ShiftMouseMovePixelPerDelta = 1;
+	if (ShiftMouseMovePixelPerDeltaString.Len())
+	{
+		TTypeFromString<int32>::FromString(ShiftMouseMovePixelPerDelta, *ShiftMouseMovePixelPerDeltaString);
+		//The value should be greater or equal to 1
+		// 1 is neutral since it is a multiplier of the mouse drag pixel
+		if (ShiftMouseMovePixelPerDelta < 1)
+		{
+			ShiftMouseMovePixelPerDelta = 1;
+		}
+	}
+
 	if (ClampMin >= ClampMax && (ClampMinString.Len() || ClampMaxString.Len()))
 	{
 		//UE_LOG(LogPropertyNode, Warning, TEXT("Clamp Min (%s) >= Clamp Max (%s) for Ranged Numeric"), *ClampMinString, *ClampMaxString);
@@ -209,6 +223,8 @@ void ExtractNumericMetadata(TSharedRef<IPropertyHandle>& PropertyHandle, TOption
 	{
 		//UE_LOG(LogPropertyNode, Warning, TEXT("UI Min (%s) >= UI Max (%s) for Ranged Numeric"), *UIMinString, *UIMaxString);
 	}
+
+	DisplayChannelColor = ColorGradingModeString.Len() > 0 ? true : false;
 }
 
 
@@ -219,12 +235,36 @@ TSharedRef<SWidget> FMathStructCustomization::MakeNumericWidget(
 {
 	TOptional<NumericType> MinValue, MaxValue, SliderMinValue, SliderMaxValue;
 	NumericType SliderExponent, Delta;
-	ExtractNumericMetadata(StructurePropertyHandle, MinValue, MaxValue, SliderMinValue, SliderMaxValue, SliderExponent, Delta);
+	int32 ShiftMouseMovePixelPerDelta = 1;
+	bool DisplayChannelColor = false;
+	ExtractNumericMetadata(StructurePropertyHandle, MinValue, MaxValue, SliderMinValue, SliderMaxValue, SliderExponent, Delta, ShiftMouseMovePixelPerDelta, DisplayChannelColor);
+
+	TSharedPtr<SWidget> LabelWidget = SNew(STextBlock)
+		.Font(IDetailLayoutBuilder::GetDetailFont())
+		.Text(PropertyHandle->GetPropertyDisplayName());
+
+	int32 ColorIndex = INDEX_NONE;
+	if (DisplayChannelColor && SortedChildHandles.Num() > 2)
+	{
+		for (int32 PropertyHandleIndex = 0; PropertyHandleIndex < SortedChildHandles.Num(); ++PropertyHandleIndex)
+		{
+			TSharedRef<IPropertyHandle> CurrentPropertyHandle = SortedChildHandles[PropertyHandleIndex];
+			if (CurrentPropertyHandle == PropertyHandle)
+			{
+				ColorIndex = PropertyHandleIndex;
+				break;
+			}
+		}
+	}
 
 	TWeakPtr<IPropertyHandle> WeakHandlePtr = PropertyHandle;
 
+	const FEditableTextBoxStyle& DarkEditorStyle = FCoreStyle::Get().GetWidgetStyle<FEditableTextBoxStyle>("DarkEditableTextBox");
+	const FEditableTextBoxStyle& NormalEditorStyle = FCoreStyle::Get().GetWidgetStyle<FEditableTextBoxStyle>("NormalEditableTextBox");
+
 	return SNew(SNumericEntryBox<NumericType>)
 		.IsEnabled(this, &FMathStructCustomization::IsValueEnabled, WeakHandlePtr)
+		.EditableTextBoxStyle(DisplayChannelColor ? &DarkEditorStyle : &NormalEditorStyle)
 		.Value(this, &FMathStructCustomization::OnGetValue, WeakHandlePtr)
 		.Font(IDetailLayoutBuilder::GetDetailFont())
 		.UndeterminedString(NSLOCTEXT("PropertyEditor", "MultipleValues", "Multiple Values"))
@@ -235,6 +275,9 @@ TSharedRef<SWidget> FMathStructCustomization::MakeNumericWidget(
 		.LabelVAlign(VAlign_Center)
 		// Only allow spin on handles with one object.  Otherwise it is not clear what value to spin
 		.AllowSpin(PropertyHandle->GetNumOuterObjects() == 1)
+		.UseDarkStyle(DisplayChannelColor)
+		.ColorIndex(ColorIndex)
+		.ShiftMouseMovePixelPerDelta(ShiftMouseMovePixelPerDelta)
 		.MinValue(MinValue)
 		.MaxValue(MaxValue)
 		.MinSliderValue(SliderMinValue)
@@ -243,13 +286,10 @@ TSharedRef<SWidget> FMathStructCustomization::MakeNumericWidget(
 		.Delta(Delta)
 		.Label()
 		[
-			SNew(STextBlock)
-				.Font(IDetailLayoutBuilder::GetDetailFont())
-				.Text(PropertyHandle->GetPropertyDisplayName())
+			LabelWidget.ToSharedRef()
 		];
 }
 
-
 TSharedRef<SWidget> FMathStructCustomization::MakeChildWidget(
 	TSharedRef<IPropertyHandle>& StructurePropertyHandle,
 	TSharedRef<IPropertyHandle>& PropertyHandle)
diff --git a/Engine/Source/Editor/DetailCustomizations/Private/PhysicsSettingsDetails.cpp b/Engine/Source/Editor/DetailCustomizations/Private/PhysicsSettingsDetails.cpp
--- a/Engine/Source/Editor/DetailCustomizations/Private/PhysicsSettingsDetails.cpp
+++ b/Engine/Source/Editor/DetailCustomizations/Private/PhysicsSettingsDetails.cpp
@@ -10,79 +10,47 @@
 #include "DetailCategoryBuilder.h"
 #include "Widgets/SToolTip.h"
 #include "IDocumentation.h"
+#include "IDetailChildrenBuilder.h"
 
 #define LOCTEXT_NAMESPACE "PhysicalSurfaceDetails"
 
-//====================================================================================
-// SPhysicalSurfaceListItem 
-//=====================================================================================
-
-void SPhysicalSurfaceListItem::Construct(const FArguments& InArgs, const TSharedRef<STableViewBase>& InOwnerTableView)
+void SPhysicalSurfaceEditBox::Construct(const FArguments& InArgs)
 {
 	PhysicalSurface = InArgs._PhysicalSurface;
 	PhysicalSurfaceEnum = InArgs._PhysicalSurfaceEnum;
 	OnCommitChange = InArgs._OnCommitChange;
 	check(PhysicalSurface.IsValid() && PhysicalSurfaceEnum);
-	SMultiColumnTableRow< TSharedPtr<FPhysicalSurfaceListItem> >::Construct(FSuperRowType::FArguments(), InOwnerTableView);
-}
 
-FString SPhysicalSurfaceListItem::GetTypeString() const
-{
-	// i need to convert to enum name
-	int32 EnumIndex = (int32)PhysicalSurface->Type;
-	return PhysicalSurfaceEnum->GetEnumName(EnumIndex);
+	ChildSlot
+	[
+		SAssignNew(NameEditBox, SEditableTextBox)
+		.Text(this, &SPhysicalSurfaceEditBox::GetName)
+		.Font(IDetailLayoutBuilder::GetDetailFont())
+		.OnTextCommitted(this, &SPhysicalSurfaceEditBox::NewNameEntered)
+		.OnTextChanged(this, &SPhysicalSurfaceEditBox::OnTextChanged)
+		.IsReadOnly(PhysicalSurface->Type == SurfaceType_Default)
+		.SelectAllTextWhenFocused(true)
+	];
 }
 
-TSharedRef<SWidget> SPhysicalSurfaceListItem::GenerateWidgetForColumn(const FName& ColumnName)
-{
-	if (ColumnName == TEXT("Type"))
-	{
-		return	SNew(SBox)
-			.HeightOverride(20)
-			.Padding(FMargin(5, 2))
-			.VAlign(VAlign_Center)
-			[
-				SNew(STextBlock)
-				.Text(FText::FromString(GetTypeString()))
-				.Font(IDetailLayoutBuilder::GetDetailFont())
-			];
-	}
-	else if (ColumnName == TEXT("Name"))
-	{
-		return	SNew(SBox)
-			.HeightOverride(20)
-			.Padding(FMargin(5, 2))
-			.VAlign(VAlign_Center)
-			[
-				SAssignNew(NameBox, SEditableTextBox)
-				.MinDesiredWidth(128.0f)
-				.Text(this, &SPhysicalSurfaceListItem::GetName)
-				.Font(IDetailLayoutBuilder::GetDetailFont())
-				.OnTextCommitted(this, &SPhysicalSurfaceListItem::NewNameEntered)
-				.OnTextChanged(this, &SPhysicalSurfaceListItem::OnTextChanged)
-				.IsReadOnly(PhysicalSurface->Type == SurfaceType_Default)
-			];
-	}
 
-	return SNullWidget::NullWidget;
-}
 
-void SPhysicalSurfaceListItem::OnTextChanged(const FText& NewText)
+void SPhysicalSurfaceEditBox::OnTextChanged(const FText& NewText)
 {
 	FString NewName = NewText.ToString();
 
 	if(NewName.Find(TEXT(" "))!=INDEX_NONE)
 	{
 		// no white space
-		NameBox->SetError(TEXT("No white space is allowed"));
+		NameEditBox->SetError(TEXT("No white space is allowed"));
 	}
 	else
 	{
-		NameBox->SetError(TEXT(""));
+		NameEditBox->SetError(TEXT(""));
 	}
 }
 
-void SPhysicalSurfaceListItem::NewNameEntered(const FText& NewText, ETextCommit::Type CommitInfo)
+void SPhysicalSurfaceEditBox::NewNameEntered(const FText& NewText, ETextCommit::Type CommitInfo)
 {
 	// Don't digest the number if we just clicked away from the pop-up
 	if((CommitInfo == ETextCommit::OnEnter) || (CommitInfo == ETextCommit::OnUserMovedFocus))
@@ -107,224 +75,240 @@ void SPhysicalSurfaceListItem::NewNameEntered(const FText& NewText, ETextCommit:
 		else
 		{
 			// clear error
-			NameBox->SetError(TEXT(""));
+			NameEditBox->SetError(TEXT(""));
 		}
 	}
 }
 
-FText SPhysicalSurfaceListItem::GetName() const
+FText SPhysicalSurfaceEditBox::GetName() const
 {
 	return FText::FromName(PhysicalSurface->Name);
 }
 
-//====================================================================================
-// FPhysicsSettingsDetails
-//=====================================================================================
-TSharedRef<IDetailCustomization> FPhysicsSettingsDetails::MakeInstance()
+class FPhysicalSurfaceList : public IDetailCustomNodeBuilder, public TSharedFromThis<FPhysicalSurfaceList>
 {
-	return MakeShareable( new FPhysicsSettingsDetails );
-}
+public:
+	FPhysicalSurfaceList(UPhysicsSettings* InPhysicsSettings, UEnum* InPhysicalSurfaceEnum, TSharedPtr<IPropertyHandle>& InPhysicalSurfacesProperty )
+		: PhysicsSettings(InPhysicsSettings)
+		, PhysicalSurfaceEnum(InPhysicalSurfaceEnum)
+		, PhysicalSurfacesProperty(InPhysicalSurfacesProperty)
+	{}
+
+	void RefreshPhysicalSurfaceList()
+	{
+		// make sure no duplicate exists
+		// if exists, use the last one
+		for(auto Iter = PhysicsSettings->PhysicalSurfaces.CreateIterator(); Iter; ++Iter)
+		{
+			for(auto InnerIter = Iter+1; InnerIter; ++InnerIter)
+			{
+				// see if same type
+				if(Iter->Type == InnerIter->Type)
+				{
+					// remove the current one
+					PhysicsSettings->PhysicalSurfaces.RemoveAt(Iter.GetIndex());
+					--Iter;
+					break;
+				}
+			}
+		}
 
-void FPhysicsSettingsDetails::CustomizeDetails( IDetailLayoutBuilder& DetailBuilder )
-{
-  	IDetailCategoryBuilder& PhysicalSurfaceCategory = DetailBuilder.EditCategory("Physical Surface", FText::GetEmpty(), ECategoryPriority::Uncommon);
+		bool bCreatedItem[SurfaceType_Max];
+		FGenericPlatformMemory::Memzero(bCreatedItem, sizeof(bCreatedItem));
 
-	PhysicsSettings= UPhysicsSettings::Get();
-	check(PhysicsSettings);
+		PhysicalSurfaceList.Empty();
 
-	PhysicalSurfaceEnum = FindObject<UEnum>(ANY_PACKAGE, TEXT("EPhysicalSurface"), true);
-	check(PhysicalSurfaceEnum);
+		// I'm listing all of these because it is easier for users to understand how does this work. 
+		// I can't just link behind of scene because internally it will only save the enum
+		// for example if you name SurfaceType5 to be Water and later changed to Sand, everything that used
+		// SurfaceType5 will changed to Sand
+		// I think what might be better is to show what options they have, and it's for them to choose how to name
 
-	RefreshPhysicalSurfaceList();
+		// add the first one by default
+		{
+			bCreatedItem[0] = true;
+			PhysicalSurfaceList.Add(MakeShareable(new FPhysicalSurfaceListItem(MakeShareable(new FPhysicalSurfaceName(SurfaceType_Default, TEXT("Default"))))));
+		}
 
-	const FString PhysicalSurfaceDocLink = TEXT("Shared/Physics");
-	TSharedPtr<SToolTip> PhysicalSurfaceTooltip = IDocumentation::Get()->CreateToolTip(LOCTEXT("PhysicalSurface", "Edit physical surface."), NULL, PhysicalSurfaceDocLink, TEXT("PhysicalSurface"));
+		// we don't create the first one. First one is always default. 
+		for(auto Iter = PhysicsSettings->PhysicalSurfaces.CreateIterator(); Iter; ++Iter)
+		{
+			bCreatedItem[Iter->Type] = true;
+			PhysicalSurfaceList.Add(MakeShareable(new FPhysicalSurfaceListItem(MakeShareable(new FPhysicalSurfaceName(*Iter)))));
+		}
 
-	// Customize collision section
-	PhysicalSurfaceCategory.AddCustomRow(LOCTEXT("FPhysicsSettingsDetails_PhysicalSurface", "Physical Surface"))
-	[
-		SNew(SVerticalBox)
-
-		+SVerticalBox::Slot()
-		.Padding(5)
-		.AutoHeight()
-		[
-			SNew(STextBlock)
-			.Font(IDetailLayoutBuilder::GetDetailFont())
-			.ToolTip(PhysicalSurfaceTooltip)
-			.AutoWrapText(true)
-			.Text(LOCTEXT("PhysicalSurface_Menu_Description", " You can have up to 62 custom surface types for your project. \nOnce you name each type, they will show up as surface type in the physical material."))
-		]
-
-		+SVerticalBox::Slot()
-		.Padding(5)
-		.FillHeight(1)
-		[
-			SAssignNew(PhysicalSurfaceListView, SPhysicalSurfaceListView)
-			.ItemHeight(15.f)
-			.ListItemsSource(&PhysicalSurfaceList)
-			.OnGenerateRow(this, &FPhysicsSettingsDetails::HandleGenerateListWidget)
-			.SelectionMode(ESelectionMode::None)
-
-			.HeaderRow(
-				SNew(SHeaderRow)
-				+ SHeaderRow::Column("Type")
-				.HAlignCell(HAlign_Left)
-				.FillWidth(1)
-				.HeaderContentPadding(FMargin(0, 3))
-				[
-					SNew(SHorizontalBox)
-					+SHorizontalBox::Slot()
-					.FillWidth(1)
-					[
-						SNew(STextBlock)
-						.Text(LOCTEXT("CPhysicalSurfaceListHeader_Type", "Type"))
-						.Font(IDetailLayoutBuilder::GetDetailFont())
-					]
-				]
-				// Name
-				+ SHeaderRow::Column("Name")
-				.HAlignCell(HAlign_Left)
-				.FillWidth(1)
-				.HeaderContentPadding(FMargin(0, 3))
-				[
-					SNew(SHorizontalBox)
-					+SHorizontalBox::Slot()
-					.FillWidth(1)
-					[
-						SNew(STextBlock)
-						.Text(LOCTEXT("PhysicalSurfaceListHeader_Name", "Name"))
-						.Font(IDetailLayoutBuilder::GetDetailFont())
-					]
-				]
-			)
-		]
-	];
-}
+		for(int32 Index = (int32)SurfaceType1; Index < SurfaceType_Max; ++Index)
+		{
+			if(bCreatedItem[Index] == false)
+			{
+				FPhysicalSurfaceName NeweElement((EPhysicalSurface)Index, TEXT(""));
+				PhysicalSurfaceList.Add(MakeShareable(new FPhysicalSurfaceListItem(MakeShareable(new FPhysicalSurfaceName(NeweElement)))));
+			}
+		}
 
-void FPhysicsSettingsDetails::UpdatePhysicalSurfaceList()
-{
-	RefreshPhysicalSurfaceList();
-	PhysicalSurfaceListView->RequestListRefresh();
-	PhysicsSettings->LoadSurfaceType();
-	PhysicsSettings->UpdateDefaultConfigFile();
-}
+		// sort PhysicalSurfaceList by Type
 
-void FPhysicsSettingsDetails::RefreshPhysicalSurfaceList()
-{
-	// make sure no duplicate exists
-	// if exists, use the last one
-	for(auto Iter = PhysicsSettings->PhysicalSurfaces.CreateIterator(); Iter; ++Iter)
-	{
-		for(auto InnerIter = Iter+1; InnerIter; ++InnerIter)
+		struct FComparePhysicalSurface
 		{
-			// see if same type
-			if(Iter->Type == InnerIter->Type)
+			FORCEINLINE bool operator()(const TSharedPtr<FPhysicalSurfaceListItem> A, const TSharedPtr<FPhysicalSurfaceListItem> B) const
 			{
-				// remove the current one
-				PhysicsSettings->PhysicalSurfaces.RemoveAt(Iter.GetIndex());
-				--Iter;
-				break;
+				check(A.IsValid());
+				check(B.IsValid());
+				return A->PhysicalSurface->Type < B->PhysicalSurface->Type;
 			}
-		}
-	}
+		};
 
-	bool bCreatedItem[SurfaceType_Max];
-	FGenericPlatformMemory::Memzero(bCreatedItem, sizeof(bCreatedItem));
+		PhysicalSurfaceList.Sort(FComparePhysicalSurface());
 
-	PhysicalSurfaceList.Empty();
+		PhysicsSettings->LoadSurfaceType();
+		PhysicsSettings->UpdateDefaultConfigFile();
 
-	// I'm listing all of these because it is easier for users to understand how does this work. 
-	// I can't just link behind of scene because internally it will only save the enum
-	// for example if you name SurfaceType5 to be Water and later changed to Sand, everything that used
-	// SurfaceType5 will changed to Sand
-	// I think what might be better is to show what options they have, and it's for them to choose how to name
+		RegenerateChildren.ExecuteIfBound();
+	}
 
-	// add the first one by default
+	virtual void SetOnRebuildChildren( FSimpleDelegate InOnRegenerateChildren ) override
 	{
-		bCreatedItem[0] = true;
-		PhysicalSurfaceList.Add(MakeShareable(new FPhysicalSurfaceListItem(MakeShareable(new FPhysicalSurfaceName(SurfaceType_Default, TEXT("Default"))))));
+		RegenerateChildren = InOnRegenerateChildren;
 	}
 
-	// we don't create the first one. First one is always default. 
-	for(auto Iter = PhysicsSettings->PhysicalSurfaces.CreateIterator(); Iter; ++Iter)
+	virtual void GenerateHeaderRowContent( FDetailWidgetRow& NodeRow ) override
 	{
-		bCreatedItem[Iter->Type] = true;
-		PhysicalSurfaceList.Add(MakeShareable(new FPhysicalSurfaceListItem(MakeShareable(new FPhysicalSurfaceName(*Iter)))));
+		// no header row
 	}
 
-	for(int32 Index = (int32)SurfaceType1; Index < SurfaceType_Max; ++Index)
+	virtual void GenerateChildContent( IDetailChildrenBuilder& ChildrenBuilder ) override
 	{
-		if(bCreatedItem[Index] == false)
+		FText SearchString = LOCTEXT("FPhysicsSettingsDetails_PhysicalSurface", "Physical Surface");
+
+		for(TSharedPtr<FPhysicalSurfaceListItem>& Item : PhysicalSurfaceList)
 		{
-			FPhysicalSurfaceName NeweElement((EPhysicalSurface)Index, TEXT(""));
-			PhysicalSurfaceList.Add(MakeShareable(new FPhysicalSurfaceListItem(MakeShareable(new FPhysicalSurfaceName(NeweElement)))));
-		}
-	}
+			FDetailWidgetRow& Row = ChildrenBuilder.AddChildContent(SearchString);
 
-	// sort PhysicalSurfaceList by Type
 
-	struct FComparePhysicalSurface
-	{
-		FORCEINLINE bool operator()(const TSharedPtr<FPhysicalSurfaceListItem> A, const TSharedPtr<FPhysicalSurfaceListItem> B) const
-		{
-			check(A.IsValid());
-			check(B.IsValid());
-			return A->PhysicalSurface->Type < B->PhysicalSurface->Type;
-		}
-	};
+			FString TypeString = PhysicalSurfaceEnum->GetEnumName((int32)Item->PhysicalSurface->Type);
 
-	PhysicalSurfaceList.Sort(FComparePhysicalSurface());
-}
+			Row.NameContent()
+			[
+				SNew(STextBlock)
+				.Text(FText::FromString(TypeString))
+				.Font(IDetailLayoutBuilder::GetDetailFont())
+			];
 
-TSharedRef<ITableRow> FPhysicsSettingsDetails::HandleGenerateListWidget(TSharedPtr< FPhysicalSurfaceListItem> InItem, const TSharedRef<STableViewBase>& OwnerTable)
-{
-	return SNew(SPhysicalSurfaceListItem, OwnerTable)
-		.PhysicalSurface(InItem->PhysicalSurface)
-		.PhysicalSurfaceEnum(PhysicalSurfaceEnum)
-		.OnCommitChange(this, &FPhysicsSettingsDetails::OnCommitChange);
-}
+			Row.ValueContent()
+			[
+				SNew(SPhysicalSurfaceEditBox)
+				.PhysicalSurface(Item->PhysicalSurface)
+				.PhysicalSurfaceEnum(PhysicalSurfaceEnum)
+				.OnCommitChange(this, &FPhysicalSurfaceList::OnCommitChange)
+			];
+		}
+	}
 
-void FPhysicsSettingsDetails::OnCommitChange()
-{
-	bool bDoCommit=true;
-	// make sure it verifies all data is correct
-	// skip the first one
-	for(auto Iter = PhysicalSurfaceList.CreateConstIterator()+1; Iter; ++Iter)
+	virtual void Tick( float DeltaTime ) override {}
+	virtual bool RequiresTick() const override { return false; }
+	virtual bool InitiallyCollapsed() const { return false; }
+	virtual FName GetName() const override
 	{
-		TSharedPtr<FPhysicalSurfaceListItem> ListItem = *Iter;
-		if(ListItem->PhysicalSurface->Name != NAME_None)
+		static const FName Name(TEXT("PhysicalSurfaceList"));
+		return Name;
+	}
+private:
+
+	void OnCommitChange()
+	{
+		bool bDoCommit=true;
+		// make sure it verifies all data is correct
+		// skip the first one
+		for(auto Iter = PhysicalSurfaceList.CreateConstIterator()+1; Iter; ++Iter)
 		{
-			// make sure no same name exists
-			for(auto InnerIter= Iter+1; InnerIter; ++InnerIter)
+			TSharedPtr<FPhysicalSurfaceListItem> ListItem = *Iter;
+			if(ListItem->PhysicalSurface->Name != NAME_None)
 			{
-				TSharedPtr<FPhysicalSurfaceListItem> InnerItem = *InnerIter;
-				if(ListItem->PhysicalSurface->Name == InnerItem->PhysicalSurface->Name)
+				// make sure no same name exists
+				for(auto InnerIter= Iter+1; InnerIter; ++InnerIter)
 				{
-					// duplicate name, warn uer and get out
-					FMessageDialog::Open(EAppMsgType::Ok, LOCTEXT("FPhysicsSettingsDetails_InvalidName", "Duplicate name found."));
-					bDoCommit = false;
-					break;
+					TSharedPtr<FPhysicalSurfaceListItem> InnerItem = *InnerIter;
+					if(ListItem->PhysicalSurface->Name == InnerItem->PhysicalSurface->Name)
+					{
+						// duplicate name, warn user and get out
+						FMessageDialog::Open(EAppMsgType::Ok, LOCTEXT("FPhysicsSettingsDetails_InvalidName", "Duplicate name found."));
+						bDoCommit = false;
+						break;
+					}
 				}
 			}
 		}
-	}
 
-	if(bDoCommit)
-	{
-		PhysicsSettings->PhysicalSurfaces.Empty();
-		for(auto Iter = PhysicalSurfaceList.CreateConstIterator()+1; Iter; ++Iter)
+		if(bDoCommit)
 		{
-			TSharedPtr<FPhysicalSurfaceListItem> ListItem = *Iter;
-			if(ListItem->PhysicalSurface->Name != NAME_None)
+			PhysicalSurfacesProperty->NotifyPreChange();
+
+			PhysicsSettings->PhysicalSurfaces.Empty();
+			for(auto Iter = PhysicalSurfaceList.CreateConstIterator()+1; Iter; ++Iter)
 			{
-				PhysicsSettings->PhysicalSurfaces.Add(FPhysicalSurfaceName(ListItem->PhysicalSurface->Type, ListItem->PhysicalSurface->Name));
+				TSharedPtr<FPhysicalSurfaceListItem> ListItem = *Iter;
+				if(ListItem->PhysicalSurface->Name != NAME_None)
+				{
+					PhysicsSettings->PhysicalSurfaces.Add(FPhysicalSurfaceName(ListItem->PhysicalSurface->Type, ListItem->PhysicalSurface->Name));
+				}
 			}
+
+			PhysicsSettings->UpdateDefaultConfigFile();
+
+			PhysicalSurfacesProperty->NotifyPostChange();
 		}
 	}
+private:
+	FSimpleDelegate RegenerateChildren;
+	TArray< TSharedPtr< FPhysicalSurfaceListItem > >	PhysicalSurfaceList;
+	UPhysicsSettings* PhysicsSettings;
+	UEnum* PhysicalSurfaceEnum;
+	TSharedPtr<IPropertyHandle> PhysicalSurfacesProperty;
+};
+
 
-	// if so, refresh window
-	UpdatePhysicalSurfaceList();
+//====================================================================================
+// FPhysicsSettingsDetails
+//=====================================================================================
+TSharedRef<IDetailCustomization> FPhysicsSettingsDetails::MakeInstance()
+{
+	return MakeShareable( new FPhysicsSettingsDetails );
 }
+
+void FPhysicsSettingsDetails::CustomizeDetails( IDetailLayoutBuilder& DetailBuilder )
+{
+  	IDetailCategoryBuilder& PhysicalSurfaceCategory = DetailBuilder.EditCategory("Physical Surface", FText::GetEmpty(), ECategoryPriority::Uncommon);
+
+	PhysicsSettings= UPhysicsSettings::Get();
+	check(PhysicsSettings);
+
+	PhysicalSurfaceEnum = FindObject<UEnum>(ANY_PACKAGE, TEXT("EPhysicalSurface"), true);
+	check(PhysicalSurfaceEnum);
+
+	TSharedPtr<IPropertyHandle> PhysicalSurfacesProperty = DetailBuilder.GetProperty(GET_MEMBER_NAME_CHECKED(UPhysicsSettings, PhysicalSurfaces));
+
+	TSharedRef<FPhysicalSurfaceList> PhysicalSurfaceListCustomization = MakeShareable(new FPhysicalSurfaceList(PhysicsSettings, PhysicalSurfaceEnum, PhysicalSurfacesProperty) );
+	PhysicalSurfaceListCustomization->RefreshPhysicalSurfaceList();
+
+	const FString PhysicalSurfaceDocLink = TEXT("Shared/Physics");
+	TSharedPtr<SToolTip> PhysicalSurfaceTooltip = IDocumentation::Get()->CreateToolTip(LOCTEXT("PhysicalSurface", "Edit physical surface."), NULL, PhysicalSurfaceDocLink, TEXT("PhysicalSurface"));
+
+
+	// Customize collision section
+	PhysicalSurfaceCategory.AddCustomRow(LOCTEXT("FPhysicsSettingsDetails_PhysicalSurface", "Physical Surface"))
+	[
+		SNew(STextBlock)
+		.Font(IDetailLayoutBuilder::GetDetailFont())
+		.ToolTip(PhysicalSurfaceTooltip)
+		.AutoWrapText(true)
+		.Text(LOCTEXT("PhysicalSurface_Menu_Description", " You can have up to 62 custom surface types for your project. \nOnce you name each type, they will show up as surface type in the physical material."))
+	];
+
+
+	PhysicalSurfaceCategory.AddCustomBuilder(PhysicalSurfaceListCustomization);
+}
+
+
+
 #undef LOCTEXT_NAMESPACE
 
diff --git a/Engine/Source/Editor/DetailCustomizations/Private/PhysicsSettingsDetails.h b/Engine/Source/Editor/DetailCustomizations/Private/PhysicsSettingsDetails.h
--- a/Engine/Source/Editor/DetailCustomizations/Private/PhysicsSettingsDetails.h
+++ b/Engine/Source/Editor/DetailCustomizations/Private/PhysicsSettingsDetails.h
@@ -34,12 +34,11 @@ public:
 /**
 * Implements the FriendsList
 */
-class SPhysicalSurfaceListItem
-	: public SMultiColumnTableRow< TSharedPtr<class FPhysicalSurfaceListItem> >
+class SPhysicalSurfaceEditBox : public SCompoundWidget
 {
 public:
 
-	SLATE_BEGIN_ARGS(SPhysicalSurfaceListItem) { }
+	SLATE_BEGIN_ARGS(SPhysicalSurfaceEditBox) { }
 		SLATE_ARGUMENT(TSharedPtr<FPhysicalSurfaceName>, PhysicalSurface)
 		SLATE_ARGUMENT(UEnum*, PhysicalSurfaceEnum)
 		SLATE_EVENT(FOnCommitChange, OnCommitChange)
@@ -52,8 +51,8 @@ public:
 	*
 	* @param InArgs - The Slate argument list.
 	*/
-	void Construct(const FArguments& InArgs, const TSharedRef<STableViewBase>& InOwnerTableView);
-	virtual TSharedRef<SWidget> GenerateWidgetForColumn(const FName& ColumnName) override;
+	void Construct(const FArguments& InArgs);
+	
 	FString GetTypeString() const;
 
 	FText GetName() const;
@@ -64,7 +63,7 @@ private:
 	TSharedPtr<FPhysicalSurfaceName>	PhysicalSurface;
 	UEnum *								PhysicalSurfaceEnum;
 	FOnCommitChange						OnCommitChange;
-	TSharedPtr<SEditableTextBox>		NameBox;
+	TSharedPtr<SEditableTextBox>		NameEditBox;
 };
 
 typedef  SListView< TSharedPtr< FPhysicalSurfaceListItem > > SPhysicalSurfaceListView;
@@ -90,14 +89,11 @@ private:
 
 
 private:
-	TSharedPtr<SPhysicalSurfaceListView>				PhysicalSurfaceListView;
 	TArray< TSharedPtr< FPhysicalSurfaceListItem > >	PhysicalSurfaceList;
 
 	UPhysicsSettings *							PhysicsSettings;
 	UEnum*										PhysicalSurfaceEnum;
 	// functions
-	void UpdatePhysicalSurfaceList();
-	void RefreshPhysicalSurfaceList();
 	void OnCommitChange();
 };
 
diff --git a/Engine/Source/Editor/DetailCustomizations/Private/PostProcessSettingsCustomization.cpp b/Engine/Source/Editor/DetailCustomizations/Private/PostProcessSettingsCustomization.cpp
--- a/Engine/Source/Editor/DetailCustomizations/Private/PostProcessSettingsCustomization.cpp
+++ b/Engine/Source/Editor/DetailCustomizations/Private/PostProcessSettingsCustomization.cpp
@@ -20,13 +20,32 @@
 #include "PropertyCustomizationHelpers.h"
 #include "ObjectEditorUtils.h"
 #include "Widgets/Layout/SWidgetSwitcher.h"
+#include "DetailCategoryBuilder.h"
+#include "DetailLayoutBuilder.h"
 
 #define LOCTEXT_NAMESPACE "PostProcessSettingsCustomization"
 
-void FPostProcessSettingsCustomization::CustomizeChildren( TSharedRef<IPropertyHandle> StructPropertyHandle, class IDetailChildrenBuilder& StructBuilder, IPropertyTypeCustomizationUtils& StructCustomizationUtils )
+struct FPostProcessGroup
 {
-	TMap<FName, IDetailGroup*> CategoryNameToGroupMap;
+	FString RawGroupName;
+	FString DisplayName;
+	IDetailCategoryBuilder* RootCategory;
+	TArray<TSharedPtr<IPropertyHandle>> SimplePropertyHandles;
+	TArray<TSharedPtr<IPropertyHandle>> AdvancedPropertyHandles;
+
+	bool IsValid() const
+	{
+		return !RawGroupName.IsEmpty() && !DisplayName.IsEmpty() && RootCategory;
+	}
 
+	FPostProcessGroup()
+		: RootCategory(nullptr)
+	{}
+};
+
+
+void FPostProcessSettingsCustomization::CustomizeChildren( TSharedRef<IPropertyHandle> StructPropertyHandle, class IDetailChildrenBuilder& StructBuilder, IPropertyTypeCustomizationUtils& StructCustomizationUtils )
+{
 	uint32 NumChildren = 0;
 	FPropertyAccess::Result Result = StructPropertyHandle->GetNumChildren(NumChildren);
 
@@ -36,32 +55,114 @@ void FPostProcessSettingsCustomization::CustomizeChildren( TSharedRef<IPropertyH
 	// a category with this name should be one level higher, should be "PostProcessSettings"
 	FName ClassName = StructProp->Struct->GetFName();
 
-	if( Result == FPropertyAccess::Success && NumChildren > 0)
+	// Create new categories in the parent layout rather than adding all post process settings to one category
+	IDetailLayoutBuilder& LayoutBuilder = StructBuilder.GetParentCategory().GetParentLayout();
+
+	TMap<FString, IDetailCategoryBuilder*> NameToCategoryBuilderMap;
+	TMap<FString, FPostProcessGroup> NameToGroupMap;
+
+	static const auto VarTonemapperFilm = IConsoleManager::Get().FindTConsoleVariableDataInt(TEXT("r.TonemapperFilm"));
+	static const FName LegacyTonemapperName("LegacyTonemapper");
+	static const FName TonemapperCategory("Tonemapper");
+
+	const bool bUseNewTonemapper = VarTonemapperFilm->GetValueOnGameThread() == 1;
+	if(Result == FPropertyAccess::Success && NumChildren > 0)
 	{
 		for( uint32 ChildIndex = 0; ChildIndex < NumChildren; ++ChildIndex )
 		{
 			TSharedPtr<IPropertyHandle> ChildHandle = StructPropertyHandle->GetChildHandle( ChildIndex );
+
+			
 			if( ChildHandle.IsValid() && ChildHandle->GetProperty() )
 			{
 				UProperty* Property = ChildHandle->GetProperty();
-				FName Category = FObjectEditorUtils::GetCategoryFName( Property );
 
-				if(Category == ClassName)
+				FName CategoryFName = FObjectEditorUtils::GetCategoryFName(Property);
+			
+				if(CategoryFName == TonemapperCategory)
+				{
+					if(ChildHandle->HasMetaData(LegacyTonemapperName) && bUseNewTonemapper)
+					{
+						// Hide because property is for legacy tonemapper and we are using the new tonemapper
+						ChildHandle->MarkHiddenByCustomization();
+						continue;
+					}
+					else if(!ChildHandle->HasMetaData(LegacyTonemapperName) && !bUseNewTonemapper)
+					{
+						// Hide because property is for new tone mapper and we are using old tone mapper
+						ChildHandle->MarkHiddenByCustomization();
+						continue;
+					}
+				}
+				
+				FString RawCategoryName = CategoryFName.ToString();
+
+				TArray<FString> CategoryAndGroups;
+				RawCategoryName.ParseIntoArray(CategoryAndGroups, TEXT("|"), 1);
+
+				FString RootCategoryName = CategoryAndGroups.Num() > 0 ? CategoryAndGroups[0] : RawCategoryName;
+
+				IDetailCategoryBuilder* Category = NameToCategoryBuilderMap.FindRef(RootCategoryName);
+				if(!Category)
+				{
+					IDetailCategoryBuilder& NewCategory = LayoutBuilder.EditCategory(*RootCategoryName, FText::GetEmpty(), ECategoryPriority::TypeSpecific);
+					NameToCategoryBuilderMap.Add(RootCategoryName, &NewCategory);
+
+					Category = &NewCategory;
+				}
+
+				if(CategoryAndGroups.Num() > 1)
 				{
-					// Some elements should be outside of the categories, they need to specify the struct nama as category
-					StructBuilder.AddChildProperty( ChildHandle.ToSharedRef() );
+					// Only handling one group for now
+					// There are sub groups so add them now
+					FPostProcessGroup& PPGroup = NameToGroupMap.FindOrAdd(RawCategoryName);
+					
+					// Is this a new group? It wont be valid if it is
+					if(!PPGroup.IsValid())
+					{
+						PPGroup.RootCategory = Category;
+						PPGroup.RawGroupName = RawCategoryName;
+						PPGroup.DisplayName = CategoryAndGroups[1].Trim().TrimTrailing();
+					}
+	
+					bool bIsSimple = !ChildHandle->GetProperty()->HasAnyPropertyFlags(CPF_AdvancedDisplay);
+					if(bIsSimple)
+					{
+						PPGroup.SimplePropertyHandles.Add(ChildHandle);
+					}
+					else
+					{
+						PPGroup.AdvancedPropertyHandles.Add(ChildHandle);
+					}
 				}
 				else
 				{
-					IDetailGroup* Group = CategoryNameToGroupMap.FindRef( Category );
+					Category->AddProperty(ChildHandle.ToSharedRef());
+				}
+
+			}
+		}
+
+		for(auto& NameAndGroup : NameToGroupMap)
+		{
+			const FPostProcessGroup& PPGroup = NameAndGroup.Value;
+
+			if(PPGroup.SimplePropertyHandles.Num() > 0 || PPGroup.AdvancedPropertyHandles.Num() > 0 )
+			{
+				IDetailGroup& SimpleGroup = PPGroup.RootCategory->AddGroup(*PPGroup.RawGroupName, FText::FromString(PPGroup.DisplayName));
+				for(auto& SimpleProperty : PPGroup.SimplePropertyHandles)
+				{
+					SimpleGroup.AddPropertyRow(SimpleProperty.ToSharedRef());
+				}
 
-					if( !Group )
+				if(PPGroup.AdvancedPropertyHandles.Num() > 0)
+				{
+					IDetailGroup& AdvancedGroup = SimpleGroup.AddGroup(*(PPGroup.RawGroupName+TEXT("Advanced")), LOCTEXT("PostProcessAdvancedGroup", "Advanced"));
+					
+					for(auto& AdvancedProperty : PPGroup.AdvancedPropertyHandles)
 					{
-						Group = &StructBuilder.AddChildGroup( Category, FText::FromString( FName::NameToDisplayString( Category.ToString(), false ) ) );
-						CategoryNameToGroupMap.Add( Category, Group );
+						AdvancedGroup.AddPropertyRow(AdvancedProperty.ToSharedRef());
 					}
-				
-					Group->AddPropertyRow( ChildHandle.ToSharedRef() );
 				}
 			}
 		}
@@ -72,15 +173,7 @@ void FPostProcessSettingsCustomization::CustomizeChildren( TSharedRef<IPropertyH
 
 void FPostProcessSettingsCustomization::CustomizeHeader( TSharedRef<IPropertyHandle> StructPropertyHandle, class FDetailWidgetRow& HeaderRow, IPropertyTypeCustomizationUtils& StructCustomizationUtils )
 {
-	HeaderRow.NameContent()
-	[
-		StructPropertyHandle->CreatePropertyNameWidget()
-	];
-
-	HeaderRow.ValueContent()
-	[
-		StructPropertyHandle->CreatePropertyValueWidget()
-	];
+	// No header
 }
 
 void FWeightedBlendableCustomization::AddDirectAsset(TSharedRef<IPropertyHandle> StructPropertyHandle, UPackage* Package, TSharedPtr<IPropertyHandle> Weight, TSharedPtr<IPropertyHandle> Value, UClass* Class)
diff --git a/Engine/Source/Editor/DetailCustomizations/Private/PrimitiveComponentDetails.cpp b/Engine/Source/Editor/DetailCustomizations/Private/PrimitiveComponentDetails.cpp
--- a/Engine/Source/Editor/DetailCustomizations/Private/PrimitiveComponentDetails.cpp
+++ b/Engine/Source/Editor/DetailCustomizations/Private/PrimitiveComponentDetails.cpp
@@ -109,7 +109,7 @@ void FPrimitiveComponentDetails::AddMaterialCategory( IDetailLayoutBuilder& Deta
 
 void FPrimitiveComponentDetails::AddLightingCategory(IDetailLayoutBuilder& DetailBuilder)
 {
-	IDetailCategoryBuilder& CollisionCategory = DetailBuilder.EditCategory("Lighting");
+	IDetailCategoryBuilder& LightingCategory = DetailBuilder.EditCategory("Lighting");
 }
 
 void FPrimitiveComponentDetails::AddAdvancedSubCategory( IDetailLayoutBuilder& DetailBuilder, FName MainCategoryName, FName SubCategoryName)
diff --git a/Engine/Source/Editor/DetailCustomizations/Private/StaticMeshComponentDetails.cpp b/Engine/Source/Editor/DetailCustomizations/Private/StaticMeshComponentDetails.cpp
--- a/Engine/Source/Editor/DetailCustomizations/Private/StaticMeshComponentDetails.cpp
+++ b/Engine/Source/Editor/DetailCustomizations/Private/StaticMeshComponentDetails.cpp
@@ -4,6 +4,8 @@
 #include "Components/StaticMeshComponent.h"
 #include "PropertyHandle.h"
 #include "DetailLayoutBuilder.h"
+#include "SNumericEntryBox.h"
+#include "SWidgetSwitcher.h"
 
 #define LOCTEXT_NAMESPACE "StaticMeshComponentDetails"
 
@@ -13,13 +15,179 @@ TSharedRef<IDetailCustomization> FStaticMeshComponentDetails::MakeInstance()
 	return MakeShareable( new FStaticMeshComponentDetails );
 }
 
-void FStaticMeshComponentDetails::CustomizeDetails( IDetailLayoutBuilder& DetailBuilder )
+void FStaticMeshComponentDetails::CustomizeDetails(IDetailLayoutBuilder& DetailBuilder)
 {
 	// Create a category so this is displayed early in the properties
-	DetailBuilder.EditCategory( "StaticMesh", FText::GetEmpty(), ECategoryPriority::Important);
+	DetailBuilder.EditCategory("StaticMesh", FText::GetEmpty(), ECategoryPriority::Important);
 
 	TSharedRef<IPropertyHandle> UseDefaultCollision = DetailBuilder.GetProperty(GET_MEMBER_NAME_CHECKED(UStaticMeshComponent, bUseDefaultCollision));
 	UseDefaultCollision->MarkHiddenByCustomization();
+
+	IDetailCategoryBuilder& LightingCategory = DetailBuilder.EditCategory("Lighting");
+
+	// Store off the properties for analysis in later function calls.
+	OverrideLightResProperty = DetailBuilder.GetProperty(GET_MEMBER_NAME_CHECKED(UStaticMeshComponent, OverriddenLightMapRes));
+
+	// Add the row that we will be customizing below.
+	IDetailPropertyRow& OverrideLightResRow = LightingCategory.AddProperty(OverrideLightResProperty);
+
+	// Create the default property value widget, we'll use this in certain circumstances outlined below.
+	TSharedRef<SWidget> ValueWidget = OverrideLightResProperty->CreatePropertyValueWidget();
+
+	// Store off the objects that we are editing for analysis in later function calls.
+	DetailBuilder.GetObjectsBeingCustomized(ObjectsCustomized);
+
+	// We use similar logic here to where it is ultimately used:
+	//   UStaticMeshComponent::GetLightMapResolution
+	// If there is a static mesh and the override is enabled, use the real value of the property.
+	// If there is a static mesh and the override is disabled, use the static mesh's resolution value.
+	// If no static mesh is assigned, use 0.
+	// Ultimately, the last case is an error and we need to warn the user. We also need to handle multiple selection appropriately, 
+	// thus the widget switcher below.
+	
+	OverrideLightResRow.CustomWidget()
+		.NameContent()
+		[
+			// Use the default lightmap property name
+			OverrideLightResProperty->CreatePropertyNameWidget()
+		]
+		.ValueContent()
+		[
+			SNew(SWidgetSwitcher)
+			.WidgetIndex(this, &FStaticMeshComponentDetails::HandleNoticeSwitcherWidgetIndex)
+			+ SWidgetSwitcher::Slot() // Slot 0, the single common value for static mesh lightmap resolutions.
+			[
+				SNew(SNumericEntryBox<int32>)
+				.ToolTipText(OverrideLightResProperty->GetToolTipText())
+				.IsEnabled(false)
+				.Font(IDetailLayoutBuilder::GetDetailFont())
+				.Value(this, &FStaticMeshComponentDetails::GetStaticMeshLightResValue)
+			]
+			+ SWidgetSwitcher::Slot() // Slot 1, the editor for when overrides are enabled
+			[
+				ValueWidget
+			]
+			+ SWidgetSwitcher::Slot() // Slot 2, the editor for when we are missing one or more static meshes.
+			.HAlign(HAlign_Center)
+			.VAlign(VAlign_Center)
+			[
+				SNew(STextBlock)
+				.ToolTipText(OverrideLightResProperty->GetToolTipText())
+				.IsEnabled(false)
+				.Font(IDetailLayoutBuilder::GetDetailFont())
+				.Text(LOCTEXT("DetailsMissingStaticMesh", "Missing StaticMesh!"))
+			]
+			+ SWidgetSwitcher::Slot() // Slot 3, the editor for when we have multiple heterogenous values.
+			.HAlign(HAlign_Center)
+			.VAlign(VAlign_Center)
+			[
+				SNew(STextBlock)
+				.ToolTipText(OverrideLightResProperty->GetToolTipText())
+				.IsEnabled(false)
+				.Font(IDetailLayoutBuilder::GetDetailFont())
+				.Text(LOCTEXT("DetailsMultipleValues", "Multiple Values"))
+			]
+		];
+
+}
+
+TOptional<int32> FStaticMeshComponentDetails::GetStaticMeshLightResValue() const
+{
+	// Go ahead and grab the current value for the static mesh resolution as this shouldn't change during the lifetime
+	// of this customization.
+	TOptional<int32> DefaultResolution = 0;
+	for (int32 i = 0; i < ObjectsCustomized.Num() && DefaultResolution == 0; i++)
+	{
+		UStaticMeshComponent*  Component = Cast<UStaticMeshComponent>(ObjectsCustomized[i].Get());
+		if (Component != nullptr && Component->GetStaticMesh() != nullptr)
+		{
+			DefaultResolution = Component->GetStaticMesh()->LightMapResolution;
+		}
+	}
+	return DefaultResolution;
+}
+
+int FStaticMeshComponentDetails::HandleNoticeSwitcherWidgetIndex() const
+{
+	// Determine if the override enabled is set homogenously or not.
+	bool bOverrideEnabled = false;
+	FPropertyAccess::Result AccessResultOverrideEnabled = FPropertyAccess::Fail;
+	bool bOverrideFound = false;
+
+	// Loop through all the components being edited to see if any are missing static meshes
+	// and to see if their lightmap resolutions are all the same on those static meshes.
+	FPropertyAccess::Result AccessResultOverride = FPropertyAccess::Success;
+	bool bHasMissingStaticMeshes = false;
+	int32 OverrideLightRes = -1;
+	check(ObjectsCustomized.Num() > 0);
+	for (int32 i = 0; i < ObjectsCustomized.Num(); i++)
+	{
+		UStaticMeshComponent*  Component = Cast<UStaticMeshComponent>(ObjectsCustomized[i].Get());
+		if (Component != nullptr)
+		{
+			if (Component->GetStaticMesh() == nullptr)
+			{
+				bHasMissingStaticMeshes = true;
+				AccessResultOverride = FPropertyAccess::MultipleValues;
+			}
+			else
+			{
+				if (bOverrideFound == false)
+				{
+					AccessResultOverrideEnabled = FPropertyAccess::Success;
+					bOverrideEnabled = Component->bOverrideLightMapRes;
+					bOverrideFound = true;
+				}
+				else if (bOverrideEnabled != Component->bOverrideLightMapRes)
+				{
+					AccessResultOverrideEnabled = FPropertyAccess::MultipleValues;
+				}
+
+				int32 CurrentStaticLightRes = Component->GetStaticMesh()->LightMapResolution;
+				if (OverrideLightRes == -1)
+				{
+					OverrideLightRes = CurrentStaticLightRes;
+				}
+				else if (CurrentStaticLightRes != OverrideLightRes)
+				{
+					AccessResultOverride = FPropertyAccess::MultipleValues;
+				}
+			}
+		}
+	}
+
+	// This gets a little hairy. The desired behavior is as follows:
+	// For single item selection:
+	//     a) if override is enabled and we have a valid static mesh, display the general editor for the property.
+	//     b) if override is disabled and we have a valid static mesh, display the static mesh's value for the property.
+	//     c) otherwise, warn the user if they are missing the static mesh.
+	// For multiple selection:
+	//     d) if all overrides are enabled, all have valid meshes, and all have the same value, display the general editor.
+	//     e) if all overrides are enabled, all have valid meshes, and all have different values, display the general editor (which should say Multiple Values).
+	//     f) if all overrides are disabled, all have valid meshes, and all static meshes have the same value for the resolution, display that common resolution
+	//     g) if all overrides are disabled, all have valid meshes, and all static meshes have heterogenous values for the resolution, display the multiple values read-only text.
+	//     h) if overrides are heterogenous, display the multiple values read-only text.
+	//     i) if any of the above have invalid static meshes, warn the user.
+	if (bHasMissingStaticMeshes) // Covers cases c and i above.
+	{
+		return 2;
+	}
+	else if (AccessResultOverrideEnabled == FPropertyAccess::MultipleValues) // covers case h above
+	{
+		return 3;
+	}
+	else if (AccessResultOverride == FPropertyAccess::MultipleValues && bOverrideEnabled == false) // covers case g above
+	{
+		return 3;
+	}
+	else if (bOverrideEnabled == false) // covers cases b and f above
+	{
+		return 0;
+	}
+	else // covers cases a, d, and e above
+	{
+		return 1;
+	}
 }
 
 #undef LOCTEXT_NAMESPACE
diff --git a/Engine/Source/Editor/DetailCustomizations/Private/StaticMeshComponentDetails.h b/Engine/Source/Editor/DetailCustomizations/Private/StaticMeshComponentDetails.h
--- a/Engine/Source/Editor/DetailCustomizations/Private/StaticMeshComponentDetails.h
+++ b/Engine/Source/Editor/DetailCustomizations/Private/StaticMeshComponentDetails.h
@@ -19,12 +19,15 @@ public:
 
 private:
 	/**
-	 * @return Whether or not override lightmap res is enabled
+	 * @return Which UI to show whether or not override lightmap res is enabled
 	 */
-	bool IsOverrideLightmapResEnabled() const;
+	int HandleNoticeSwitcherWidgetIndex() const;
 
-	bool IsDefaultCollisionSupported() const;
+	/**
+	* @return Accessor for the static mesh's light resolution value.
+	*/
+	TOptional<int32> GetStaticMeshLightResValue() const;
 private:
-	TSharedPtr<IPropertyHandle> EnableOverrideLightmapRes;
 	TArray<TWeakObjectPtr<UObject>> ObjectsCustomized;
+	TSharedPtr<IPropertyHandle> OverrideLightResProperty;
 };
diff --git a/Engine/Source/Editor/DetailCustomizations/Private/StringClassReferenceCustomization.cpp b/Engine/Source/Editor/DetailCustomizations/Private/StringClassReferenceCustomization.cpp
--- a/Engine/Source/Editor/DetailCustomizations/Private/StringClassReferenceCustomization.cpp
+++ b/Engine/Source/Editor/DetailCustomizations/Private/StringClassReferenceCustomization.cpp
@@ -13,10 +13,12 @@ void FStringClassReferenceCustomization::CustomizeHeader(TSharedRef<IPropertyHan
 
 	const FString& MetaClassName = PropertyHandle->GetMetaData("MetaClass");
 	const FString& RequiredInterfaceName = PropertyHandle->GetMetaData("RequiredInterface");
-	const bool bAllowAbstract = PropertyHandle->GetBoolMetaData("AllowAbstract");
+	const bool bAllowAbstract = PropertyHandle->HasMetaData("AllowAbstract");
 	const bool bIsBlueprintBaseOnly = PropertyHandle->GetBoolMetaData("IsBlueprintBaseOnly");
 	const bool bAllowNone = !(PropertyHandle->GetMetaDataProperty()->PropertyFlags & CPF_NoClear);
-
+	const bool bShowTreeView = PropertyHandle->HasMetaData("ShowTreeView");
+	const bool bHideViewOptions = PropertyHandle->HasMetaData("HideViewOptions");
+	
 	const UClass* const MetaClass = !MetaClassName.IsEmpty()
 		? FEditorClassUtils::GetClassFromString(MetaClassName)
 		: UClass::StaticClass();
@@ -38,6 +40,8 @@ void FStringClassReferenceCustomization::CustomizeHeader(TSharedRef<IPropertyHan
 			.AllowAbstract(bAllowAbstract)
 			.IsBlueprintBaseOnly(bIsBlueprintBaseOnly)
 			.AllowNone(bAllowNone)
+			.ShowTreeView(bShowTreeView)
+			.HideViewOptions(bHideViewOptions)
 			.SelectedClass(this, &FStringClassReferenceCustomization::OnGetClass)
 			.OnSetClass(this, &FStringClassReferenceCustomization::OnSetClass)
 	];
diff --git a/Engine/Source/Editor/DetailCustomizations/Private/TextCustomization.cpp b/Engine/Source/Editor/DetailCustomizations/Private/TextCustomization.cpp
--- a/Engine/Source/Editor/DetailCustomizations/Private/TextCustomization.cpp
+++ b/Engine/Source/Editor/DetailCustomizations/Private/TextCustomization.cpp
@@ -84,6 +84,11 @@ namespace
 			}
 		}
 
+		virtual bool IsValidText(const FText& InText, FText& OutErrorMsg) const override
+		{
+			return true;
+		}
+
 #if USE_STABLE_LOCALIZATION_KEYS
 		virtual void GetStableTextId(const int32 InIndex, const ETextPropertyEditAction InEditAction, const FString& InTextSource, const FString& InProposedNamespace, const FString& InProposedKey, FString& OutStableNamespace, FString& OutStableKey) const override
 		{
diff --git a/Engine/Source/Editor/DetailCustomizations/Private/Vector4StructCustomization.cpp b/Engine/Source/Editor/DetailCustomizations/Private/Vector4StructCustomization.cpp
new file mode 100644
--- /dev/null
+++ b/Engine/Source/Editor/DetailCustomizations/Private/Vector4StructCustomization.cpp
@@ -0,0 +1,63 @@
+// Copyright 1998-2015 Epic Games, Inc. All Rights Reserved.
+#include "Vector4StructCustomization.h"
+#include "IPropertyUtilities.h"
+#include "SNumericEntryBox.h"
+#include "SColorGradingPicker.h"
+#include "IDetailChildrenBuilder.h"
+#include "DetailLayoutBuilder.h"
+#include "DetailWidgetRow.h"
+
+TSharedRef<IPropertyTypeCustomization> FVector4StructCustomization::MakeInstance()
+{
+	return MakeShareable(new FVector4StructCustomization);
+}
+
+FVector4StructCustomization::FVector4StructCustomization()
+	: FMathStructCustomization()
+	, ColorGradingVectorCustomization(nullptr)
+{
+}
+
+FVector4StructCustomization::~FVector4StructCustomization()
+{
+	//Release all resources
+	ColorGradingVectorCustomization = nullptr;
+}
+
+void FVector4StructCustomization::CustomizeChildren(TSharedRef<IPropertyHandle> StructPropertyHandle, IDetailChildrenBuilder& StructBuilder, IPropertyTypeCustomizationUtils& StructCustomizationUtils)
+{
+	UProperty* Property = StructPropertyHandle->GetProperty();
+	if (Property)
+	{
+		const FString& ColorGradingModeString = Property->GetMetaData(TEXT("ColorGradingMode"));
+		if (!ColorGradingModeString.IsEmpty())
+		{
+			//Create our color grading customization shared pointer
+			if (!ColorGradingVectorCustomization.IsValid())
+			{
+				ColorGradingVectorCustomization = MakeShareable(new FColorGradingVectorCustomization);
+			}
+
+			//Customize the childrens
+			TArray<TWeakPtr<IPropertyHandle>> WeakChildArray;
+			for (TSharedRef<IPropertyHandle> PropertyHandleRef : SortedChildHandles)
+			{
+				WeakChildArray.Add(PropertyHandleRef);
+			}
+			ColorGradingVectorCustomization->CustomizeChildren(StructPropertyHandle, StructBuilder, StructCustomizationUtils, WeakChildArray);
+			
+			// We handle the customize Children so just return here
+			return;
+		}
+	}
+
+	//Use the base class customize children
+	FMathStructCustomization::CustomizeChildren(StructPropertyHandle, StructBuilder, StructCustomizationUtils);
+}
+
+void FVector4StructCustomization::MakeHeaderRow(TSharedRef<class IPropertyHandle>& StructPropertyHandle, FDetailWidgetRow& Row)
+{
+	//Use the base class header row
+	FMathStructCustomization::MakeHeaderRow(StructPropertyHandle, Row);
+}
+
diff --git a/Engine/Source/Editor/EditorStyle/Private/SlateEditorStyle.cpp b/Engine/Source/Editor/EditorStyle/Private/SlateEditorStyle.cpp
--- a/Engine/Source/Editor/EditorStyle/Private/SlateEditorStyle.cpp
+++ b/Engine/Source/Editor/EditorStyle/Private/SlateEditorStyle.cpp
@@ -5300,8 +5300,8 @@ void FSlateEditorStyle::FStyle::SetupPersonaStyle()
 
 		Set( "BlueprintEditor.Details.DeleteButton", new IMAGE_BRUSH( "/Icons/GenericDelete_Black", Icon16x16 ) );
 
-		Set( "BlueprintEditor.Details.ArgUpButton", new IMAGE_BRUSH( "/Icons/icon_FunctionArgUp", Icon16x16 ) );
-		Set( "BlueprintEditor.Details.ArgDownButton", new IMAGE_BRUSH( "/Icons/icon_FunctionArgDown", Icon16x16 ) );
+		Set( "BlueprintEditor.Details.ArgUpButton", new IMAGE_BRUSH( "/Icons/icon_FunctionArgUp", Icon16x16, FLinearColor(0.0f, 0.0f, 0.0f)) );
+		Set( "BlueprintEditor.Details.ArgDownButton", new IMAGE_BRUSH( "/Icons/icon_FunctionArgDown", Icon16x16, FLinearColor(0.0f, 0.0f, 0.0f)) );
 
 		Set( "FullBlueprintEditor.Diff", new IMAGE_BRUSH( "Icons/BlueprintEditorDiff", Icon40x40 ) );
 		Set( "FullBlueprintEditor.Diff.Small", new IMAGE_BRUSH( "Icons/BlueprintEditorDiff", Icon20x20 ) );
@@ -5873,6 +5873,7 @@ void FSlateEditorStyle::FStyle::SetupClassIconsAndThumbnails()
 			TEXT("Sphere"),
 			TEXT("Cylinder"),
 			TEXT("Cone"),
+			TEXT("Plane"),
 			TEXT("CineCameraActor"),
 			TEXT("CameraRig_Crane"),
 			TEXT("CameraRig_Rail")
@@ -6045,9 +6046,13 @@ void FSlateEditorStyle::FStyle::SetupContentBrowserStyle()
 		Set( "ContentBrowser.AssetActions.GenericFind", new IMAGE_BRUSH( "Icons/icon_Genericfinder_16x", Icon16x16) );
 		Set( "ContentBrowser.AssetLocalization", new IMAGE_BRUSH( "Icons/icon_localization_16x", Icon16x16 ) );
 
-		Set( "MediaAsset.AssetActions.Play", new IMAGE_BRUSH( "Icons/icon_SCueEd_PlayCue_16x", Icon16x16 ) );
-		Set( "MediaAsset.AssetActions.Stop", new IMAGE_BRUSH( "Icons/icon_SCueEd_Stop_16x", Icon16x16 ) );
-		Set( "MediaAsset.AssetActions.Pause", new IMAGE_BRUSH( "Icons/icon_SCueEd_Pause_16x", Icon16x16 ) );
+		Set( "MediaAsset.AssetActions.Play.Small", new IMAGE_BRUSH( "Icons/icon_SCueEd_PlayCue_16x", Icon16x16 ) );
+		Set( "MediaAsset.AssetActions.Stop.Small", new IMAGE_BRUSH( "Icons/icon_SCueEd_Stop_16x", Icon16x16 ) );
+		Set( "MediaAsset.AssetActions.Pause.Small", new IMAGE_BRUSH( "Icons/icon_SCueEd_Pause_16x", Icon16x16 ) );
+
+		Set("MediaAsset.AssetActions.Play.Large", new IMAGE_BRUSH("Icons/icon_SCueEd_PlayCue_40x", Icon40x40));
+		Set("MediaAsset.AssetActions.Stop.Large", new IMAGE_BRUSH("Icons/icon_SCueEd_Stop_40x", Icon40x40));
+		Set("MediaAsset.AssetActions.Pause.Large", new IMAGE_BRUSH("Icons/icon_SCueEd_Pause_40x", Icon40x40));
 
 		// Misc
 		Set( "ContentBrowser.ThumbnailShadow", new BOX_BRUSH( "ContentBrowser/ThumbnailShadow" , FMargin( 4.0f / 64.0f ) ) );
diff --git a/Engine/Source/Editor/EditorWidgets/Private/STextPropertyEditableTextBox.cpp b/Engine/Source/Editor/EditorWidgets/Private/STextPropertyEditableTextBox.cpp
--- a/Engine/Source/Editor/EditorWidgets/Private/STextPropertyEditableTextBox.cpp
+++ b/Engine/Source/Editor/EditorWidgets/Private/STextPropertyEditableTextBox.cpp
@@ -102,6 +102,7 @@ void STextPropertyEditableTextBox::Construct(const FArguments& InArgs, const TSh
 					.ForegroundColor(InArgs._ForegroundColor)
 					.SelectAllTextWhenFocused(false)
 					.ClearKeyboardFocusOnCommit(false)
+					.OnTextChanged(this, &STextPropertyEditableTextBox::OnTextChanged)
 					.OnTextCommitted(this, &STextPropertyEditableTextBox::OnTextCommitted)
 					.SelectAllTextOnCommit(false)
 					.IsReadOnly(this, &STextPropertyEditableTextBox::IsReadOnly)
@@ -134,6 +135,7 @@ void STextPropertyEditableTextBox::Construct(const FArguments& InArgs, const TSh
 					.ForegroundColor(InArgs._ForegroundColor)
 					.SelectAllTextWhenFocused(true)
 					.ClearKeyboardFocusOnCommit(false)
+					.OnTextChanged(this, &STextPropertyEditableTextBox::OnTextChanged)
 					.OnTextCommitted(this, &STextPropertyEditableTextBox::OnTextCommitted)
 					.SelectAllTextOnCommit(true)
 					.IsReadOnly(this, &STextPropertyEditableTextBox::IsReadOnly)
@@ -383,7 +385,13 @@ FText STextPropertyEditableTextBox::GetToolTipText() const
 
 		if (bIsLocalized)
 		{
-			LocalizedTextToolTip = FText::Format(LOCTEXT("LocalizedTextToolTipFmt", "--- Localized Text ---\nNamespace: {0}\nKey: {1}\nSource: {2}"), FText::FromString(Namespace), FText::FromString(Key), FText::FromString(*SourceString));
+			const FString PackageNamespace = TextNamespaceUtil::ExtractPackageNamespace(Namespace);
+			const FString TextNamespace = TextNamespaceUtil::StripPackageNamespace(Namespace);
+
+			LocalizedTextToolTip = FText::Format(
+				LOCTEXT("LocalizedTextToolTipFmt", "--- Localized Text ---\nPackage: {0}\nNamespace: {1}\nKey: {2}\nSource: {3}"), 
+				FText::FromString(PackageNamespace), FText::FromString(TextNamespace), FText::FromString(Key), FText::FromString(*SourceString)
+				);
 		}
 	}
 	
@@ -417,13 +425,42 @@ FText STextPropertyEditableTextBox::GetTextValue() const
 	return TextValue;
 }
 
+void STextPropertyEditableTextBox::OnTextChanged(const FText& NewText)
+{
+	const int32 NumTexts = EditableTextProperty->GetNumTexts();
+
+	FText TextErrorMsg;
+
+	// Don't validate the Multiple Values text if there are multiple properties being set
+	if (NumTexts > 0 && (NumTexts == 1 || NewText.ToString().Equals(MultipleValuesText.ToString(), ESearchCase::CaseSensitive)))
+	{
+		EditableTextProperty->IsValidText(NewText, TextErrorMsg);
+	}
+
+	// Update or clear the error message
+	SetTextError(TextErrorMsg);
+}
+
 void STextPropertyEditableTextBox::OnTextCommitted(const FText& NewText, ETextCommit::Type CommitInfo)
 {
 	const int32 NumTexts = EditableTextProperty->GetNumTexts();
 
 	// Don't commit the Multiple Values text if there are multiple properties being set
-	if (NumTexts > 0 && (NumTexts == 1 || NewText.ToString() != MultipleValuesText.ToString()))
+	if (NumTexts > 0 && (NumTexts == 1 || NewText.ToString().Equals(MultipleValuesText.ToString(), ESearchCase::CaseSensitive)))
 	{
+		FText TextErrorMsg;
+		if (EditableTextProperty->IsValidText(NewText, TextErrorMsg))
+		{
+			// Valid text; clear any error
+			SetTextError(FText::GetEmpty());
+		}
+		else
+		{
+			// Invalid text; set the error and prevent the new text from being set
+			SetTextError(TextErrorMsg);
+			return;
+		}
+
 		const FString& SourceString = NewText.ToString();
 		for (int32 TextIndex = 0; TextIndex < NumTexts; ++TextIndex)
 		{
@@ -480,6 +517,19 @@ void STextPropertyEditableTextBox::OnTextCommitted(const FText& NewText, ETextCo
 	}
 }
 
+void STextPropertyEditableTextBox::SetTextError(const FText& InErrorMsg)
+{
+	if (MultiLineWidget.IsValid())
+	{
+		MultiLineWidget->SetError(InErrorMsg);
+	}
+
+	if (SingleLineWidget.IsValid())
+	{
+		SingleLineWidget->SetError(InErrorMsg);
+	}
+}
+
 FText STextPropertyEditableTextBox::GetNamespaceValue() const
 {
 	FText NamespaceValue;
diff --git a/Engine/Source/Editor/EditorWidgets/Public/STextPropertyEditableTextBox.h b/Engine/Source/Editor/EditorWidgets/Public/STextPropertyEditableTextBox.h
--- a/Engine/Source/Editor/EditorWidgets/Public/STextPropertyEditableTextBox.h
+++ b/Engine/Source/Editor/EditorWidgets/Public/STextPropertyEditableTextBox.h
@@ -55,6 +55,9 @@ public:
 	/** Set the text at the given index (check against GetNumTexts) */
 	virtual void SetText(const int32 InIndex, const FText& InText) = 0;
 
+	/** Check to see if the given text is valid to use */
+	virtual bool IsValidText(const FText& InText, FText& OutErrorMsg) const = 0;
+
 #if USE_STABLE_LOCALIZATION_KEYS
 	/** Get the stable text ID for the given index (check against GetNumTexts) */
 	virtual void GetStableTextId(const int32 InIndex, const ETextPropertyEditAction InEditAction, const FString& InTextSource, const FString& InProposedNamespace, const FString& InProposedKey, FString& OutStableNamespace, FString& OutStableKey) const = 0;
@@ -115,7 +118,9 @@ private:
 	FText GetToolTipText() const;
 
 	FText GetTextValue() const;
+	void OnTextChanged(const FText& NewText);
 	void OnTextCommitted(const FText& NewText, ETextCommit::Type CommitInfo);
+	void SetTextError(const FText& InErrorMsg);
 
 	FText GetNamespaceValue() const;
 	void OnNamespaceChanged(const FText& NewText);
diff --git a/Engine/Source/Editor/FoliageEdit/Private/FoliageEdMode.cpp b/Engine/Source/Editor/FoliageEdit/Private/FoliageEdMode.cpp
--- a/Engine/Source/Editor/FoliageEdit/Private/FoliageEdMode.cpp
+++ b/Engine/Source/Editor/FoliageEdit/Private/FoliageEdMode.cpp
@@ -759,7 +759,15 @@ void FEdModeFoliage::StartFoliageBrushTrace(FEditorViewportClient* ViewportClien
 		}
 		PreApplyBrush();
 		ApplyBrush(ViewportClient);
-		bToolActive = true;
+
+		if (UISettings.GetIsInQuickSingleInstantiationMode() || UISettings.GetIsInSingleInstantiationMode())
+		{
+			EndFoliageBrushTrace();
+		}
+		else
+		{
+			bToolActive = true;
+		}
 	}
 }
 
@@ -797,6 +805,7 @@ void FEdModeFoliage::FoliageBrushTrace(FEditorViewportClient* ViewportClient, co
 					{
 						// Adjust the brush location
 						BrushLocation = Hit.Location;
+						BrushNormal = Hit.Normal;
 					}
 
 					// Still want to draw the brush when resizing
@@ -889,16 +898,16 @@ void FEdModeFoliage::GetRandomVectorInBrush(FVector& OutStart, FVector& OutEnd)
 	float Ru = (2.f * FMath::FRand() - 1.f);
 	float Rv = (2.f * FMath::FRand() - 1.f) * FMath::Sqrt(1.f - FMath::Square(Ru));
 
-	// find random point in circle through brush location parallel to screen surface
+	// find random point in circle through brush location on the same plane to brush location hit surface normal
 	FVector U, V;
-	BrushTraceDirection.FindBestAxisVectors(U, V);
+	BrushNormal.FindBestAxisVectors(U, V);
 	FVector Point = Ru * U + Rv * V;
 
 	// find distance to surface of sphere brush from this point
-	FVector Rw = FMath::Sqrt(1.f - (FMath::Square(Ru) + FMath::Square(Rv))) * BrushTraceDirection;
+	FVector Rw = FMath::Sqrt(1.f - (FMath::Square(Ru) + FMath::Square(Rv))) * BrushNormal;
 
-	OutStart = BrushLocation + UISettings.GetRadius() * (Point - Rw);
-	OutEnd = BrushLocation + UISettings.GetRadius() * (Point + Rw);
+	OutStart = BrushLocation + UISettings.GetRadius() * (Point + Rw);
+	OutEnd = BrushLocation + UISettings.GetRadius() * (Point - Rw);
 }
 
 static bool IsWithinSlopeAngle(float NormalZ, float MinAngle, float MaxAngle, float Tolerance = SMALL_NUMBER)
@@ -1246,10 +1255,10 @@ void FEdModeFoliage::AddInstances(UWorld* InWorld, const TArray<FDesiredFoliageI
 	}
 }
 
-static void SpawnFoliageInstance(UWorld* InWorld, const UFoliageType* Settings, const FFoliageInstance& Instance, UActorComponent* BaseComponent)
+static void SpawnFoliageInstance(UWorld* InWorld, const UFoliageType* Settings, const FFoliageUISettings* UISettings, const FFoliageInstance& Instance, UActorComponent* BaseComponent)
 {
 	// We always spawn instances in base component level
-	ULevel* TargetLevel = BaseComponent->GetComponentLevel();
+	ULevel* TargetLevel = UISettings->GetIsInSpawnInCurrentLevelMode() ? InWorld->GetCurrentLevel() : BaseComponent->GetComponentLevel();
 	AInstancedFoliageActor* IFA = AInstancedFoliageActor::GetInstancedFoliageActorForLevel(TargetLevel, true);
 
 	FFoliageMeshInfo* MeshInfo;
@@ -1312,12 +1321,30 @@ void FEdModeFoliage::AddInstancesImp(UWorld* InWorld, const UFoliageType* Settin
 			{
 				Inst.ProceduralGuid = PotentialInstance.DesiredInstance.ProceduralGuid;
 
-				SpawnFoliageInstance(InWorld, Settings, Inst, PotentialInstance.HitComponent);
+				SpawnFoliageInstance(InWorld, Settings, UISettings, Inst, PotentialInstance.HitComponent);
 			}
 		}
 	}
 }
 
+void FEdModeFoliage::AddSingleInstanceForBrush(UWorld* InWorld, const UFoliageType* Settings, float Pressure)
+{
+	TArray<FDesiredFoliageInstance> DesiredInstances;
+	DesiredInstances.Reserve(1);
+
+	// Simply generate a start/end around the brush location so the line check will hit the brush location
+	FVector Start = BrushLocation + BrushNormal;
+	FVector End = BrushLocation - BrushNormal;
+
+	FDesiredFoliageInstance* DesiredInstance = new (DesiredInstances)FDesiredFoliageInstance(Start, End);
+
+	// We do not apply the density limitation based on the brush size
+	TArray<int32> ExistingInstanceBuckets;
+	ExistingInstanceBuckets.AddZeroed(NUM_INSTANCE_BUCKETS);
+
+	AddInstancesImp(InWorld, Settings, DesiredInstances, ExistingInstanceBuckets, Pressure, &LandscapeLayerCaches, &UISettings);
+}
+
 /** Add instances inside the brush to match DesiredInstanceCount */
 void FEdModeFoliage::AddInstancesForBrush(UWorld* InWorld, const UFoliageType* Settings, const FSphere& BrushSphere, int32 DesiredInstanceCount, float Pressure)
 {
@@ -2179,12 +2206,19 @@ void FEdModeFoliage::ApplyBrush(FEditorViewportClient* ViewportClient)
 			}
 			else
 			{
-				// This is the total set of instances disregarding parameters like slope, height or layer.
-				float DesiredInstanceCountFloat = BrushArea * Settings->Density * UISettings.GetPaintDensity() / (1000.f*1000.f);
-				// Allow a single instance with a random chance, if the brush is smaller than the density
-				int32 DesiredInstanceCount = DesiredInstanceCountFloat > 1.f ? FMath::RoundToInt(DesiredInstanceCountFloat) : FMath::FRand() < DesiredInstanceCountFloat ? 1 : 0;
+				if (UISettings.GetIsInQuickSingleInstantiationMode() || UISettings.GetIsInSingleInstantiationMode())
+				{
+					AddSingleInstanceForBrush(World, Settings, Pressure);
+				}
+				else
+				{
+					// This is the total set of instances disregarding parameters like slope, height or layer.
+					float DesiredInstanceCountFloat = BrushArea * Settings->Density * UISettings.GetPaintDensity() / (1000.f*1000.f);
+					// Allow a single instance with a random chance, if the brush is smaller than the density
+					int32 DesiredInstanceCount = DesiredInstanceCountFloat > 1.f ? FMath::RoundToInt(DesiredInstanceCountFloat) : FMath::FRand() < DesiredInstanceCountFloat ? 1 : 0;
 
-				AddInstancesForBrush(World, Settings, BrushSphere, DesiredInstanceCount, Pressure);
+					AddInstancesForBrush(World, Settings, BrushSphere, DesiredInstanceCount, Pressure);
+				}
 			}
 		}
 
@@ -2417,7 +2451,7 @@ void FEdModeFoliage::ApplyPaintBucket_Add(AActor* Actor)
 			FFoliageInstance Inst;
 			if (PotentialInstance.PlaceInstance(World, Settings, Inst))
 			{
-				SpawnFoliageInstance(World, Settings, Inst, PotentialInstance.HitComponent);
+				SpawnFoliageInstance(World, Settings, &UISettings, Inst, PotentialInstance.HitComponent);
 			}
 		}
 
@@ -3119,6 +3153,14 @@ bool FEdModeFoliage::InputKey(FEditorViewportClient* ViewportClient, FViewport*
 				bHandled = true;
 			}
 		}
+		else if (Key == EKeys::I && Event == IE_Released)
+		{
+			UISettings.SetIsInQuickSingleInstantiationMode(false);
+		}
+		else if (Key == EKeys::I && Event == IE_Pressed)
+		{
+			UISettings.SetIsInQuickSingleInstantiationMode(true);
+		}
 	}
 
 	if (!bHandled && (UISettings.GetLassoSelectToolSelected() || UISettings.GetSelectToolSelected()) && FoliageInteractor == nullptr)
diff --git a/Engine/Source/Editor/FoliageEdit/Private/FoliageEdMode.h b/Engine/Source/Editor/FoliageEdit/Private/FoliageEdMode.h
--- a/Engine/Source/Editor/FoliageEdit/Private/FoliageEdMode.h
+++ b/Engine/Source/Editor/FoliageEdit/Private/FoliageEdMode.h
@@ -38,6 +38,8 @@ namespace EFoliagePaletteViewMode
 	};
 }
 
+const float SingleInstanceModeBrushSize = 20.0f;
+
 // Current user settings in Foliage UI
 struct FFoliageUISettings
 {
@@ -62,7 +64,7 @@ struct FFoliageUISettings
 	bool GetReapplyPaintBucketToolSelected() const { return bReapplyPaintBucketToolSelected ? true : false; }
 	void SetReapplyPaintBucketToolSelected(bool InbReapplyPaintBucketToolSelected) { bReapplyPaintBucketToolSelected = InbReapplyPaintBucketToolSelected; }
 
-	float GetRadius() const { return Radius; }
+	float GetRadius() const { return (GetIsInSingleInstantiationMode() || GetIsInQuickSingleInstantiationMode()) ? SingleInstanceModeBrushSize : Radius; }
 	void SetRadius(float InRadius) { Radius = InRadius; }
 	float GetPaintDensity() const { return PaintDensity; }
 	void SetPaintDensity(float InPaintDensity) { PaintDensity = InPaintDensity; }
@@ -79,6 +81,15 @@ struct FFoliageUISettings
 	bool GetFilterTranslucent() const { return bFilterTranslucent; }
 	void SetFilterTranslucent(bool InbFilterTranslucent) { bFilterTranslucent = InbFilterTranslucent; }
 
+	bool GetIsInSingleInstantiationMode() const { return IsInSingleInstantiationMode; }
+	void SetIsInSingleInstantiationMode(bool InIsInSingleInstantiationMode) { IsInSingleInstantiationMode = InIsInSingleInstantiationMode; }
+
+	bool GetIsInQuickSingleInstantiationMode() const { return IsInQuickSingleInstantiationMode; }
+	void SetIsInQuickSingleInstantiationMode(bool InIsInQuickSingleInstantiationMode) { IsInQuickSingleInstantiationMode = InIsInQuickSingleInstantiationMode; }
+
+	bool GetIsInSpawnInCurrentLevelMode() const { return IsInSpawnInCurrentLevelMode; }
+	void SetSpawnInCurrentLevelMode(bool InSpawnInCurrentLevelMode) { IsInSpawnInCurrentLevelMode = InSpawnInCurrentLevelMode; }
+
 	bool GetShowPaletteItemDetails() const { return bShowPaletteItemDetails; }
 	void SetShowPaletteItemDetails(bool InbShowPaletteItemDetails) { bShowPaletteItemDetails = InbShowPaletteItemDetails; }
 	bool GetShowPaletteItemTooltips() const { return bShowPaletteItemTooltips; }
@@ -106,6 +117,9 @@ struct FFoliageUISettings
 		, Radius(512.f)
 		, PaintDensity(0.5f)
 		, UnpaintDensity(0.f)
+		, IsInSingleInstantiationMode(false)
+		, IsInQuickSingleInstantiationMode(false)
+		, IsInSpawnInCurrentLevelMode(false)
 		, bFilterLandscape(true)
 		, bFilterStaticMesh(true)
 		, bFilterBSP(true)
@@ -140,6 +154,10 @@ private:
 	float PaintDensity;
 	float UnpaintDensity;
 
+	bool IsInSingleInstantiationMode;
+	bool IsInQuickSingleInstantiationMode;
+	bool IsInSpawnInCurrentLevelMode;
+
 public:
 	bool bFilterLandscape;
 	bool bFilterStaticMesh;
@@ -470,6 +488,8 @@ private:
 	/** Add instances inside the brush to match DesiredInstanceCount */
 	void AddInstancesForBrush(UWorld* InWorld, const UFoliageType* Settings, const FSphere& BrushSphere, int32 DesiredInstanceCount, float Pressure);
 
+	void AddSingleInstanceForBrush(UWorld* InWorld, const UFoliageType* Settings, float Pressure);
+
 	/** Remove instances inside the brush to match DesiredInstanceCount. */
 	void RemoveInstancesForBrush(UWorld* InWorld, const UFoliageType* Settings, const FSphere& BrushSphere, int32 DesiredInstanceCount, float Pressure);
 
@@ -539,6 +559,7 @@ private:
 
 	bool bBrushTraceValid;
 	FVector BrushLocation;
+	FVector BrushNormal;
 	FVector BrushTraceDirection;
 	UStaticMeshComponent* SphereBrushComponent;
 
diff --git a/Engine/Source/Editor/FoliageEdit/Private/SFoliageEdit.cpp b/Engine/Source/Editor/FoliageEdit/Private/SFoliageEdit.cpp
--- a/Engine/Source/Editor/FoliageEdit/Private/SFoliageEdit.cpp
+++ b/Engine/Source/Editor/FoliageEdit/Private/SFoliageEdit.cpp
@@ -24,6 +24,8 @@
 #include "IIntroTutorials.h"
 #include "Widgets/Input/SNumericEntryBox.h"
 #include "SFoliagePalette.h"
+#include "SHeader.h"
+#include "SSeparator.h"
 
 #define LOCTEXT_NAMESPACE "FoliageEd_Mode"
 
@@ -93,6 +95,19 @@ void SFoliageEdit::Construct(const FArguments& InArgs)
 						]
 					]
 
+					+ SVerticalBox::Slot()
+					.AutoHeight()
+					.Padding(StandardPadding)
+					[
+						SNew(SHeader)
+						.Visibility(this, &SFoliageEdit::GetVisibility_Options)
+						[
+							SNew(STextBlock)
+							.Text(LOCTEXT("OptionHeader", "Brush Options"))
+							.Font(StandardFont)
+						]
+					]
+
 					// Brush Size
 					+ SVerticalBox::Slot()
 					.AutoHeight()
@@ -122,10 +137,12 @@ void SFoliageEdit::Construct(const FArguments& InArgs)
 							.MinValue(0.0f)
 							.MaxValue(65536.0f)
 							.MaxSliderValue(8192.0f)
+							.MinDesiredValueWidth(50.0f)
 							.SliderExponent(3.0f)
 							.Value(this, &SFoliageEdit::GetRadius)
 							.OnValueChanged(this, &SFoliageEdit::SetRadius)
-						]
+							.IsEnabled(this, &SFoliageEdit::IsEnabled_BrushSize)
+						]						
 					]
 
 					// Paint Density
@@ -157,8 +174,10 @@ void SFoliageEdit::Construct(const FArguments& InArgs)
 							.MinValue(0.0f)
 							.MaxValue(2.0f)
 							.MaxSliderValue(1.0f)
+							.MinDesiredValueWidth(50.0f)
 							.Value(this, &SFoliageEdit::GetPaintDensity)
 							.OnValueChanged(this, &SFoliageEdit::SetPaintDensity)
+							.IsEnabled(this, &SFoliageEdit::IsEnabled_PaintDensity)
 						]
 					]
 
@@ -191,13 +210,82 @@ void SFoliageEdit::Construct(const FArguments& InArgs)
 							.MinValue(0.0f)
 							.MaxValue(2.0f)
 							.MaxSliderValue(1.0f)
+							.MinDesiredValueWidth(50.0f)
 							.Value(this, &SFoliageEdit::GetEraseDensity)
 							.OnValueChanged(this, &SFoliageEdit::SetEraseDensity)
+							.IsEnabled(this, &SFoliageEdit::IsEnabled_EraseDensity)
+						]
+					]					
+					
+					+ SVerticalBox::Slot()
+					.Padding(StandardPadding)
+					.AutoHeight()
+					[
+						SNew(SHorizontalBox)
+						.Visibility(this, &SFoliageEdit::GetVisibility_Options)
+
+						+ SHorizontalBox::Slot()
+						.VAlign(VAlign_Center)
+						.Padding(StandardPadding)
+						[
+							SNew(SWrapBox)
+							.UseAllottedWidth(true)
+							.InnerSlotPadding({6, 5})
+
+							+ SWrapBox::Slot()
+							[
+								SNew(SBox)
+								.MinDesiredWidth(150)
+								[
+									SNew(SCheckBox)
+									.Visibility(this, &SFoliageEdit::GetVisibility_SingleInstantiationMode)
+									.OnCheckStateChanged(this, &SFoliageEdit::OnCheckStateChanged_SingleInstantiationMode)
+									.IsChecked(this, &SFoliageEdit::GetCheckState_SingleInstantiationMode)
+									.ToolTipText(LOCTEXT("SingleInstantiationModeTooltips", "Paint a single foliage instance at the mouse cursor location (i + Mouse Click)"))
+									[
+										SNew(STextBlock)
+										.Text(LOCTEXT("SingleInstantiationMode", "Single Instance Mode"))
+										.Font(StandardFont)
+									]
+								]
+							]
+
+							+ SWrapBox::Slot()
+							[
+								SNew(SBox)
+								.MinDesiredWidth(150)
+								[
+									SNew(SCheckBox)
+									.Visibility(this, &SFoliageEdit::GetVisibility_SpawnInCurrentLevelMode)
+									.OnCheckStateChanged(this, &SFoliageEdit::OnCheckStateChanged_SpawnInCurrentLevelMode)
+									.IsChecked(this, &SFoliageEdit::GetCheckState_SpawnInCurrentLevelMode)
+									.ToolTipText(LOCTEXT("SpawnInCurrentLevelModeTooltips", "Wether to place foliage meshes in the current level or in the level containing the mesh being painted on."))
+									[
+										SNew(STextBlock)
+										.Text(LOCTEXT("SpawnInCurrentLevelMode", "Place in Current Level"))
+										.Font(StandardFont)
+									]
+								]
+							]
 						]
 					]
 
 					// Filters
 					+ SVerticalBox::Slot()
+					.Padding(StandardPadding)
+					.AutoHeight()
+					[
+						SNew(SHeader)
+						.Visibility(this, &SFoliageEdit::GetVisibility_Filters)
+						[
+							SNew(STextBlock)
+							.Text(LOCTEXT("FiltersHeader", "Filters"))
+							.Font(StandardFont)
+						]
+					]
+
+					+ SVerticalBox::Slot()
+					.Padding(StandardPadding)
 					.AutoHeight()
 					[
 						SNew(SHorizontalBox)
@@ -304,8 +392,23 @@ void SFoliageEdit::Construct(const FArguments& InArgs)
 						]
 					]
 
+					// Actions
 					+ SVerticalBox::Slot()
 					.Padding(StandardPadding)
+					.AutoHeight()
+					[
+						SNew(SHeader)
+						.Visibility(this, &SFoliageEdit::GetVisibility_Actions)
+						[
+							SNew(STextBlock)
+							.Text(LOCTEXT("ActionsHeader", "Actions"))
+							.Font(StandardFont)
+						]
+					]
+
+					+ SVerticalBox::Slot()
+					.Padding(StandardPadding)
+					.AutoHeight()
 					[
 						SNew(SWrapBox)
 						.UseAllottedWidth(true)
@@ -487,6 +590,11 @@ TOptional<float> SFoliageEdit::GetRadius() const
 	return FoliageEditMode->UISettings.GetRadius();
 }
 
+bool SFoliageEdit::IsEnabled_BrushSize() const
+{
+	return !FoliageEditMode->UISettings.GetIsInSingleInstantiationMode() && !FoliageEditMode->UISettings.GetIsInQuickSingleInstantiationMode();
+}
+
 void SFoliageEdit::SetPaintDensity(float InDensity)
 {
 	FoliageEditMode->UISettings.SetPaintDensity(InDensity);
@@ -497,6 +605,10 @@ TOptional<float> SFoliageEdit::GetPaintDensity() const
 	return FoliageEditMode->UISettings.GetPaintDensity();
 }
 
+bool SFoliageEdit::IsEnabled_PaintDensity() const
+{
+	return !FoliageEditMode->UISettings.GetIsInSingleInstantiationMode() && !FoliageEditMode->UISettings.GetIsInQuickSingleInstantiationMode();
+}
 
 void SFoliageEdit::SetEraseDensity(float InDensity)
 {
@@ -508,6 +620,11 @@ TOptional<float> SFoliageEdit::GetEraseDensity() const
 	return FoliageEditMode->UISettings.GetUnpaintDensity();
 }
 
+bool SFoliageEdit::IsEnabled_EraseDensity() const
+{
+	return !FoliageEditMode->UISettings.GetIsInSingleInstantiationMode() && !FoliageEditMode->UISettings.GetIsInQuickSingleInstantiationMode();
+}
+
 EVisibility SFoliageEdit::GetVisibility_SelectionOptions() const
 {
 	return IsSelectTool() || IsLassoSelectTool() ? EVisibility::SelfHitTestInvisible : EVisibility::Collapsed;
@@ -585,6 +702,26 @@ ECheckBoxState SFoliageEdit::GetCheckState_Landscape() const
 	return FoliageEditMode->UISettings.GetFilterLandscape() ? ECheckBoxState::Checked : ECheckBoxState::Unchecked;
 }
 
+void SFoliageEdit::OnCheckStateChanged_SingleInstantiationMode(ECheckBoxState InState)
+{
+	FoliageEditMode->UISettings.SetIsInSingleInstantiationMode(InState == ECheckBoxState::Checked ? true : false);
+}
+
+ECheckBoxState SFoliageEdit::GetCheckState_SingleInstantiationMode() const
+{
+	return FoliageEditMode->UISettings.GetIsInSingleInstantiationMode() ? ECheckBoxState::Checked : ECheckBoxState::Unchecked;
+}
+
+void SFoliageEdit::OnCheckStateChanged_SpawnInCurrentLevelMode(ECheckBoxState InState)
+{
+	FoliageEditMode->UISettings.SetSpawnInCurrentLevelMode(InState == ECheckBoxState::Checked ? true : false);
+}
+
+ECheckBoxState SFoliageEdit::GetCheckState_SpawnInCurrentLevelMode() const
+{
+	return FoliageEditMode->UISettings.GetIsInSpawnInCurrentLevelMode() ? ECheckBoxState::Checked : ECheckBoxState::Unchecked;
+}
+
 FText SFoliageEdit::GetTooltipText_Landscape() const
 {
 	FText TooltipText;
@@ -760,4 +897,44 @@ EVisibility SFoliageEdit::GetVisibility_Filters() const
 	return EVisibility::Visible;
 }
 
+EVisibility SFoliageEdit::GetVisibility_Actions() const
+{
+	if (FoliageEditMode->UISettings.GetSelectToolSelected() || FoliageEditMode->UISettings.GetLassoSelectToolSelected())
+	{
+		return EVisibility::Visible;
+	}
+
+	return EVisibility::Collapsed;
+}
+
+EVisibility SFoliageEdit::GetVisibility_SingleInstantiationMode() const
+{
+	if (FoliageEditMode->UISettings.GetPaintToolSelected() || FoliageEditMode->UISettings.GetReapplyToolSelected() || FoliageEditMode->UISettings.GetLassoSelectToolSelected())
+	{
+		return EVisibility::Visible;
+	}
+
+	return EVisibility::Collapsed;
+}
+
+EVisibility SFoliageEdit::GetVisibility_SpawnInCurrentLevelMode() const
+{
+	if (FoliageEditMode->UISettings.GetPaintToolSelected())
+	{
+		return EVisibility::Visible;
+	}
+
+	return EVisibility::Collapsed;
+}
+
+EVisibility SFoliageEdit::GetVisibility_Options() const
+{
+	if (FoliageEditMode->UISettings.GetSelectToolSelected() || FoliageEditMode->UISettings.GetPaintBucketToolSelected())
+	{
+		return EVisibility::Collapsed;
+	}
+
+	return EVisibility::Visible;
+}
+
 #undef LOCTEXT_NAMESPACE
