Fix non-unity

diff --git a/Engine/Source/Editor/DetailCustomizations/Private/BlackboardEntryDetails.cpp b/Engine/Source/Editor/DetailCustomizations/Private/BlackboardEntryDetails.cpp
--- a/Engine/Source/Editor/DetailCustomizations/Private/BlackboardEntryDetails.cpp
+++ b/Engine/Source/Editor/DetailCustomizations/Private/BlackboardEntryDetails.cpp
@@ -6,6 +6,7 @@
 #include "SlateOptMacros.h"
 #include "IDetailChildrenBuilder.h"
 #include "DetailWidgetRow.h"
+#include "PropertyHandle.h"
 
 #define LOCTEXT_NAMESPACE "BlackboardEntryDetails"
 
diff --git a/Engine/Source/Editor/DetailCustomizations/Private/StaticMeshComponentDetails.cpp b/Engine/Source/Editor/DetailCustomizations/Private/StaticMeshComponentDetails.cpp
--- a/Engine/Source/Editor/DetailCustomizations/Private/StaticMeshComponentDetails.cpp
+++ b/Engine/Source/Editor/DetailCustomizations/Private/StaticMeshComponentDetails.cpp
@@ -6,6 +6,9 @@
 #include "DetailLayoutBuilder.h"
 #include "SNumericEntryBox.h"
 #include "SWidgetSwitcher.h"
+#include "IDetailPropertyRow.h"
+#include "DetailCategoryBuilder.h"
+#include "DetailWidgetRow.h"
 
 #define LOCTEXT_NAMESPACE "StaticMeshComponentDetails"
 
diff --git a/Engine/Source/Editor/UMGEditor/Private/WidgetBlueprintCompiler.h b/Engine/Source/Editor/UMGEditor/Private/WidgetBlueprintCompiler.h
--- a/Engine/Source/Editor/UMGEditor/Private/WidgetBlueprintCompiler.h
+++ b/Engine/Source/Editor/UMGEditor/Private/WidgetBlueprintCompiler.h
@@ -46,7 +46,7 @@ protected:
 protected:
 	UWidgetBlueprintGeneratedClass* NewWidgetBlueprintClass;
 
-	UWidgetGraphSchema* WidgetSchema;
+	class UWidgetGraphSchema* WidgetSchema;
 
 	// Map of properties created for widgets; to aid in debug data generation
 	TMap<class UWidget*, class UProperty*> WidgetToMemberVariableMap;
diff --git a/Engine/Source/Editor/UMGEditor/Private/WidgetBlueprintFactory.cpp b/Engine/Source/Editor/UMGEditor/Private/WidgetBlueprintFactory.cpp
--- a/Engine/Source/Editor/UMGEditor/Private/WidgetBlueprintFactory.cpp
+++ b/Engine/Source/Editor/UMGEditor/Private/WidgetBlueprintFactory.cpp
@@ -11,6 +11,7 @@
 
 
 #include "Components/CanvasPanel.h"
+#include "WidgetTree.h"
 
 #define LOCTEXT_NAMESPACE "UWidgetBlueprintFactory"
 
