Copying //UE4/Dev-Framework to //UE4/Dev-Main (Source: //UE4/Dev-Framework @ 3510040)

diff --git a/Engine/Source/Editor/BlueprintGraph/Private/K2Node_BreakStruct.cpp b/Engine/Source/Editor/BlueprintGraph/Private/K2Node_BreakStruct.cpp
--- a/Engine/Source/Editor/BlueprintGraph/Private/K2Node_BreakStruct.cpp
+++ b/Engine/Source/Editor/BlueprintGraph/Private/K2Node_BreakStruct.cpp
@@ -84,9 +84,7 @@ public:
 
 	void RegisterOutputTerm(FKismetFunctionContext& Context, UScriptStruct* StructType, UEdGraphPin* Net, FBPTerminal* ContextTerm)
 	{
-		UProperty* BoundProperty = FindField<UProperty>(StructType, *(Net->PinName));
-
-		if (BoundProperty != NULL)
+		if (UProperty* BoundProperty = FindField<UProperty>(StructType, *(Net->PinName)))
 		{
 			if (BoundProperty->HasAnyPropertyFlags(CPF_Deprecated) && Net->LinkedTo.Num())
 			{
@@ -117,19 +115,18 @@ public:
 	virtual void RegisterNets(FKismetFunctionContext& Context, UEdGraphNode* InNode) override
 	{
 		UK2Node_BreakStruct* Node = Cast<UK2Node_BreakStruct>(InNode);
-		check(NULL != Node);
+		check(Node);
 
-		if(!UK2Node_BreakStruct::CanBeBroken(Node->StructType))
+		if(!UK2Node_BreakStruct::CanBeBroken(Node->StructType, Node->IsIntermediateNode()))
 		{
 			CompilerContext.MessageLog.Warning(*LOCTEXT("BreakStruct_NoBreak_Error", "The structure cannot be broken using generic 'break' node @@. Try use specialized 'break' function if available.").ToString(), Node);
 		}
 
 		if(FBPTerminal* StructContextTerm = RegisterInputTerm(Context, Node))
 		{
-			for (int32 PinIndex = 0; PinIndex < Node->Pins.Num(); ++PinIndex)
+			for (UEdGraphPin* Pin : Node->Pins)
 			{
-				UEdGraphPin* Pin = Node->Pins[PinIndex];
-				if(NULL != Pin && EGPD_Output == Pin->Direction)
+				if(Pin && EGPD_Output == Pin->Direction)
 				{
 					RegisterOutputTerm(Context, Node->StructType, Pin, StructContextTerm);
 				}
@@ -145,29 +142,23 @@ UK2Node_BreakStruct::UK2Node_BreakStruct(const FObjectInitializer& ObjectInitial
 {
 }
 
-static bool CanCreatePinForProperty(const UProperty* Property, bool bIncludeEditAnywhere = true)
+static bool CanCreatePinForProperty(const UProperty* Property)
 {
 	const UEdGraphSchema_K2* Schema = GetDefault<UEdGraphSchema_K2>();
 	FEdGraphPinType DumbGraphPinType;
 	const bool bConvertable = Schema->ConvertPropertyToPinType(Property, /*out*/ DumbGraphPinType);
 
-	const bool bVisible = Property ? (bIncludeEditAnywhere ? Property->HasAnyPropertyFlags(CPF_Edit | CPF_BlueprintVisible) : Property->HasAnyPropertyFlags(CPF_BlueprintVisible) ): false;
+	const bool bVisible = (Property && Property->HasAnyPropertyFlags(CPF_BlueprintVisible));
 	return bVisible && bConvertable;
 }
 
-bool UK2Node_BreakStruct::CanBeBroken(const UScriptStruct* Struct, bool bIncludeEditAnywhere, bool bMustHaveValidProperties )
+bool UK2Node_BreakStruct::CanBeBroken(const UScriptStruct* Struct, const bool bForInternalUse)
 {
-	const UEdGraphSchema_K2* Schema = GetDefault<UEdGraphSchema_K2>();
-	if(Struct && Schema && !Struct->HasMetaData(TEXT("HasNativeBreak")))
+	if (Struct && !Struct->HasMetaData(TEXT("HasNativeBreak")) && UEdGraphSchema_K2::IsAllowableBlueprintVariableType(Struct, bForInternalUse))
 	{
-		if (!bMustHaveValidProperties && UEdGraphSchema_K2::IsAllowableBlueprintVariableType(Struct))
-		{
-			return true;
-		}
-
 		for (TFieldIterator<UProperty> It(Struct); It; ++It)
 		{
-			if (CanCreatePinForProperty(*It, bIncludeEditAnywhere))
+			if (CanCreatePinForProperty(*It))
 			{
 				return true;
 			}
@@ -253,6 +244,8 @@ FText UK2Node_BreakStruct::GetTooltipText() const
 
 void UK2Node_BreakStruct::ValidateNodeDuringCompilation(class FCompilerResultsLog& MessageLog) const
 {
+	Super::ValidateNodeDuringCompilation(MessageLog); 
+
 	if(!StructType)
 	{
 		MessageLog.Error(*LOCTEXT("NoStruct_Error", "No Struct in @@").ToString(), this);
@@ -377,7 +370,7 @@ void UK2Node_BreakStruct::GetMenuActions(FBlueprintActionDatabaseRegistrar& Acti
 	{
 		UBlueprintFieldNodeSpawner* NodeSpawner = nullptr;
 		
-		if (UK2Node_BreakStruct::CanBeBroken(Struct, /*bIncludeEditAnywhere =*/false))
+		if (UK2Node_BreakStruct::CanBeBroken(Struct))
 		{
 			NodeSpawner = UBlueprintFieldNodeSpawner::Create(NodeClass, Struct);
 			check(NodeSpawner != nullptr);
diff --git a/Engine/Source/Editor/BlueprintGraph/Private/K2Node_CallFunction.cpp b/Engine/Source/Editor/BlueprintGraph/Private/K2Node_CallFunction.cpp
--- a/Engine/Source/Editor/BlueprintGraph/Private/K2Node_CallFunction.cpp
+++ b/Engine/Source/Editor/BlueprintGraph/Private/K2Node_CallFunction.cpp
@@ -439,13 +439,14 @@ UK2Node_CallFunction::UK2Node_CallFunction(const FObjectInitializer& ObjectIniti
 	: Super(ObjectInitializer)
 	, bPinTooltipsValid(false)
 {
+	OrphanedPinSaveMode = ESaveOrphanPinMode::SaveAll;
 }
 
 
 bool UK2Node_CallFunction::IsDeprecated() const
 {
 	UFunction* Function = GetTargetFunction();
-	return (Function != NULL) && Function->HasMetaData(FBlueprintMetadata::MD_DeprecatedFunction);
+	return (Function && Function->HasMetaData(FBlueprintMetadata::MD_DeprecatedFunction));
 }
 
 bool UK2Node_CallFunction::ShouldWarnOnDeprecation() const
@@ -457,7 +458,7 @@ bool UK2Node_CallFunction::ShouldWarnOnDeprecation() const
 FString UK2Node_CallFunction::GetDeprecationMessage() const
 {
 	UFunction* Function = GetTargetFunction();
-	if ((Function != NULL) && Function->HasMetaData(FBlueprintMetadata::MD_DeprecationMessage))
+	if (Function && Function->HasMetaData(FBlueprintMetadata::MD_DeprecationMessage))
 	{
 		return FString::Printf(TEXT("%s %s"), *LOCTEXT("CallFunctionDeprecated_Warning", "@@ is deprecated;").ToString(), *Function->GetMetaData(FBlueprintMetadata::MD_DeprecationMessage));
 	}
@@ -1736,8 +1737,7 @@ void UK2Node_CallFunction::PostPasteNode()
 	Super::PostPasteNode();
 	FixupSelfMemberContext();
 
-	UFunction* Function = GetTargetFunction();
-	if(Function != NULL)
+	if (UFunction* Function = GetTargetFunction())
 	{
 		// After pasting we need to go through and ensure the hidden the self pins is correct in case the source blueprint had different metadata
 		TSet<FString> PinsToHide;
@@ -1753,8 +1753,8 @@ void UK2Node_CallFunction::PostPasteNode()
 		{
 			UEdGraphPin* Pin = Pins[PinIndex];
 
-			bool bIsSelfPin = ((Pin->PinName == DefaultToSelfMetaValue) || (Pin->PinName == WorldContextMetaValue));
-			bool bPinShouldBeHidden = PinsToHide.Contains(Pin->PinName) && (!bShowWorldContextPin || !bIsSelfPin);
+			const bool bIsSelfPin = ((Pin->PinName == DefaultToSelfMetaValue) || (Pin->PinName == WorldContextMetaValue));
+			const bool bPinShouldBeHidden = ((Pin->SubPins.Num() > 0) || (PinsToHide.Contains(Pin->PinName) && (!bShowWorldContextPin || !bIsSelfPin)));
 
 			if (bPinShouldBeHidden && !Pin->bHidden)
 			{
@@ -2324,6 +2324,8 @@ bool UK2Node_CallFunction::ReconnectPureExecPins(TArray<UEdGraphPin*>& OldPins)
 		}
 		if (PinExec)
 		{
+			PinExec->bSavePinIfOrphaned = false; 
+
 			// look for old then pin
 			UEdGraphPin* PinThen = nullptr;
 			for (int32 PinIdx = 0; PinIdx < OldPins.Num(); PinIdx++)
@@ -2336,6 +2338,8 @@ bool UK2Node_CallFunction::ReconnectPureExecPins(TArray<UEdGraphPin*>& OldPins)
 			}
 			if (PinThen)
 			{
+				PinThen->bSavePinIfOrphaned = false;
+
 				// reconnect all incoming links to old exec pin to the far end of the old then pin.
 				if (PinThen->LinkedTo.Num() > 0)
 				{
diff --git a/Engine/Source/Editor/BlueprintGraph/Private/K2Node_ConstructObjectFromClass.cpp b/Engine/Source/Editor/BlueprintGraph/Private/K2Node_ConstructObjectFromClass.cpp
--- a/Engine/Source/Editor/BlueprintGraph/Private/K2Node_ConstructObjectFromClass.cpp
+++ b/Engine/Source/Editor/BlueprintGraph/Private/K2Node_ConstructObjectFromClass.cpp
@@ -109,7 +109,8 @@ void UK2Node_ConstructObjectFromClass::CreatePinsForClass(UClass* InClass, TArra
 			bIsSettableExternally &&
 			Property->HasAllPropertyFlags(CPF_BlueprintVisible) &&
 			!bIsDelegate &&
-			(NULL == FindPin(Property->GetName()) ) )
+			(nullptr == FindPin(Property->GetName()) ) &&
+			FBlueprintEditorUtils::PropertyStillExists(Property))
 		{
 			UEdGraphPin* Pin = CreatePin(EGPD_Input, FString(), FString(), nullptr, Property->GetName());
 			const bool bPinGood = (Pin != NULL) && K2Schema->ConvertPropertyToPinType(Property, /*out*/ Pin->PinType);
@@ -137,7 +138,7 @@ void UK2Node_ConstructObjectFromClass::CreatePinsForClass(UClass* InClass, TArra
 
 	// Change class of output pin
 	UEdGraphPin* ResultPin = GetResultPin();
-	ResultPin->PinType.PinSubCategoryObject = InClass;
+	ResultPin->PinType.PinSubCategoryObject = InClass->GetAuthoritativeClass();
 }
 
 UClass* UK2Node_ConstructObjectFromClass::GetClassToSpawn(const TArray<UEdGraphPin*>* InPinsToSearch /*=NULL*/) const
@@ -187,16 +188,14 @@ void UK2Node_ConstructObjectFromClass::OnClassPinChanged()
 	const UEdGraphSchema_K2* K2Schema = GetDefault<UEdGraphSchema_K2>();
 
 	// Remove all pins related to archetype variables
-	TArray<UEdGraphPin*> OldPins = Pins;
 	TArray<UEdGraphPin*> OldClassPins;
 
-	for (int32 i = 0; i < OldPins.Num(); i++)
+	for (int32 i = Pins.Num() - 1; i >= 0; --i)
 	{
-		UEdGraphPin* OldPin = OldPins[i];
+		UEdGraphPin* OldPin = Pins[i];
 		if (IsSpawnVarPin(OldPin))
 		{
-			OldPin->MarkPendingKill();
-			Pins.Remove(OldPin);
+			Pins.RemoveAt(i, 1, false);
 			OldClassPins.Add(OldPin);
 		}
 	}
@@ -205,7 +204,7 @@ void UK2Node_ConstructObjectFromClass::OnClassPinChanged()
 
 	UClass* UseSpawnClass = GetClassToSpawn();
 	TArray<UEdGraphPin*> NewClassPins;
-	if (UseSpawnClass != NULL)
+	if (UseSpawnClass)
 	{
 		CreatePinsForClass(UseSpawnClass, &NewClassPins);
 	}
@@ -213,9 +212,6 @@ void UK2Node_ConstructObjectFromClass::OnClassPinChanged()
 	// Rewire the old pins to the new pins so connections are maintained if possible
 	RewireOldPinsToNewPins(OldClassPins, NewClassPins);
 
-	// Destroy the old pins
-	DestroyPinList(OldClassPins);
-
 	// Refresh the UI for the graph so the pin changes show up
 	UEdGraph* Graph = GetGraph();
 	Graph->NotifyGraphChanged();
diff --git a/Engine/Source/Editor/BlueprintGraph/Private/K2Node_EditablePinBase.cpp b/Engine/Source/Editor/BlueprintGraph/Private/K2Node_EditablePinBase.cpp
--- a/Engine/Source/Editor/BlueprintGraph/Private/K2Node_EditablePinBase.cpp
+++ b/Engine/Source/Editor/BlueprintGraph/Private/K2Node_EditablePinBase.cpp
@@ -81,32 +81,13 @@ UEdGraphPin* UK2Node_EditablePinBase::CreateUserDefinedPin(const FString& InPinN
 
 void UK2Node_EditablePinBase::RemoveUserDefinedPin(TSharedPtr<FUserPinInfo> PinToRemove)
 {
-	// Try to find the pin with the same name and params as the specified description, if any
-	const FString PinName = PinToRemove->PinName;
-	for(int32 i = 0; i < Pins.Num(); i++)
-	{
-		UEdGraphPin* Pin = Pins[i];
-		if( Pin->PinName == PinName )
-		{
-			Pins.Remove(Pin);
-			Pin->MarkPendingKill();
-
-			if (UBlueprint* Blueprint = GetBlueprint())
-			{
-				FKismetDebugUtilities::RemovePinWatch(Blueprint, Pin);
-			}
-		}
-	}
-
-	// Remove the description from the user-defined pins array
-	UserDefinedPins.Remove(PinToRemove);
+	RemoveUserDefinedPinByName(PinToRemove->PinName);
 }
 
 void UK2Node_EditablePinBase::RemoveUserDefinedPinByName(const FString& PinName)
 {
-	for (int32 i = 0; i < Pins.Num(); i++)
+	for (UEdGraphPin* Pin : Pins)
 	{
-		UEdGraphPin* Pin = Pins[i];
 		if (Pin->PinName == PinName)
 		{
 			Pin->Modify();
@@ -290,7 +271,7 @@ void UK2Node_EditablePinBase::Serialize(FArchive& Ar)
 			UserDefinedPins.Add(PinInfo);
 		}
 	}
-	else
+	else if(Ar.IsSaving())
 	{
 		SerializedItems.Empty(UserDefinedPins.Num());
 
@@ -300,14 +281,14 @@ void UK2Node_EditablePinBase::Serialize(FArchive& Ar)
 		}
 
 		Ar << SerializedItems;
-
-		if (Ar.IsModifyingWeakAndStrongReferences())
+	}
+	else
+	{
+		// We want to avoid destroying and recreating FUserPinInfo, because that will invalidate 
+		// any WeakPtrs to those entries:
+		for(TSharedPtr<FUserPinInfo>& PinInfo : UserDefinedPins )
 		{
-			UserDefinedPins.Empty(SerializedItems.Num());
-			for (int32 Index = 0; Index < SerializedItems.Num(); ++Index)
-			{
-				UserDefinedPins.Add(MakeShareable(new FUserPinInfo(SerializedItems[Index])));
-			}
+			Ar << *PinInfo;
 		}
 	}
 }
diff --git a/Engine/Source/Editor/BlueprintGraph/Private/K2Node_Event.cpp b/Engine/Source/Editor/BlueprintGraph/Private/K2Node_Event.cpp
--- a/Engine/Source/Editor/BlueprintGraph/Private/K2Node_Event.cpp
+++ b/Engine/Source/Editor/BlueprintGraph/Private/K2Node_Event.cpp
@@ -388,7 +388,8 @@ bool UK2Node_Event::IsFunctionEntryCompatible(const UK2Node_FunctionEntry* Entry
 		if( CurPin->PinType.PinCategory == K2Schema->PC_Exec 
 			|| CurPin->PinType.PinSubCategory == K2Schema->PSC_Self
 			|| CurPin->PinName == DelegateOutputName
-			|| CurPin->Direction == EGPD_Input )
+			|| CurPin->Direction == EGPD_Input
+			|| CurPin->ParentPin != nullptr )
 		{
 			EventPins.RemoveAt(i, 1);
 			i--;
@@ -400,7 +401,9 @@ bool UK2Node_Event::IsFunctionEntryCompatible(const UK2Node_FunctionEntry* Entry
 		const UEdGraphPin* CurPin = EntryPins[i];
 		if( CurPin->PinType.PinCategory == K2Schema->PC_Exec 
 			|| CurPin->PinType.PinSubCategory == K2Schema->PSC_Self
-			|| CurPin->Direction == EGPD_Input)
+			|| CurPin->PinName == DelegateOutputName
+			|| CurPin->Direction == EGPD_Input
+			|| CurPin->ParentPin != nullptr )
 		{
 			EntryPins.RemoveAt(i, 1);
 			i--;
diff --git a/Engine/Source/Editor/BlueprintGraph/Private/K2Node_ExecutionSequence.cpp b/Engine/Source/Editor/BlueprintGraph/Private/K2Node_ExecutionSequence.cpp
--- a/Engine/Source/Editor/BlueprintGraph/Private/K2Node_ExecutionSequence.cpp
+++ b/Engine/Source/Editor/BlueprintGraph/Private/K2Node_ExecutionSequence.cpp
@@ -211,8 +211,9 @@ FString UK2Node_ExecutionSequence::GetUniquePinName()
 	return NewPinName;
 }
 
-void UK2Node_ExecutionSequence::AddPinToExecutionNode()
+void UK2Node_ExecutionSequence::AddInputPin()
 {
+	Modify();
 	const UEdGraphSchema_K2* K2Schema = GetDefault<UEdGraphSchema_K2>();
 	CreatePin(EGPD_Output, K2Schema->PC_Exec, FString(), nullptr, GetUniquePinName());
 }
diff --git a/Engine/Source/Editor/BlueprintGraph/Private/K2Node_FormatText.cpp b/Engine/Source/Editor/BlueprintGraph/Private/K2Node_FormatText.cpp
--- a/Engine/Source/Editor/BlueprintGraph/Private/K2Node_FormatText.cpp
+++ b/Engine/Source/Editor/BlueprintGraph/Private/K2Node_FormatText.cpp
@@ -126,8 +126,8 @@ FString UK2Node_FormatText::GetUniquePinName()
 
 void UK2Node_FormatText::PostEditChangeProperty(struct FPropertyChangedEvent& PropertyChangedEvent)
 {
-	FName PropertyName = (PropertyChangedEvent.Property != NULL) ? PropertyChangedEvent.Property->GetFName() : NAME_None;
-	if (PropertyName == TEXT("PinNames"))
+	const FName PropertyName = (PropertyChangedEvent.Property  ? PropertyChangedEvent.Property->GetFName() : NAME_None);
+	if (PropertyName == GET_MEMBER_NAME_CHECKED(UK2Node_FormatText, PinNames))
 	{
 		ReconstructNode();
 		GetGraph()->NotifyGraphChanged();
diff --git a/Engine/Source/Editor/BlueprintGraph/Private/K2Node_FunctionEntry.cpp b/Engine/Source/Editor/BlueprintGraph/Private/K2Node_FunctionEntry.cpp
--- a/Engine/Source/Editor/BlueprintGraph/Private/K2Node_FunctionEntry.cpp
+++ b/Engine/Source/Editor/BlueprintGraph/Private/K2Node_FunctionEntry.cpp
@@ -192,12 +192,52 @@ UK2Node_FunctionEntry::UK2Node_FunctionEntry(const FObjectInitializer& ObjectIni
 	bEnforceConstCorrectness = true;
 }
 
+void UK2Node_FunctionEntry::PreSave(const class ITargetPlatform* TargetPlatform)
+{
+	Super::PreSave(TargetPlatform);
+	
+	const UBlueprint* Blueprint = HasValidBlueprint() ? GetBlueprint() : nullptr;
+	if (LocalVariables.Num() > 0 && Blueprint)
+	{
+		// This code is here as it's unsafe to call when GIsSavingPackage is true
+		UFunction* Function = FindField<UFunction>(Blueprint->SkeletonGeneratedClass, *GetOuter()->GetName());
+
+		if (Function)
+		{
+			if (Function->GetStructureSize() > 0 || !ensure(Function->PropertyLink == nullptr))
+			{
+				TSharedPtr<FStructOnScope> LocalVarData = MakeShareable(new FStructOnScope(Function));
+
+				for (TFieldIterator<UProperty> It(Function); It; ++It)
+				{
+					if (const UProperty* Property = *It)
+					{
+						const UStructProperty* PotentialUDSProperty = Cast<const UStructProperty>(Property);
+						// UDS requires default data even when the LocalVariable value is empty
+						const bool bUDSProperty = PotentialUDSProperty && Cast<const UUserDefinedStruct>(PotentialUDSProperty->Struct);
+
+						for (FBPVariableDescription& LocalVar : LocalVariables)
+						{
+							if (LocalVar.VarName == Property->GetFName() && (bUDSProperty || !LocalVar.DefaultValue.IsEmpty()))
+							{
+								// Go to property and back, this handles redirector fixup and will sanitize the output
+								// The asset registry only knows about these references because when the node is expanded it turns into a hard reference
+								FBlueprintEditorUtils::PropertyValueFromString(Property, LocalVar.DefaultValue, LocalVarData->GetStructMemory());
+								FBlueprintEditorUtils::PropertyValueToString(Property, LocalVarData->GetStructMemory(), LocalVar.DefaultValue);
+							}
+						}
+					}
+				}
+			}
+		}
+	}
+}
+
 void UK2Node_FunctionEntry::Serialize(FArchive& Ar)
 {
 	Super::Serialize(Ar);
 
 	Ar.UsingCustomVersion(FBlueprintsObjectVersion::GUID);
-	Ar.UsingCustomVersion(FFrameworkObjectVersion::GUID);
 
 	if (Ar.IsLoading())
 	{
@@ -277,7 +317,11 @@ void UK2Node_FunctionEntry::AllocateDefaultPins()
 
 	UFunction* Function = FindField<UFunction>(SignatureClass, SignatureName);
 
-	if (Function == nullptr)
+	// This FindDelegateSignature call was added to support multiple UClasses in a single file.
+	// For blueprint declared functions it can generate an "Ambiguous search" warning, and may also
+	// be very slow:
+	bool bIsNativeFunction = (SignatureClass == nullptr || SignatureClass->HasAnyClassFlags(CLASS_Native));
+	if (Function == nullptr && bIsNativeFunction)
 	{
 		Function = FindDelegateSignature(SignatureName);
 	}
diff --git a/Engine/Source/Editor/BlueprintGraph/Private/K2Node_GetClassDefaults.cpp b/Engine/Source/Editor/BlueprintGraph/Private/K2Node_GetClassDefaults.cpp
--- a/Engine/Source/Editor/BlueprintGraph/Private/K2Node_GetClassDefaults.cpp
+++ b/Engine/Source/Editor/BlueprintGraph/Private/K2Node_GetClassDefaults.cpp
@@ -188,12 +188,23 @@ UK2Node_GetClassDefaults::UK2Node_GetClassDefaults(const FObjectInitializer& Obj
 {
 }
 
-void UK2Node_GetClassDefaults::PostEditChangeProperty(struct FPropertyChangedEvent& PropertyChangedEvent)
+void UK2Node_GetClassDefaults::PreEditChange(UProperty* PropertyThatWillChange)
 {
-	// Rebuild the node if the "Show Pin" property was changed
-	FName PropertyName = (PropertyChangedEvent.Property != NULL) ? PropertyChangedEvent.Property->GetFName() : NAME_None;
-	if((PropertyName == TEXT("bShowPin")))
+	Super::PreEditChange(PropertyThatWillChange);
+
+	if (PropertyThatWillChange && PropertyThatWillChange->GetFName() == GET_MEMBER_NAME_CHECKED(FOptionalPinFromProperty, bShowPin))
+	{
+		FOptionalPinManager::CacheShownPins(ShowPinForProperties, OldShownPins);
+	}
+}
+
+void UK2Node_GetClassDefaults::PostEditChangeProperty(FPropertyChangedEvent& PropertyChangedEvent)
+{
+	const FName PropertyName = (PropertyChangedEvent.Property ? PropertyChangedEvent.Property->GetFName() : NAME_None);
+
+	if (PropertyName == GET_MEMBER_NAME_CHECKED(FOptionalPinFromProperty, bShowPin))
 	{
+		FOptionalPinManager::EvaluateOldShownPins(ShowPinForProperties, OldShownPins, this);
 		GetSchema()->ReconstructNode(*this);
 	}
 
diff --git a/Engine/Source/Editor/BlueprintGraph/Private/K2Node_LoadAsset.cpp b/Engine/Source/Editor/BlueprintGraph/Private/K2Node_LoadAsset.cpp
--- a/Engine/Source/Editor/BlueprintGraph/Private/K2Node_LoadAsset.cpp
+++ b/Engine/Source/Editor/BlueprintGraph/Private/K2Node_LoadAsset.cpp
@@ -134,7 +134,23 @@ void UK2Node_LoadAsset::ExpandNode(class FKismetCompilerContext& CompilerContext
 	{
 		UEdGraphPin* AssetPin = FindPin(GetInputPinName());
 		ensure(CallFunctionAssetPin);
-		bIsErrorFree &= AssetPin && CallFunctionAssetPin && CompilerContext.MovePinLinksToIntermediate(*AssetPin, *CallFunctionAssetPin).CanSafeConnect();
+
+		if (AssetPin && CallFunctionAssetPin)
+		{
+			if (AssetPin->LinkedTo.Num() > 0)
+			{
+				bIsErrorFree &= CompilerContext.MovePinLinksToIntermediate(*AssetPin, *CallFunctionAssetPin).CanSafeConnect();
+			}
+			else
+			{
+				// Copy literal value
+				CallFunctionAssetPin->DefaultValue = AssetPin->DefaultValue;
+			}
+		}
+		else
+		{
+			bIsErrorFree = false;
+		}
 	}
 
 	// Create OnLoadEvent
diff --git a/Engine/Source/Editor/BlueprintGraph/Private/K2Node_MakeArray.cpp b/Engine/Source/Editor/BlueprintGraph/Private/K2Node_MakeArray.cpp
--- a/Engine/Source/Editor/BlueprintGraph/Private/K2Node_MakeArray.cpp
+++ b/Engine/Source/Editor/BlueprintGraph/Private/K2Node_MakeArray.cpp
@@ -189,7 +189,7 @@ void UK2Node_MakeArray::NotifyPinConnectionListChanged(UEdGraphPin* Pin)
 	if (Pin->LinkedTo.Num() > 0)
 	{
 		// Just made a connection, was it the first and was it the parent pin?
-		if (Pin->ParentPin == nullptr && NumPinsWithLinks == 1)
+		if (Pin->ParentPin == nullptr && NumPinsWithLinks == 1 && (OutputPin->PinType.PinCategory == UEdGraphSchema_K2::PC_Wildcard || Pin->LinkedTo[0]->PinType.PinCategory != UEdGraphSchema_K2::PC_Wildcard))
 		{
 			// Update the types on all the pins
 			OutputPin->PinType = Pin->LinkedTo[0]->PinType;
@@ -386,7 +386,6 @@ FSlateIcon UK2Node_MakeArray::GetIconAndTint(FLinearColor& OutColor) const
 
 void UK2Node_MakeArray::AddInputPin()
 {
-	FScopedTransaction Transaction( LOCTEXT("AddPinTx", "AddPin") );
 	Modify();
 
 	++NumInputs;
diff --git a/Engine/Source/Editor/BlueprintGraph/Private/K2Node_MakeStruct.cpp b/Engine/Source/Editor/BlueprintGraph/Private/K2Node_MakeStruct.cpp
--- a/Engine/Source/Editor/BlueprintGraph/Private/K2Node_MakeStruct.cpp
+++ b/Engine/Source/Editor/BlueprintGraph/Private/K2Node_MakeStruct.cpp
@@ -83,22 +83,23 @@ void UK2Node_MakeStruct::FMakeStructPinManager::CustomizePinData(UEdGraphPin* Pi
 	}
 }
 
-static bool CanBeExposed(const UProperty* Property, bool bIncludeEditOnly = true)
+static bool CanBeExposed(const UProperty* Property)
 {
 	if (Property)
 	{
 		const UEdGraphSchema_K2* Schema = GetDefault<UEdGraphSchema_K2>();
 		check(Schema);
 
-		FEdGraphPinType DumbGraphPinType;
-		const bool bConvertable = Schema->ConvertPropertyToPinType(Property, /*out*/ DumbGraphPinType);
-
-		//TODO: remove CPF_Edit in a future release
-		const bool bVisible = (bIncludeEditOnly ? Property->HasAnyPropertyFlags(CPF_BlueprintVisible | CPF_Edit) : Property->HasAnyPropertyFlags(CPF_BlueprintVisible)) && !(Property->ArrayDim > 1);
-		const bool bBlueprintReadOnly = Property->HasAllPropertyFlags(CPF_BlueprintReadOnly);
-		if (bVisible && bConvertable && !bBlueprintReadOnly)
+		if (!Property->HasAllPropertyFlags(CPF_BlueprintReadOnly))
 		{
-			return true;
+			if (Property->HasAllPropertyFlags(CPF_BlueprintVisible) && !(Property->ArrayDim > 1))
+			{
+				FEdGraphPinType DumbGraphPinType;
+				if (Schema->ConvertPropertyToPinType(Property, /*out*/ DumbGraphPinType))
+				{
+					return true;
+				}
+			}
 		}
 	}
 	return false;
@@ -155,30 +156,22 @@ void UK2Node_MakeStruct::AllocateDefaultPins()
 
 void UK2Node_MakeStruct::ValidateNodeDuringCompilation(class FCompilerResultsLog& MessageLog) const
 {
+	Super::ValidateNodeDuringCompilation(MessageLog);
+
 	if(!StructType)
 	{
 		MessageLog.Error(*LOCTEXT("NoStruct_Error", "No Struct in @@").ToString(), this);
 	}
 	else
 	{
-		bool bHasAnyBlueprintVisibleProperty = false;
 		for (TFieldIterator<UProperty> It(StructType); It; ++It)
 		{
 			const UProperty* Property = *It;
 			if (CanBeExposed(Property))
 			{
-				const bool bIsBlueprintVisible = Property->HasAnyPropertyFlags(CPF_BlueprintVisible) || (Property->GetOwnerStruct() && Property->GetOwnerStruct()->IsA<UUserDefinedStruct>());
-				bHasAnyBlueprintVisibleProperty |= bIsBlueprintVisible;
-
-				const UEdGraphPin* Pin = FindPin(Property->GetName());
-
-				if (Pin && !bIsBlueprintVisible)
-				{
-					MessageLog.Warning(*LOCTEXT("PropertyIsNotBPVisible_Warning", "@@ - the native property is not tagged as BlueprintReadWrite, the pin will be removed in a future release.").ToString(), Pin);
-				}
-
 				if (Property->ArrayDim > 1)
 				{
+					const UEdGraphPin* Pin = FindPin(Property->GetName());
 					MessageLog.Warning(*LOCTEXT("StaticArray_Warning", "@@ - the native property is a static array, which is not supported by blueprints").ToString(), Pin);
 				}
 			}
@@ -243,18 +236,18 @@ FLinearColor UK2Node_MakeStruct::GetNodeTitleColor() const
 	return UK2Node::GetNodeTitleColor();
 }
 
-bool UK2Node_MakeStruct::CanBeMade(const UScriptStruct* Struct, bool bIncludeEditOnly, bool bMustHaveValidProperties )
+bool UK2Node_MakeStruct::CanBeMade(const UScriptStruct* Struct, const bool bForInternalUse)
 {
-	if (Struct && !Struct->HasMetaData(TEXT("HasNativeMake")))
-	{
-		if (!bMustHaveValidProperties && UEdGraphSchema_K2::IsAllowableBlueprintVariableType(Struct))
-		{
-			return true;
-		}
+	return (Struct && !Struct->HasMetaData(TEXT("HasNativeMake")) && UEdGraphSchema_K2::IsAllowableBlueprintVariableType(Struct, bForInternalUse));
+}
 
+bool UK2Node_MakeStruct::CanBeSplit(const UScriptStruct* Struct)
+{
+	if (CanBeMade(Struct))
+	{
 		for (TFieldIterator<UProperty> It(Struct); It; ++It)
 		{
-			if (CanBeExposed(*It, bIncludeEditOnly))
+			if (CanBeExposed(*It))
 			{
 				return true;
 			}
@@ -327,7 +320,7 @@ void UK2Node_MakeStruct::GetMenuActions(FBlueprintActionDatabaseRegistrar& Actio
 	{
 		UBlueprintFieldNodeSpawner* NodeSpawner = nullptr;
 		
-		if (UK2Node_MakeStruct::CanBeMade(Struct, /*bIncludeEditOnly =*/false))
+		if (UK2Node_MakeStruct::CanBeMade(Struct))
 		{
 			NodeSpawner = UBlueprintFieldNodeSpawner::Create(NodeClass, Struct);
 			check(NodeSpawner != nullptr);
diff --git a/Engine/Source/Editor/BlueprintGraph/Private/K2Node_MathExpression.cpp b/Engine/Source/Editor/BlueprintGraph/Private/K2Node_MathExpression.cpp
--- a/Engine/Source/Editor/BlueprintGraph/Private/K2Node_MathExpression.cpp
+++ b/Engine/Source/Editor/BlueprintGraph/Private/K2Node_MathExpression.cpp
@@ -2438,6 +2438,7 @@ UK2Node_MathExpression::UK2Node_MathExpression(const FObjectInitializer& ObjectI
 	bCanRenameNode = true;
 
 	bMadeAfterRotChange = false;
+	OrphanedPinSaveMode = ESaveOrphanPinMode::SaveNone;
 }
 
 void UK2Node_MathExpression::Serialize(FArchive& Ar)
diff --git a/Engine/Source/Editor/BlueprintGraph/Private/K2Node_Select.cpp b/Engine/Source/Editor/BlueprintGraph/Private/K2Node_Select.cpp
--- a/Engine/Source/Editor/BlueprintGraph/Private/K2Node_Select.cpp
+++ b/Engine/Source/Editor/BlueprintGraph/Private/K2Node_Select.cpp
@@ -167,7 +167,7 @@ public:
 
 		// Create a term to determine if the compare was successful or not
 		FBPTerminal* BoolTerm = Context.CreateLocalTerminal();
-		BoolTerm->Type.PinCategory = CompilerContext.GetSchema()->PC_Boolean;
+		BoolTerm->Type.PinCategory = UEdGraphSchema_K2::PC_Boolean;
 		BoolTerm->Source = Node;
 		BoolTerm->Name = Context.NetNameMap->MakeValidName(Node) + TEXT("_CmpSuccess");
 		BoolTermMap.Add(Node, BoolTerm);
@@ -222,7 +222,7 @@ public:
 				// Create a local int for use in the equality call function below (LiteralTerm = the right hand side of the EqualEqual_IntInt or NotEqual_BoolBool statement)
 				FBPTerminal* LiteralTerm = Context.CreateLocalTerminal(ETerminalSpecification::TS_Literal);
 				LiteralTerm->bIsLiteral = true;
-				LiteralTerm->Type.PinCategory = CompilerContext.GetSchema()->PC_Int;
+				LiteralTerm->Type.PinCategory = UEdGraphSchema_K2::PC_Int;
 
 				if (UEnum* NodeEnum = SelectNode->GetEnum())
 				{
@@ -283,7 +283,7 @@ public:
 					// Create a local int for use in the equality call function below (LiteralTerm = the right hand side of the EqualEqual_IntInt or NotEqual_BoolBool statement)
 					FBPTerminal* LiteralStringTerm = Context.CreateLocalTerminal(ETerminalSpecification::TS_Literal);
 					LiteralStringTerm->bIsLiteral = true;
-					LiteralStringTerm->Type.PinCategory = CompilerContext.GetSchema()->PC_String;
+					LiteralStringTerm->Type.PinCategory = UEdGraphSchema_K2::PC_String;
 
 					FString SelectionNodeType(TEXT("NONE"));
 					if (IndexPin)
@@ -334,16 +334,15 @@ UK2Node_Select::UK2Node_Select(const FObjectInitializer& ObjectInitializer)
 
 	NumOptionPins = 2;
 
-	const UEdGraphSchema_K2* Schema = GetDefault<UEdGraphSchema_K2>();
-	IndexPinType.PinCategory = Schema->PC_Wildcard;
-	IndexPinType.PinSubCategory = Schema->PSC_Index;
-	IndexPinType.PinSubCategoryObject = NULL;
+	IndexPinType.PinCategory = UEdGraphSchema_K2::PC_Wildcard;
+	IndexPinType.PinSubCategory = UEdGraphSchema_K2::PSC_Index;
+	IndexPinType.PinSubCategoryObject = nullptr;
+
+	OrphanedPinSaveMode = ESaveOrphanPinMode::SaveNone;
 }
 
 void UK2Node_Select::AllocateDefaultPins()
 {
-	const UEdGraphSchema_K2* Schema = GetDefault<UEdGraphSchema_K2>();
-
 	// To refresh, just in case it changed
 	SetEnum(Enum, true);
 
@@ -368,13 +367,13 @@ void UK2Node_Select::AllocateDefaultPins()
 			UEdGraphPin* TempPin = FindPin(PinName);
 			if (!TempPin)
 			{
-				NewPin = CreatePin(EGPD_Input, Schema->PC_Wildcard, FString(), nullptr, PinName);
+				NewPin = CreatePin(EGPD_Input, UEdGraphSchema_K2::PC_Wildcard, FString(), nullptr, PinName);
 			}
 		}
 		else
 		{
 			const FString PinName = FString::Printf(TEXT("Option %d"), Idx);
-			NewPin = CreatePin(EGPD_Input, Schema->PC_Wildcard, FString(), nullptr, PinName);
+			NewPin = CreatePin(EGPD_Input, UEdGraphSchema_K2::PC_Wildcard, FString(), nullptr, PinName);
 		}
 
 		if (NewPin)
@@ -395,7 +394,7 @@ void UK2Node_Select::AllocateDefaultPins()
 	CreatePin(EGPD_Input, IndexPinType.PinCategory, IndexPinType.PinSubCategory, IndexPinType.PinSubCategoryObject.Get(), TEXT("Index"));
 
 	// Create the return value
-	UEdGraphPin* ReturnPin = CreatePin(EGPD_Output, Schema->PC_Wildcard, FString(), nullptr, Schema->PN_ReturnValue);
+	UEdGraphPin* ReturnPin = CreatePin(EGPD_Output, UEdGraphSchema_K2::PC_Wildcard, FString(), nullptr, UEdGraphSchema_K2::PN_ReturnValue);
 	ReturnPin->bDisplayAsMutableRef = UseSelectRef;
 
 	Super::AllocateDefaultPins();
@@ -753,16 +752,16 @@ void UK2Node_Select::RemoveOptionPinToNode()
 
 void UK2Node_Select::SetEnum(UEnum* InEnum, bool bForceRegenerate)
 {
-	const UEdGraphSchema_K2* Schema = GetDefault<UEdGraphSchema_K2>();
-
 	UEnum* PrevEnum = Enum;
 	Enum = InEnum;
 
-	if ((PrevEnum != Enum) || bForceRegenerate)
+	OrphanedPinSaveMode = (Enum ? ESaveOrphanPinMode::SaveAll : ESaveOrphanPinMode::SaveNone);
+
+	if (bForceRegenerate || (PrevEnum != Enum))
 	{
 		// regenerate enum name list
-		EnumEntries.Empty();
-		EnumEntryFriendlyNames.Empty();
+		EnumEntries.Reset();
+		EnumEntryFriendlyNames.Reset();
 
 		if (Enum)
 		{
@@ -895,17 +894,23 @@ void UK2Node_Select::PinTypeChanged(UEdGraphPin* Pin)
 		{
 			IndexPinType = Pin->PinType;
 
+			// Since it is an interactive action we want the pins to go away regardless of the new type
+			for (UEdGraphPin* PinToDiscard : Pins)
+			{
+				PinToDiscard->bSavePinIfOrphaned = false;
+			}
+
 			if (IndexPinType.PinSubCategoryObject.IsValid())
 			{
 				SetEnum(Cast<UEnum>(IndexPinType.PinSubCategoryObject.Get()));
 			}
 			else if (Enum)
 			{
-				SetEnum(NULL);
+				SetEnum(nullptr);
 			}
 
 			// Remove all but two options if we switched to a bool index
-			if (IndexPinType.PinCategory == Schema->PC_Boolean)
+			if (IndexPinType.PinCategory == UEdGraphSchema_K2::PC_Boolean)
 			{
 				NumOptionPins = 2;
 			}
@@ -987,8 +992,7 @@ FSlateIcon UK2Node_Select::GetIconAndTint(FLinearColor& OutColor) const
 
 bool UK2Node_Select::IsConnectionDisallowed(const UEdGraphPin* MyPin, const UEdGraphPin* OtherPin, FString& OutReason) const
 {
-	const UEdGraphSchema_K2* K2Schema = GetDefault<UEdGraphSchema_K2>();
-	if (OtherPin && (OtherPin->PinType.PinCategory == K2Schema->PC_Exec))
+	if (OtherPin && (OtherPin->PinType.PinCategory == UEdGraphSchema_K2::PC_Exec))
 	{
 		OutReason = LOCTEXT("ExecConnectionDisallowed", "Cannot connect with Exec pin.").ToString();
 		return true;
diff --git a/Engine/Source/Editor/BlueprintGraph/Private/K2Node_SetFieldsInStruct.cpp b/Engine/Source/Editor/BlueprintGraph/Private/K2Node_SetFieldsInStruct.cpp
--- a/Engine/Source/Editor/BlueprintGraph/Private/K2Node_SetFieldsInStruct.cpp
+++ b/Engine/Source/Editor/BlueprintGraph/Private/K2Node_SetFieldsInStruct.cpp
@@ -213,7 +213,7 @@ bool UK2Node_SetFieldsInStruct::ShowCustomPinActions(const UEdGraphPin* Pin, boo
 		&& !Schema->IsMetaPin(*Pin);
 }
 
-void UK2Node_SetFieldsInStruct::RemoveFieldPins(const UEdGraphPin* Pin, EPinsToRemove Selection)
+void UK2Node_SetFieldsInStruct::RemoveFieldPins(UEdGraphPin* Pin, EPinsToRemove Selection)
 {
 	if (ShowCustomPinActions(Pin, false) && (Pin->GetOwningNodeUnchecked() == this))
 	{
@@ -234,6 +234,7 @@ void UK2Node_SetFieldsInStruct::RemoveFieldPins(const UEdGraphPin* Pin, EPinsToR
 			{
 				bWasChanged = true;
 				OptionalProperty.bShowPin = false;
+				Pin->bSavePinIfOrphaned = false;
 			}
 		}
 
diff --git a/Engine/Source/Editor/BlueprintGraph/Private/K2Node_SpawnActorFromClass.cpp b/Engine/Source/Editor/BlueprintGraph/Private/K2Node_SpawnActorFromClass.cpp
--- a/Engine/Source/Editor/BlueprintGraph/Private/K2Node_SpawnActorFromClass.cpp
+++ b/Engine/Source/Editor/BlueprintGraph/Private/K2Node_SpawnActorFromClass.cpp
@@ -104,7 +104,8 @@ void UK2Node_SpawnActorFromClass::CreatePinsForClass(UClass* InClass, TArray<UEd
 			bIsSettableExternally &&
 			Property->HasAllPropertyFlags(CPF_BlueprintVisible) &&
 			!bIsDelegate &&
-			(NULL == FindPin(Property->GetName()) ) )
+			(nullptr == FindPin(Property->GetName()) ) &&
+			FBlueprintEditorUtils::PropertyStillExists(Property) )
 		{
 			UEdGraphPin* Pin = CreatePin(EGPD_Input, FString(), FString(), nullptr, Property->GetName());
 			const bool bPinGood = (Pin != nullptr) && K2Schema->ConvertPropertyToPinType(Property, /*out*/ Pin->PinType);
@@ -363,9 +364,6 @@ void UK2Node_SpawnActorFromClass::OnClassPinChanged()
 	// Rewire the old pins to the new pins so connections are maintained if possible
 	RewireOldPinsToNewPins(OldClassPins, NewClassPins);
 
-	// Destroy the old pins
-	DestroyPinList(OldClassPins);
-
 	// Refresh the UI for the graph so the pin changes show up
 	UEdGraph* Graph = GetGraph();
 	Graph->NotifyGraphChanged();
diff --git a/Engine/Source/Editor/BlueprintGraph/Private/K2Node_StructMemberGet.cpp b/Engine/Source/Editor/BlueprintGraph/Private/K2Node_StructMemberGet.cpp
--- a/Engine/Source/Editor/BlueprintGraph/Private/K2Node_StructMemberGet.cpp
+++ b/Engine/Source/Editor/BlueprintGraph/Private/K2Node_StructMemberGet.cpp
@@ -13,12 +13,23 @@ UK2Node_StructMemberGet::UK2Node_StructMemberGet(const FObjectInitializer& Objec
 {
 }
 
-void UK2Node_StructMemberGet::PostEditChangeProperty(struct FPropertyChangedEvent& PropertyChangedEvent)
+void UK2Node_StructMemberGet::PreEditChange(UProperty* PropertyThatWillChange)
 {
-	FName PropertyName = (PropertyChangedEvent.Property != NULL) ? PropertyChangedEvent.Property->GetFName() : NAME_None;
+	Super::PreEditChange(PropertyThatWillChange);
 
-	if ((PropertyName == TEXT("bShowPin")))
+	if (PropertyThatWillChange && PropertyThatWillChange->GetFName() == GET_MEMBER_NAME_CHECKED(FOptionalPinFromProperty, bShowPin))
 	{
+		FOptionalPinManager::CacheShownPins(ShowPinForProperties, OldShownPins);
+	}
+}
+
+void UK2Node_StructMemberGet::PostEditChangeProperty(FPropertyChangedEvent& PropertyChangedEvent)
+{
+	const FName PropertyName = (PropertyChangedEvent.Property ? PropertyChangedEvent.Property->GetFName() : NAME_None);
+
+	if (PropertyName == GET_MEMBER_NAME_CHECKED(FOptionalPinFromProperty, bShowPin))
+	{
+		FOptionalPinManager::EvaluateOldShownPins(ShowPinForProperties, OldShownPins, this);
 		GetSchema()->ReconstructNode(*this);
 	}
 
diff --git a/Engine/Source/Editor/BlueprintGraph/Private/K2Node_StructMemberSet.cpp b/Engine/Source/Editor/BlueprintGraph/Private/K2Node_StructMemberSet.cpp
--- a/Engine/Source/Editor/BlueprintGraph/Private/K2Node_StructMemberSet.cpp
+++ b/Engine/Source/Editor/BlueprintGraph/Private/K2Node_StructMemberSet.cpp
@@ -14,12 +14,23 @@ UK2Node_StructMemberSet::UK2Node_StructMemberSet(const FObjectInitializer& Objec
 {
 }
 
-void UK2Node_StructMemberSet::PostEditChangeProperty(struct FPropertyChangedEvent& PropertyChangedEvent)
+void UK2Node_StructMemberSet::PreEditChange(UProperty* PropertyThatWillChange)
 {
-	FName PropertyName = (PropertyChangedEvent.Property != NULL) ? PropertyChangedEvent.Property->GetFName() : NAME_None;
+	Super::PreEditChange(PropertyThatWillChange);
 
-	if ((PropertyName == TEXT("bShowPin")))
+	if (PropertyThatWillChange && PropertyThatWillChange->GetFName() == GET_MEMBER_NAME_CHECKED(FOptionalPinFromProperty, bShowPin))
 	{
+		FOptionalPinManager::CacheShownPins(ShowPinForProperties, OldShownPins);
+	}
+}
+
+void UK2Node_StructMemberSet::PostEditChangeProperty(FPropertyChangedEvent& PropertyChangedEvent)
+{
+	const FName PropertyName = (PropertyChangedEvent.Property ? PropertyChangedEvent.Property->GetFName() : NAME_None);
+
+	if (PropertyName == GET_MEMBER_NAME_CHECKED(FOptionalPinFromProperty, bShowPin))
+	{
+		FOptionalPinManager::EvaluateOldShownPins(ShowPinForProperties, OldShownPins, this);
 		GetSchema()->ReconstructNode(*this);
 	}
 
diff --git a/Engine/Source/Editor/BlueprintGraph/Private/K2Node_StructOperation.cpp b/Engine/Source/Editor/BlueprintGraph/Private/K2Node_StructOperation.cpp
--- a/Engine/Source/Editor/BlueprintGraph/Private/K2Node_StructOperation.cpp
+++ b/Engine/Source/Editor/BlueprintGraph/Private/K2Node_StructOperation.cpp
@@ -103,15 +103,22 @@ FString UK2Node_StructOperation::GetFindReferenceSearchString() const
 bool UK2Node_StructOperation::IsActionFilteredOut(const FBlueprintActionFilter& Filter)
 {
 	bool bIsFiltered = false;
-	if (StructType && !StructType->GetBoolMetaData(FBlueprintMetadata::MD_AllowableBlueprintVariableType))
+	if (StructType)
 	{
-		bIsFiltered = true;
-		for (UEdGraphPin* ContextPin : Filter.Context.Pins)
+		if (StructType->GetBoolMetaData(FBlueprintMetadata::MD_BlueprintInternalUseOnly))
 		{
-			if (ContextPin->PinType.PinCategory == UEdGraphSchema_K2::PC_Struct && ContextPin->PinType.PinSubCategoryObject == StructType)
+			bIsFiltered = true;
+		}
+		else if (!StructType->GetBoolMetaData(FBlueprintMetadata::MD_AllowableBlueprintVariableType))
+		{
+			bIsFiltered = true;
+			for (UEdGraphPin* ContextPin : Filter.Context.Pins)
 			{
-				bIsFiltered = false;
-				break;
+				if (ContextPin->PinType.PinCategory == UEdGraphSchema_K2::PC_Struct && ContextPin->PinType.PinSubCategoryObject == StructType)
+				{
+					bIsFiltered = false;
+					break;
+				}
 			}
 		}
 	}
diff --git a/Engine/Source/Editor/BlueprintGraph/Private/K2Node_SwitchEnum.cpp b/Engine/Source/Editor/BlueprintGraph/Private/K2Node_SwitchEnum.cpp
--- a/Engine/Source/Editor/BlueprintGraph/Private/K2Node_SwitchEnum.cpp
+++ b/Engine/Source/Editor/BlueprintGraph/Private/K2Node_SwitchEnum.cpp
@@ -13,6 +13,7 @@ UK2Node_SwitchEnum::UK2Node_SwitchEnum(const FObjectInitializer& ObjectInitializ
 {
 
 	bHasDefaultPin = false;
+	OrphanedPinSaveMode = ESaveOrphanPinMode::SaveAll;
 
 	FunctionName = TEXT("NotEqual_ByteByte");
 	FunctionClass = UKismetMathLibrary::StaticClass();
diff --git a/Engine/Source/Editor/BlueprintGraph/Private/K2Node_VariableGet.cpp b/Engine/Source/Editor/BlueprintGraph/Private/K2Node_VariableGet.cpp
--- a/Engine/Source/Editor/BlueprintGraph/Private/K2Node_VariableGet.cpp
+++ b/Engine/Source/Editor/BlueprintGraph/Private/K2Node_VariableGet.cpp
@@ -407,6 +407,48 @@ void UK2Node_VariableGet::SetPurity(bool bNewPurity)
 	}
 }
 
+void UK2Node_VariableGet::ValidateNodeDuringCompilation(FCompilerResultsLog& MessageLog) const
+{
+	Super::ValidateNodeDuringCompilation(MessageLog);
+
+	// Some expansions, such as timelines, will create gets for non-blueprint visible properties, and we don't want to validate against that
+	if (!IsIntermediateNode())
+	{
+		if (UProperty* Property = GetPropertyForVariable())
+		{
+			const FBlueprintEditorUtils::EPropertyReadableState PropertyReadableState = FBlueprintEditorUtils::IsPropertyReadableInBlueprint(GetBlueprint(), Property);
+
+			if (PropertyReadableState != FBlueprintEditorUtils::EPropertyReadableState::Readable)
+			{
+				FFormatNamedArguments Args;
+				if (UObject* Class = Property->GetOuter())
+				{
+					Args.Add(TEXT("VariableName"), FText::AsCultureInvariant(FString::Printf(TEXT("%s.%s"), *Class->GetName(), *Property->GetName())));
+				}
+				else
+				{
+					Args.Add(TEXT("VariableName"), FText::AsCultureInvariant(Property->GetName()));
+				}
+
+				if (PropertyReadableState == FBlueprintEditorUtils::EPropertyReadableState::NotBlueprintVisible)
+				{
+					// DEPRECATED(4.17) ... make this an error
+					MessageLog.Warning(*FText::Format(LOCTEXT("UnableToGet_NotVisible", "{VariableName} is not blueprint visible. Please fix mark up or cease accessing as this will be made an error in a future release. @@"), Args).ToString(), this);
+				}
+				else if (PropertyReadableState == FBlueprintEditorUtils::EPropertyReadableState::Private)
+				{
+					// DEPRECATED(4.17) ... make this an error
+					MessageLog.Warning(*FText::Format(LOCTEXT("UnableToGet_ReadOnly", "{VariableName} is private and not accessible in this context. Please fix mark up or cease accessing as this will be an error in a future release. @@"), Args).ToString(), this);
+				}
+				else
+				{
+					check(false);
+				}
+			}
+		}
+	}
+}
+
 void UK2Node_VariableGet::ExpandNode(class FKismetCompilerContext& CompilerContext, UEdGraph* SourceGraph)
 {
 	Super::ExpandNode(CompilerContext, SourceGraph);
diff --git a/Engine/Source/Editor/BlueprintGraph/Private/K2Node_VariableSet.cpp b/Engine/Source/Editor/BlueprintGraph/Private/K2Node_VariableSet.cpp
--- a/Engine/Source/Editor/BlueprintGraph/Private/K2Node_VariableSet.cpp
+++ b/Engine/Source/Editor/BlueprintGraph/Private/K2Node_VariableSet.cpp
@@ -347,6 +347,46 @@ FText UK2Node_VariableSet::GetPinNameOverride(const UEdGraphPin& Pin) const
 	return !Pin.PinFriendlyName.IsEmpty() ? Pin.PinFriendlyName : FText::FromString(Pin.PinName);
 }
 
+void UK2Node_VariableSet::ValidateNodeDuringCompilation(FCompilerResultsLog& MessageLog) const
+{
+	Super::ValidateNodeDuringCompilation(MessageLog);
+
+	// Some expansions will create sets for non-blueprint visible properties, and we don't want to validate against that
+	if (!IsIntermediateNode())
+	{
+		if (UProperty* Property = GetPropertyForVariable())
+		{
+			const FBlueprintEditorUtils::EPropertyWritableState PropertyWritableState = FBlueprintEditorUtils::IsPropertyWritableInBlueprint(GetBlueprint(), Property);
+
+			if (PropertyWritableState != FBlueprintEditorUtils::EPropertyWritableState::Writable)
+			{
+				FFormatNamedArguments Args;
+				if (UObject* Class = Property->GetOuter())
+				{
+					Args.Add(TEXT("VariableName"), FText::AsCultureInvariant(FString::Printf(TEXT("%s.%s"), *Class->GetName(), *Property->GetName())));
+				}
+				else
+				{
+					Args.Add(TEXT("VariableName"), FText::AsCultureInvariant(Property->GetName()));
+				}
+
+				if (PropertyWritableState == FBlueprintEditorUtils::EPropertyWritableState::BlueprintReadOnly || PropertyWritableState == FBlueprintEditorUtils::EPropertyWritableState::NotBlueprintVisible)
+				{
+					MessageLog.Error(*FText::Format(LOCTEXT("UnableToSet_NotWritable", "{VariableName} is not blueprint writable. @@"), Args).ToString(), this);
+				}
+				else if (PropertyWritableState == FBlueprintEditorUtils::EPropertyWritableState::Private)
+				{
+					MessageLog.Error(*LOCTEXT("UnableToSet_ReadOnly", "{VariableName} is private and not accessible in this context. @@").ToString(), this);
+				}
+				else
+				{
+					check(false);
+				}
+			}
+		}
+	}
+}
+
 void UK2Node_VariableSet::ExpandNode(class FKismetCompilerContext& CompilerContext, UEdGraph* SourceGraph)
 {
 	Super::ExpandNode(CompilerContext, SourceGraph);
@@ -404,7 +444,10 @@ void UK2Node_VariableSet::ExpandNode(class FKismetCompilerContext& CompilerConte
 				CompilerContext.MovePinLinksToIntermediate(*FindPinChecked(UEdGraphSchema_K2::PN_Then, EGPD_Output), *CallFuncNode->GetThenPin());
 				
 				// Move Self pin connections
-				CompilerContext.MovePinLinksToIntermediate(*K2Schema->FindSelfPin(*this, EGPD_Input), *K2Schema->FindSelfPin(*CallFuncNode, EGPD_Input));
+				if (UEdGraphPin* SetSelfPin = K2Schema->FindSelfPin(*this, EGPD_Input))
+				{
+					CompilerContext.MovePinLinksToIntermediate(*SetSelfPin, *K2Schema->FindSelfPin(*CallFuncNode, EGPD_Input));
+				}
 
 				// Move Value pin connections
 				UEdGraphPin* SetFunctionValuePin = nullptr;
diff --git a/Engine/Source/Editor/BlueprintGraph/Private/MakeStructHandler.cpp b/Engine/Source/Editor/BlueprintGraph/Private/MakeStructHandler.cpp
--- a/Engine/Source/Editor/BlueprintGraph/Private/MakeStructHandler.cpp
+++ b/Engine/Source/Editor/BlueprintGraph/Private/MakeStructHandler.cpp
@@ -15,18 +15,17 @@
 
 UEdGraphPin* FKCHandler_MakeStruct::FindStructPinChecked(UEdGraphNode* Node) const
 {
-	check(NULL != Node);
-	UEdGraphPin* OutputPin = NULL;
-	for (int32 PinIndex = 0; PinIndex < Node->Pins.Num(); ++PinIndex)
+	check(Node);
+	UEdGraphPin* OutputPin = nullptr;
+	for (UEdGraphPin* Pin : Node->Pins)
 	{
-		UEdGraphPin* Pin = Node->Pins[PinIndex];
 		if (Pin && (EGPD_Output == Pin->Direction) && !CompilerContext.GetSchema()->IsMetaPin(*Pin))
 		{
 			OutputPin = Pin;
 			break;
 		}
 	}
-	check(NULL != OutputPin);
+	check(OutputPin);
 	return OutputPin;
 }
 
@@ -39,24 +38,25 @@ FKCHandler_MakeStruct::FKCHandler_MakeStruct(FKismetCompilerContext& InCompilerC
 void FKCHandler_MakeStruct::RegisterNets(FKismetFunctionContext& Context, UEdGraphNode* InNode)
 {
 	UK2Node_MakeStruct* Node = CastChecked<UK2Node_MakeStruct>(InNode);
-	if (NULL == Node->StructType)
+	if (nullptr == Node->StructType)
 	{
 		CompilerContext.MessageLog.Error(*LOCTEXT("MakeStruct_UnknownStructure_Error", "Unknown structure to break for @@").ToString(), Node);
 		return;
 	}
 
-	if (!UK2Node_MakeStruct::CanBeMade(Node->StructType))
+	if (!UK2Node_MakeStruct::CanBeMade(Node->StructType, Node->IsIntermediateNode()))
 	{
-		CompilerContext.MessageLog.Warning(*LOCTEXT("MakeStruct_Error", "The structure contains read-only members @@. Try use specialized 'make' function if available. ").ToString(), Node);
+		CompilerContext.MessageLog.Error(*LOCTEXT("MakeStruct_Error", "The structure @@ is not a BlueprintType. ").ToString(), Node);
+		return;
 	}
 
 	FNodeHandlingFunctor::RegisterNets(Context, Node);
 
 	UEdGraphPin* OutputPin = FindStructPinChecked(Node);
 	UEdGraphPin* Net = FEdGraphUtilities::GetNetFromPin(OutputPin);
-	check(NULL != Net);
+	check(Net);
 	FBPTerminal** FoundTerm = Context.NetMap.Find(Net);
-	FBPTerminal* Term = FoundTerm ? *FoundTerm : NULL;
+	FBPTerminal* Term = FoundTerm ? *FoundTerm : nullptr;
 
 	if (Term == nullptr)
 	{
@@ -65,7 +65,7 @@ void FKCHandler_MakeStruct::RegisterNets(FKismetFunctionContext& Context, UEdGra
 	else
 	{
 		UStruct* StructInTerm = Cast<UStruct>(Term->Type.PinSubCategoryObject.Get());
-		if (NULL == StructInTerm || !StructInTerm->IsChildOf(Node->StructType))
+		if (nullptr == StructInTerm || !StructInTerm->IsChildOf(Node->StructType))
 		{
 			CompilerContext.MessageLog.Error(*LOCTEXT("MakeStruct_NoMatch_Error", "Structures don't match for @@").ToString(), Node);
 		}
@@ -105,23 +105,22 @@ void FKCHandler_MakeStruct::Compile(FKismetFunctionContext& Context, UEdGraphNod
 	UEdGraphPin* OutputStructNet = FEdGraphUtilities::GetNetFromPin(StructPin);
 	FBPTerminal** FoundTerm = Context.NetMap.Find(OutputStructNet);
 	FBPTerminal* OutputStructTerm = FoundTerm ? *FoundTerm : NULL;
-	check(NULL != OutputStructTerm);
+	check(OutputStructTerm);
 
-	for (int32 PinIndex = 0; PinIndex < Node->Pins.Num(); ++PinIndex)
+	for (UEdGraphPin* Pin : Node->Pins)
 	{
-		UEdGraphPin* Pin = Node->Pins[PinIndex];
 		if (Pin && (Pin != StructPin) && !CompilerContext.GetSchema()->IsMetaPin(*Pin) && (Pin->Direction == EGPD_Input))
 		{
 			UProperty* BoundProperty = FindField<UProperty>(Node->StructType, *(Pin->PinName));
-			check(NULL != BoundProperty);
+			check(BoundProperty);
 
-			// If the pin is not connectable, do not forward the net
+			// If the pin is not connectible, do not forward the net
 			if (!Pin->bNotConnectable)
 			{
 				if (FBPTerminal** FoundSrcTerm = Context.NetMap.Find(FEdGraphUtilities::GetNetFromPin(Pin)))
 				{
-					FBPTerminal* SrcTerm = FoundSrcTerm ? *FoundSrcTerm : NULL;
-					check(NULL != SrcTerm);
+					FBPTerminal* SrcTerm = FoundSrcTerm ? *FoundSrcTerm : nullptr;
+					check(SrcTerm);
 
 					FBPTerminal* DstTerm = Context.CreateLocalTerminal();
 					DstTerm->CopyFromPin(Pin, Context.NetNameMap->MakeValidName(Pin));
diff --git a/Engine/Source/Editor/BlueprintGraph/Private/VariableSetHandler.cpp b/Engine/Source/Editor/BlueprintGraph/Private/VariableSetHandler.cpp
--- a/Engine/Source/Editor/BlueprintGraph/Private/VariableSetHandler.cpp
+++ b/Engine/Source/Editor/BlueprintGraph/Private/VariableSetHandler.cpp
@@ -27,13 +27,8 @@ void FKCHandler_VariableSet::RegisterNets(FKismetFunctionContext& Context, UEdGr
 	{
 		SetterNode->CheckForErrors(CompilerContext.GetSchema(), Context.MessageLog);
 
-		if(FBlueprintEditorUtils::IsPropertyReadOnlyInCurrentBlueprint(Context.Blueprint, SetterNode->GetPropertyForVariable()))
-		{
-			CompilerContext.MessageLog.Warning(*LOCTEXT("BlueprintReadOnlyOrPrivate_Error", "The property is marked as BlueprintReadOnly or Private. It cannot be modifed in the blueprint. @@").ToString(), Node);
-		}
-
 		// Report an error that the local variable could not be found
-		if(SetterNode->VariableReference.IsLocalScope() && SetterNode->GetPropertyForVariable() == NULL)
+		if(SetterNode->VariableReference.IsLocalScope() && SetterNode->GetPropertyForVariable() == nullptr)
 		{
 			FFormatNamedArguments Args;
 			Args.Add(TEXT("VariableName"), FText::FromName(SetterNode->VariableReference.GetMemberName()));
@@ -50,10 +45,9 @@ void FKCHandler_VariableSet::RegisterNets(FKismetFunctionContext& Context, UEdGr
 		}
 	}
 
-	for (int32 PinIndex = 0; PinIndex < Node->Pins.Num(); ++PinIndex)
+	for (UEdGraphPin* Net : Node->Pins)
 	{
-		UEdGraphPin* Net = Node->Pins[PinIndex];
-		if (!CompilerContext.GetSchema()->IsMetaPin(*Net) && (Net->Direction == EGPD_Input))
+		if (!Net->bOrphanedPin && (Net->Direction == EGPD_Input) && !CompilerContext.GetSchema()->IsMetaPin(*Net))
 		{
 			if (ValidateAndRegisterNetIfLiteral(Context, Net))
 			{
@@ -66,18 +60,18 @@ void FKCHandler_VariableSet::RegisterNets(FKismetFunctionContext& Context, UEdGr
 void FKCHandler_VariableSet::InnerAssignment(FKismetFunctionContext& Context, UEdGraphNode* Node, UEdGraphPin* VariablePin, UEdGraphPin* ValuePin)
 {
 	FBPTerminal** VariableTerm = Context.NetMap.Find(VariablePin);
-	if (VariableTerm == NULL)
+	if (VariableTerm == nullptr)
 	{
 		VariableTerm = Context.NetMap.Find(FEdGraphUtilities::GetNetFromPin(VariablePin));
 	}
 
 	FBPTerminal** ValueTerm = Context.LiteralHackMap.Find(ValuePin);
-	if (ValueTerm == NULL)
+	if (ValueTerm == nullptr)
 	{
 		ValueTerm = Context.NetMap.Find(FEdGraphUtilities::GetNetFromPin(ValuePin));
 	}
 
-	if ((VariableTerm != NULL) && (ValueTerm != NULL))
+	if (VariableTerm && ValueTerm)
 	{
 		FKismetCompilerUtilities::CreateObjectAssignmentStatement(Context, Node, *ValueTerm, *VariableTerm);
 
diff --git a/Engine/Source/Editor/BlueprintGraph/Public/BlueprintActionDatabase.h b/Engine/Source/Editor/BlueprintGraph/Public/BlueprintActionDatabase.h
--- a/Engine/Source/Editor/BlueprintGraph/Public/BlueprintActionDatabase.h
+++ b/Engine/Source/Editor/BlueprintGraph/Public/BlueprintActionDatabase.h
@@ -5,7 +5,7 @@
 #include "CoreMinimal.h"
 #include "Stats/Stats.h"
 #include "UObject/Object.h"
-#include "UObject/WeakObjectPtr.h"
+#include "UObject/ObjectKey.h"
 #include "UObject/GCObject.h"
 #include "TickableEditorObject.h"
 
@@ -27,10 +27,10 @@ class BLUEPRINTGRAPH_API FBlueprintActionDatabase : public FGCObject, public FTi
 public:
 	DECLARE_MULTICAST_DELEGATE_OneParam(FOnDatabaseEntryUpdated, UObject*);
 
-	typedef TMap<TWeakObjectPtr<UObject>, int32>			FPrimingQueue;
-	typedef TArray<UBlueprintNodeSpawner*>					FActionList;
-	typedef TMap<UObject const*, FActionList>				FActionRegistry;
-	typedef TMap<FName, TArray<UBlueprintNodeSpawner*>>		FUnloadedActionRegistry;
+	typedef TMap<FObjectKey, int32>						FPrimingQueue;
+	typedef TArray<UBlueprintNodeSpawner*>				FActionList;
+	typedef TMap<FObjectKey, FActionList>				FActionRegistry;
+	typedef TMap<FName, TArray<UBlueprintNodeSpawner*>>	FUnloadedActionRegistry;
 
 public:
 	/**
@@ -76,6 +76,11 @@ public:
 	 */
 	void RefreshWorlds();
 
+	/**
+	 * Removes the entry with the given key on the next tick.
+	 */
+	void DeferredRemoveEntry(FObjectKey const& InKey);
+
 	/**
 	 * Finds the database entry for the specified class and wipes it, 
 	 * repopulating it with a fresh set of associated node-spawners. 
@@ -166,6 +171,9 @@ private:
 	 */
 	FPrimingQueue ActionPrimingQueue;
 
+	/** List of action keys to be removed on the next tick. */
+	TArray<FObjectKey> ActionRemoveQueue;
+
 	/** */
 	FOnDatabaseEntryUpdated EntryRefreshDelegate;
 	FOnDatabaseEntryUpdated EntryRemovedDelegate;
diff --git a/Engine/Source/Editor/BlueprintGraph/Public/BlueprintActionFilter.h b/Engine/Source/Editor/BlueprintGraph/Public/BlueprintActionFilter.h
--- a/Engine/Source/Editor/BlueprintGraph/Public/BlueprintActionFilter.h
+++ b/Engine/Source/Editor/BlueprintGraph/Public/BlueprintActionFilter.h
@@ -134,16 +134,16 @@ struct BLUEPRINTGRAPH_API FBlueprintActionInfo
 
 private:
 	/** The class or asset-object that the NodeSpawner action is keyed to (in the action database)*/
-	const TWeakObjectPtr<UObject> ActionOwner;
+	TWeakObjectPtr<const UObject> ActionOwner;
 
 	/** Keeps track of the fields we've cached (needed in case one turns out to be null) */
 	uint32 CacheFlags;
 	
 	/** */
-	UClass const*    CachedOwnerClass;
-	UField const*    CachedActionField;
-	UProperty const* CachedActionProperty;
-	UFunction const* CachedActionFunction;
+	TWeakObjectPtr<const UClass>	CachedOwnerClass;
+	UField const*					CachedActionField;
+	UProperty const*				CachedActionProperty;
+	UFunction const*				CachedActionFunction;
 
 	/** */
 	IBlueprintNodeBinder::FBindingSet Bindings;
diff --git a/Engine/Source/Editor/DataTableEditor/Private/SRowEditor.cpp b/Engine/Source/Editor/DataTableEditor/Private/SRowEditor.cpp
--- a/Engine/Source/Editor/DataTableEditor/Private/SRowEditor.cpp
+++ b/Engine/Source/Editor/DataTableEditor/Private/SRowEditor.cpp
@@ -7,6 +7,7 @@
 #include "Widgets/Input/SEditableTextBox.h"
 #include "Widgets/Input/SButton.h"
 #include "EditorStyleSet.h"
+#include "UObject/StructOnScope.h"
 
 #include "PropertyEditorModule.h"
 #include "IStructureDetailsView.h"
diff --git a/Engine/Source/Editor/DataTableEditor/Private/SRowEditor.h b/Engine/Source/Editor/DataTableEditor/Private/SRowEditor.h
--- a/Engine/Source/Editor/DataTableEditor/Private/SRowEditor.h
+++ b/Engine/Source/Editor/DataTableEditor/Private/SRowEditor.h
@@ -14,6 +14,7 @@
 
 class IStructureDetailsView;
 class SEditableTextBox;
+class FStructOnScope;
 
 DECLARE_DELEGATE_OneParam(FOnRowModified, FName /*Row name*/);
 DECLARE_DELEGATE_OneParam(FOnRowSelected, FName /*Row name*/);
diff --git a/Engine/Source/Editor/DetailCustomizations/Private/SkeletalControlNodeDetails.cpp b/Engine/Source/Editor/DetailCustomizations/Private/SkeletalControlNodeDetails.cpp
--- a/Engine/Source/Editor/DetailCustomizations/Private/SkeletalControlNodeDetails.cpp
+++ b/Engine/Source/Editor/DetailCustomizations/Private/SkeletalControlNodeDetails.cpp
@@ -12,6 +12,7 @@
 #include "IDetailChildrenBuilder.h"
 #include "DetailCategoryBuilder.h"
 #include "PropertyCustomizationHelpers.h"
+#include "K2Node.h"
 
 #define LOCTEXT_NAMESPACE "SkeletalControlNodeDetails"
 
@@ -41,11 +42,11 @@ void FSkeletalControlNodeDetails::CustomizeDetails(class IDetailLayoutBuilder& D
 		}
 		for (int32 Index = 0; Index < NumElements; ++Index)
 		{
-			auto StructPropHandle = ArrayProperty->GetElement(Index);
-			auto CategoryPropeHandle = StructPropHandle->GetChildHandle("CategoryName");
-			check(CategoryPropeHandle.IsValid());
+			TSharedRef<IPropertyHandle> StructPropHandle = ArrayProperty->GetElement(Index);
+			TSharedPtr<IPropertyHandle> CategoryPropHandle = StructPropHandle->GetChildHandle(GET_MEMBER_NAME_CHECKED(FOptionalPinFromProperty, CategoryName));
+			check(CategoryPropHandle.IsValid());
 			FName CategoryNameValue;
-			const auto Result = CategoryPropeHandle->GetValue(CategoryNameValue);
+			const FPropertyAccess::Result Result = CategoryPropHandle->GetValue(CategoryNameValue);
 			if (ensure(FPropertyAccess::Success == Result))
 			{
 				UniqueCategoryNames.Add(CategoryNameValue);
@@ -58,8 +59,11 @@ void FSkeletalControlNodeDetails::CustomizeDetails(class IDetailLayoutBuilder& D
 	const bool bDisplayResetToDefault = false;
 	const bool bDisplayElementNum = false;
 	const bool bForAdvanced = false;
-	for (auto CategoryName : UniqueCategoryNames)
+	for (const FName& CategoryName : UniqueCategoryNames)
 	{
+		//@TODO: Pay attention to category filtering here
+		
+		
 		TSharedRef<FDetailArrayBuilder> AvailablePinsBuilder = MakeShareable(new FDetailArrayBuilder(AvailablePins, bGenerateHeader, bDisplayResetToDefault, bDisplayElementNum));
 		AvailablePinsBuilder->OnGenerateArrayElementWidget(FOnGenerateArrayElementWidget::CreateSP(this, &FSkeletalControlNodeDetails::OnGenerateElementForPropertyPin, CategoryName));
 		AvailablePinsBuilder->SetDisplayName((CategoryName == NAME_None) ? LOCTEXT("DefaultCategory", "Default Category") : FText::FromName(CategoryName));
@@ -70,13 +74,13 @@ void FSkeletalControlNodeDetails::CustomizeDetails(class IDetailLayoutBuilder& D
 ECheckBoxState FSkeletalControlNodeDetails::GetShowPinValueForProperty(TSharedRef<IPropertyHandle> InElementProperty) const
 {
 	FString Value;
-	InElementProperty->GetChildHandle("bShowPin")->GetValueAsFormattedString(Value);
+	InElementProperty->GetChildHandle(GET_MEMBER_NAME_CHECKED(FOptionalPinFromProperty, bShowPin))->GetValueAsFormattedString(Value);
 
 	if (Value == TEXT("true"))
 	{
 		return ECheckBoxState::Checked;
 	}
-	else if(Value == TEXT("false"))
+	else if (Value == TEXT("false"))
 	{
 		return ECheckBoxState::Unchecked;
 	}
@@ -85,92 +89,119 @@ ECheckBoxState FSkeletalControlNodeDetails::GetShowPinValueForProperty(TSharedRe
 
 void FSkeletalControlNodeDetails::OnShowPinChanged(ECheckBoxState InNewState, TSharedRef<IPropertyHandle> InElementProperty)
 {
-	InElementProperty->GetChildHandle("bShowPin")->SetValueFromFormattedString(InNewState == ECheckBoxState::Checked? TEXT("true") : TEXT("false"));
+	InElementProperty->GetChildHandle(GET_MEMBER_NAME_CHECKED(FOptionalPinFromProperty, bShowPin))->SetValueFromFormattedString(InNewState == ECheckBoxState::Checked ? TEXT("true") : TEXT("false"));
 }
 
 void FSkeletalControlNodeDetails::OnGenerateElementForPropertyPin(TSharedRef<IPropertyHandle> ElementProperty, int32 ElementIndex, IDetailChildrenBuilder& ChildrenBuilder, FName CategoryName)
 {
 	{
-		auto CategoryPropeHandle = ElementProperty->GetChildHandle("CategoryName");
-		check(CategoryPropeHandle.IsValid());
+		TSharedPtr<IPropertyHandle> CategoryPropHandle = ElementProperty->GetChildHandle(GET_MEMBER_NAME_CHECKED(FOptionalPinFromProperty, CategoryName));
+		check(CategoryPropHandle.IsValid());
 		FName CategoryNameValue;
-		const auto Result = CategoryPropeHandle->GetValue(CategoryNameValue);
+		const FPropertyAccess::Result Result = CategoryPropHandle->GetValue(CategoryNameValue);
 		const bool bProperCategory = ensure(FPropertyAccess::Success == Result) && (CategoryNameValue == CategoryName);
+
 		if (!bProperCategory)
 		{
 			return;
 		}
 	}
 
-	TSharedPtr<IPropertyHandle> PropertyNameHandle = ElementProperty->GetChildHandle("PropertyFriendlyName");
-	FText PropertyFriendlyName(LOCTEXT("Invalid", "Invalid"));
+	FString FilterString = CategoryName.ToString();
 
-	if (PropertyNameHandle.IsValid())
+	FText PropertyFriendlyName(LOCTEXT("Invalid", "Invalid"));
 	{
-		FString DisplayFriendlyName;
-		switch (PropertyNameHandle->GetValue(/*out*/ DisplayFriendlyName))
-		{
-		case FPropertyAccess::Success:
+		TSharedPtr<IPropertyHandle> PropertyFriendlyNameHandle = ElementProperty->GetChildHandle(GET_MEMBER_NAME_CHECKED(FOptionalPinFromProperty, PropertyFriendlyName));
+		if (PropertyFriendlyNameHandle.IsValid())
 		{
-			PropertyFriendlyName = FText::FromString(DisplayFriendlyName);
+			FString DisplayFriendlyName;
+			switch (PropertyFriendlyNameHandle->GetValue(/*out*/ DisplayFriendlyName))
+			{
+			case FPropertyAccess::Success:
+				FilterString += TEXT(" ") + DisplayFriendlyName;
+				PropertyFriendlyName = FText::FromString(DisplayFriendlyName);
+				break;
+			case FPropertyAccess::MultipleValues:
+				ChildrenBuilder.AddCustomRow(FText::GetEmpty())
+					[
+						SNew(STextBlock).Text(LOCTEXT("OnlyWorksInSingleSelectMode", "Multiple types selected"))
+					];
+				return;
+			case FPropertyAccess::Fail:
+			default:
+				check(false);
+				break;
+			}
+		}
+	}
 
-			//DetailBuilder.EditCategory
-			//DisplayNameAsName
+	{
+		TSharedPtr<IPropertyHandle> PropertyNameHandle = ElementProperty->GetChildHandle(GET_MEMBER_NAME_CHECKED(FOptionalPinFromProperty, PropertyName));
+		if (PropertyNameHandle.IsValid())
+		{
+			FString RawName;
+			if (PropertyNameHandle->GetValue(/*out*/ RawName) == FPropertyAccess::Success)
+			{
+				FilterString += TEXT(" ") + RawName;
+			}
 		}
-			break;
-		case FPropertyAccess::MultipleValues:
-			ChildrenBuilder.AddCustomRow(FText::GetEmpty())
-				[
-					SNew(STextBlock).Text(LOCTEXT("OnlyWorksInSingleSelectMode", "Multiple types selected"))
-				];
-			return;
-		case FPropertyAccess::Fail:
-		default:
-			check(false);
-			break;
+	}
+
+	FText PinTooltip;
+	{
+		TSharedPtr<IPropertyHandle> PropertyTooltipHandle = ElementProperty->GetChildHandle(GET_MEMBER_NAME_CHECKED(FOptionalPinFromProperty, PropertyTooltip));
+		if (PropertyTooltipHandle.IsValid())
+		{
+			FString PinTooltipString;
+			if (PropertyTooltipHandle->GetValue(/*out*/ PinTooltip) == FPropertyAccess::Success)
+			{
+				FilterString += TEXT(" ") + PinTooltip.ToString();
+			}
 		}
 	}
-	
-	TSharedPtr<IPropertyHandle> HasOverrideValueHandle = ElementProperty->GetChildHandle("bHasOverridePin");
+
+	TSharedPtr<IPropertyHandle> HasOverrideValueHandle = ElementProperty->GetChildHandle(GET_MEMBER_NAME_CHECKED(FOptionalPinFromProperty, bHasOverridePin));
 	bool bHasOverrideValue;
 	HasOverrideValueHandle->GetValue(/*out*/ bHasOverrideValue);
-	FText Tooltip;
+	FText OverrideCheckBoxTooltip;
 
 	// Setup a tooltip based on whether the property has an override value or not.
 	if (bHasOverrideValue)
 	{
-		Tooltip = LOCTEXT("HasOverridePin", "Enabling this pin will make it visible for setting on the node and automatically enable the value for override when using the struct. Any updates to the resulting struct will require the value be set again or the override will be automatically disabled.");
+		OverrideCheckBoxTooltip = LOCTEXT("HasOverridePin", "Enabling this pin will make it visible for setting on the node and automatically enable the value for override when using the struct. Any updates to the resulting struct will require the value be set again or the override will be automatically disabled.");
 	}
 	else
 	{
-		Tooltip = LOCTEXT("HasNoOverridePin", "Enabling this pin will make it visible for setting on the node.");
+		OverrideCheckBoxTooltip = LOCTEXT("HasNoOverridePin", "Enabling this pin will make it visible for setting on the node.");
 	}
 
-	TSharedRef<SWidget> PropertyNameWidget = ElementProperty->CreatePropertyNameWidget(PropertyFriendlyName);
-	PropertyNameWidget->SetToolTipText(Tooltip);
-
-	FString Value;
-	ElementProperty->GetChildHandle("bShowPin")->GetValueAsFormattedString(Value);
-
 	ChildrenBuilder.AddCustomRow( PropertyFriendlyName )
+	.FilterString(FText::AsCultureInvariant(FilterString))
+	.NameContent()
+	[
+		ElementProperty->CreatePropertyNameWidget(PropertyFriendlyName, PinTooltip)
+	]
+	.ValueContent()
 	[
-		SNew( SHorizontalBox )
-		+ SHorizontalBox::Slot()
+		SNew(SHorizontalBox)
+		.ToolTipText(OverrideCheckBoxTooltip)
+		+SHorizontalBox::Slot()
+		.AutoWidth()
+		.HAlign(HAlign_Left)
 		.VAlign(VAlign_Center)
-		.HAlign(HAlign_Right)
-		.Padding( 3.0f, 0.0f )
-		.FillWidth(1.0f)
 		[
-			PropertyNameWidget
+			SNew(SCheckBox)
+			.IsChecked(this, &FSkeletalControlNodeDetails::GetShowPinValueForProperty, ElementProperty)
+			.OnCheckStateChanged(this, &FSkeletalControlNodeDetails::OnShowPinChanged, ElementProperty)
 		]
-		+ SHorizontalBox::Slot()
+		+SHorizontalBox::Slot()
+		.AutoWidth()
+		.HAlign(HAlign_Left)
 		.VAlign(VAlign_Center)
-		.FillWidth(1.0f)
 		[
-			SNew(SCheckBox)
-				.IsChecked(this, &FSkeletalControlNodeDetails::GetShowPinValueForProperty, ElementProperty)
-				.OnCheckStateChanged(this, &FSkeletalControlNodeDetails::OnShowPinChanged, ElementProperty)
-				.ToolTipText(Tooltip)
+			SNew(STextBlock)
+			.Text(LOCTEXT("AsPin", " (As pin)"))
+			.Font(ChildrenBuilder.GetParentCategory().GetParentLayout().GetDetailFont())
 		]
 	];
 }
@@ -179,4 +210,3 @@ void FSkeletalControlNodeDetails::OnGenerateElementForPropertyPin(TSharedRef<IPr
 /////////////////////////////////////////////////////
 
 #undef LOCTEXT_NAMESPACE
-
diff --git a/Engine/Source/Editor/GraphEditor/Private/BlueprintConnectionDrawingPolicy.cpp b/Engine/Source/Editor/GraphEditor/Private/BlueprintConnectionDrawingPolicy.cpp
--- a/Engine/Source/Editor/GraphEditor/Private/BlueprintConnectionDrawingPolicy.cpp
+++ b/Engine/Source/Editor/GraphEditor/Private/BlueprintConnectionDrawingPolicy.cpp
@@ -37,6 +37,8 @@ FKismetConnectionDrawingPolicy::FKismetConnectionDrawingPolicy(int32 InBackLayer
 	AttackWireThickness = Settings->TraceAttackWireThickness;
 	SustainWireThickness = Settings->TraceSustainWireThickness;
 	ReleaseWireThickness = Settings->TraceReleaseWireThickness;
+	DefaultDataWireThickness = Settings->DefaultDataWireThickness;
+	DefaultExecutionWireThickness = Settings->DefaultExecutionWireThickness;
 
 	TracePositionBonusPeriod = Settings->TracePositionBonusPeriod;
 	TracePositionExponent = Settings->TracePositionExponent;
@@ -62,10 +64,13 @@ void FKismetConnectionDrawingPolicy::Draw(TMap<TSharedRef<SWidget>, FArrangedWid
 
 bool FKismetConnectionDrawingPolicy::CanBuildRoadmap() const
 {
-	UBlueprint* TargetBP = FBlueprintEditorUtils::FindBlueprintForGraphChecked(GraphObj);
-	UObject* ActiveObject = TargetBP->GetObjectBeingDebugged();
+	UObject* ActiveObject = nullptr;
+	if (UBlueprint* TargetBP = FBlueprintEditorUtils::FindBlueprintForGraph(GraphObj))
+	{
+		ActiveObject = TargetBP->GetObjectBeingDebugged();
+	}
 
-	return ActiveObject != NULL;
+	return ActiveObject != nullptr;
 }
 
 void FKismetConnectionDrawingPolicy::BuildExecutionRoadmap()
@@ -426,6 +431,7 @@ bool FKismetConnectionDrawingPolicy::ShouldChangeTangentForKnot(UK2Node_Knot* Kn
 // Give specific editor modes a chance to highlight this connection or darken non-interesting connections
 void FKismetConnectionDrawingPolicy::DetermineWiringStyle(UEdGraphPin* OutputPin, UEdGraphPin* InputPin, /*inout*/ FConnectionParams& Params)
 {
+	Params.WireThickness = DefaultDataWireThickness;
 	Params.AssociatedPin1 = OutputPin;
 	Params.AssociatedPin2 = InputPin;
 
@@ -434,7 +440,14 @@ void FKismetConnectionDrawingPolicy::DetermineWiringStyle(UEdGraphPin* OutputPin
 	check(GraphObj);
 	const UEdGraphSchema* Schema = GraphObj->GetSchema();
 
-	Params.WireColor = Schema->GetPinTypeColor(OutputPin->PinType);
+	if (OutputPin->bOrphanedPin || (InputPin && InputPin->bOrphanedPin))
+	{
+		Params.WireColor = FLinearColor::Red;
+	}
+	else
+	{
+		Params.WireColor = Schema->GetPinTypeColor(OutputPin->PinType);
+	}
 
 	UEdGraphNode* OutputNode = OutputPin->GetOwningNode();
 	UEdGraphNode* InputNode = (InputPin != nullptr) ? InputPin->GetOwningNode() : nullptr;
@@ -506,7 +519,7 @@ void FKismetConnectionDrawingPolicy::DetermineWiringStyle(UEdGraphPin* OutputPin
 			else
 			{
 				// Make exec wires slightly thicker even outside of debug
-				Params.WireThickness = 3.0f;
+				Params.WireThickness = DefaultExecutionWireThickness;
 			}
 		}
 		else
@@ -514,7 +527,7 @@ void FKismetConnectionDrawingPolicy::DetermineWiringStyle(UEdGraphPin* OutputPin
 			// Container types should draw thicker
 			if ((InputPin && InputPin->PinType.IsContainer()) || (OutputPin && OutputPin->PinType.IsContainer()))
 			{
-				Params.WireThickness = 3.0f;
+				Params.WireThickness = DefaultExecutionWireThickness;
 			}
 		}
 	}
diff --git a/Engine/Source/Editor/GraphEditor/Private/GraphEditorSettings.cpp b/Engine/Source/Editor/GraphEditor/Private/GraphEditorSettings.cpp
--- a/Engine/Source/Editor/GraphEditor/Private/GraphEditorSettings.cpp
+++ b/Engine/Source/Editor/GraphEditor/Private/GraphEditorSettings.cpp
@@ -52,6 +52,9 @@ UGraphEditorSettings::UGraphEditorSettings( const FObjectInitializer& ObjectInit
 	TransformPinTypeColor = FLinearColor(1.0f, 0.172585f, 0.0f, 1.0f);			// orange
 	IndexPinTypeColor = FLinearColor(0.013575f, 0.770000f, 0.429609f, 1.0f);	// green-blue
 
+	DefaultDataWireThickness = 1.0f;
+	DefaultExecutionWireThickness = 3.0f;
+
 	// graph debugging visuals
 	TraceAttackColor = FLinearColor(1.0f, 0.05f, 0.0f, 1.0f);
 	TraceAttackWireThickness = 12.0f;
diff --git a/Engine/Source/Editor/GraphEditor/Private/KismetNodes/KismetNodeInfoContext.cpp b/Engine/Source/Editor/GraphEditor/Private/KismetNodes/KismetNodeInfoContext.cpp
--- a/Engine/Source/Editor/GraphEditor/Private/KismetNodes/KismetNodeInfoContext.cpp
+++ b/Engine/Source/Editor/GraphEditor/Private/KismetNodes/KismetNodeInfoContext.cpp
@@ -64,7 +64,7 @@ FKismetNodeInfoContext::FKismetNodeInfoContext(UEdGraph* SourceGraph)
 
 			for (UObject* ContextObject : LatentContextObjects)
 			{
-				if (UWorld* World = GEngine->GetWorldFromContextObject(ContextObject, /*bChecked =*/false))
+				if (UWorld* World = GEngine->GetWorldFromContextObject(ContextObject, EGetWorldErrorMode::ReturnNull))
 				{
 					FLatentActionManager& Manager = World->GetLatentActionManager();
 
diff --git a/Engine/Source/Editor/GraphEditor/Private/KismetNodes/SGraphNodeK2Base.cpp b/Engine/Source/Editor/GraphEditor/Private/KismetNodes/SGraphNodeK2Base.cpp
--- a/Engine/Source/Editor/GraphEditor/Private/KismetNodes/SGraphNodeK2Base.cpp
+++ b/Engine/Source/Editor/GraphEditor/Private/KismetNodes/SGraphNodeK2Base.cpp
@@ -490,7 +490,7 @@ void SGraphNodeK2Base::GetNodeInfoPopups(FNodeInfoContext* Context, TArray<FGrap
 
 				if (Action.Object == ActiveObject)
 				{
-					if (UWorld* World = GEngine->GetWorldFromContextObject(Action.Object))
+					if (UWorld* World = GEngine->GetWorldFromContextObject(Action.Object, EGetWorldErrorMode::ReturnNull))
 					{
 						FLatentActionManager& LatentActionManager = World->GetLatentActionManager();
 
diff --git a/Engine/Source/Editor/GraphEditor/Private/KismetNodes/SGraphNodeK2Sequence.cpp b/Engine/Source/Editor/GraphEditor/Private/KismetNodes/SGraphNodeK2Sequence.cpp
--- a/Engine/Source/Editor/GraphEditor/Private/KismetNodes/SGraphNodeK2Sequence.cpp
+++ b/Engine/Source/Editor/GraphEditor/Private/KismetNodes/SGraphNodeK2Sequence.cpp
@@ -4,14 +4,14 @@
 #include "EdGraph/EdGraph.h"
 #include "Widgets/SBoxPanel.h"
 #include "GraphEditorSettings.h"
-#include "K2Node_CommutativeAssociativeBinaryOperator.h"
-#include "K2Node_DoOnceMultiInput.h"
-#include "K2Node_ExecutionSequence.h"
-#include "K2Node_MakeArray.h"
+#include "K2Node_AddPinInterface.h"
+#include "K2Node.h"
+#include "ScopedTransaction.h"
 
 
 void SGraphNodeK2Sequence::Construct( const FArguments& InArgs, UK2Node* InNode )
 {
+	ensure(InNode == nullptr || InNode->GetClass()->ImplementsInterface(UK2Node_AddPinInterface::StaticClass()));
 	GraphNode = InNode;
 
 	SetCursor( EMouseCursor::CardinalCross );
@@ -39,30 +39,16 @@ void SGraphNodeK2Sequence::CreateOutputSideAddButton(TSharedPtr<SVerticalBox> Ou
 
 FReply SGraphNodeK2Sequence::OnAddPin()
 {
-	if (UK2Node_ExecutionSequence* SequenceNode = Cast<UK2Node_ExecutionSequence>(GraphNode))
+	IK2Node_AddPinInterface* AddPinNode = Cast<IK2Node_AddPinInterface>(GraphNode);
+	ensure(AddPinNode);
+	if (AddPinNode && AddPinNode->CanAddPin())
 	{
-		SequenceNode->AddPinToExecutionNode();
-		UpdateGraphNode();
-		GraphNode->GetGraph()->NotifyGraphChanged();
-	}
-	else if (UK2Node_MakeArray* MakeArrayNode = Cast<UK2Node_MakeArray>(GraphNode))
-	{
-		MakeArrayNode->AddInputPin();
-		UpdateGraphNode();
-		GraphNode->GetGraph()->NotifyGraphChanged();
-	}
-	else if (UK2Node_CommutativeAssociativeBinaryOperator* OperatorNode = Cast<UK2Node_CommutativeAssociativeBinaryOperator>(GraphNode))
-	{
-		OperatorNode->AddInputPin();
-		UpdateGraphNode();
-		GraphNode->GetGraph()->NotifyGraphChanged();
-	}
-	else if (UK2Node_DoOnceMultiInput* DoOnceMultiNode = Cast<UK2Node_DoOnceMultiInput>(GraphNode))
-	{
-		DoOnceMultiNode->AddInputPin();
+		FScopedTransaction Transaction(NSLOCTEXT("SequencerNode", "AddPinTransaction", "AddPin"));
+
+		AddPinNode->AddInputPin();
 		UpdateGraphNode();
 		GraphNode->GetGraph()->NotifyGraphChanged();
 	}
-
+	
 	return FReply::Handled();
 }
diff --git a/Engine/Source/Editor/GraphEditor/Private/KismetPins/SGraphPinIndex.cpp b/Engine/Source/Editor/GraphEditor/Private/KismetPins/SGraphPinIndex.cpp
--- a/Engine/Source/Editor/GraphEditor/Private/KismetPins/SGraphPinIndex.cpp
+++ b/Engine/Source/Editor/GraphEditor/Private/KismetPins/SGraphPinIndex.cpp
@@ -32,6 +32,7 @@ void SGraphPinIndex::OnTypeChanged(const FEdGraphPinType& PinType)
 {
 	if (GraphPinObj)
 	{
+		GraphPinObj->Modify();
 		GraphPinObj->PinType = PinType;
 		// Let the node know that one of its' pins had their pin type changed
 		if (GraphPinObj->GetOwningNode())
diff --git a/Engine/Source/Editor/GraphEditor/Private/SGraphEditorImpl.cpp b/Engine/Source/Editor/GraphEditor/Private/SGraphEditorImpl.cpp
--- a/Engine/Source/Editor/GraphEditor/Private/SGraphEditorImpl.cpp
+++ b/Engine/Source/Editor/GraphEditor/Private/SGraphEditorImpl.cpp
@@ -577,7 +577,9 @@ void SGraphEditorImpl::ReconstructNodes()
 		{
 			if (UEdGraphNode* Node = Cast<UEdGraphNode>(*NodeIt))
 			{
+				TGuardValue<ESaveOrphanPinMode> GuardSaveMode(Node->OrphanedPinSaveMode, ESaveOrphanPinMode::SaveNone);
 				Schema->ReconstructNode(*Node);
+				Node->ClearCompilerMessage();
 			}
 		}
 	}
diff --git a/Engine/Source/Editor/GraphEditor/Private/SGraphNodeKnot.cpp b/Engine/Source/Editor/GraphEditor/Private/SGraphNodeKnot.cpp
--- a/Engine/Source/Editor/GraphEditor/Private/SGraphNodeKnot.cpp
+++ b/Engine/Source/Editor/GraphEditor/Private/SGraphNodeKnot.cpp
@@ -23,7 +23,7 @@ namespace SKnotNodeDefinitions
 	static const float KnotCenterBubbleAdjust = 20.f;
 
 	/** Knot node spacer sizes */
-	static const FVector2D NodeSpacerSize(32.0f, 32.0f);
+	static const FVector2D NodeSpacerSize(42.0f, 24.0f);
 }
 
 class FAmbivalentDirectionDragConnection : public FDragConnection
@@ -126,21 +126,14 @@ void FAmbivalentDirectionDragConnection::ValidateGraphPinList(TArray<UEdGraphPin
 
 		if (UEdGraphPin* TargetPinObj = GetHoveredPin())
 		{
-// 			if (UK2Node_Knot* TargetKnot = Cast<UK2Node_Knot>(TargetPinObj->GetOwningNode()))
-// 			{
-// 				// The visible pin on a knot is always an output, so Rely on the direction matching; since the visible pin on another knot is always an output
-// 			}
-// 			else
+			// Dragging to another pin, pick the opposite direction as a source to maximize connection chances
+			if (TargetPinObj->Direction == EGPD_Input)
 			{
-				// Dragging to another pin, pick the opposite direction as a source to maximize connection chances
-				if (TargetPinObj->Direction == EGPD_Input)
-				{
-					bUseOutput = true;
-				}
-				else
-				{
-					bUseOutput = false;
-				}
+				bUseOutput = true;
+			}
+			else
+			{
+				bUseOutput = false;
 			}
 		}
 
@@ -311,9 +304,6 @@ void SGraphNodeKnot::Construct(const FArguments& InArgs, UEdGraphNode* InKnot)
 	verify(InKnot->ShouldDrawNodeAsControlPointOnly(InputPinIndex, OutputPinIndex) == true &&
 		InputPinIndex >= 0 && OutputPinIndex >= 0);
 	SGraphNodeDefault::Construct(SGraphNodeDefault::FArguments().GraphNodeObj(InKnot));
-
-	ShadowBrush = FEditorStyle::GetBrush(TEXT("Graph.Node.RerouteShadow"));
-	ShadowBrushSelected = FEditorStyle::GetBrush(TEXT("Graph.Node.RerouteShadowSelected"));
 }
 
 
@@ -339,32 +329,47 @@ void SGraphNodeKnot::UpdateGraphNode()
 	this->ContentScale.Bind( this, &SGraphNode::GetContentScale );
 
 	this->GetOrAddSlot( ENodeZone::Center )
-	.HAlign(HAlign_Center)
-	.VAlign(VAlign_Center)
-	[
-		SNew(SOverlay)
-		+SOverlay::Slot()
-		[
-			// Grab handle to be able to move the node
-			SNew(SSpacer)
-			.Size(SKnotNodeDefinitions::NodeSpacerSize)
-			.Visibility(EVisibility::Visible)
-			.Cursor(EMouseCursor::CardinalCross)
-		]
-		+ SOverlay::Slot()
-		.VAlign(VAlign_Center)
 		.HAlign(HAlign_Center)
-		[
-			SAssignNew(LeftNodeBox, SVerticalBox)
-		]
-		+SOverlay::Slot()
 		.VAlign(VAlign_Center)
-		.HAlign(HAlign_Center)
 		[
-			SAssignNew(RightNodeBox, SVerticalBox)
-		]
-	];
-
+			SNew(SOverlay)
+			+SOverlay::Slot()
+			[
+				// Grab handle to be able to move the node
+				SNew(SSpacer)
+				.Size(SKnotNodeDefinitions::NodeSpacerSize)
+				.Visibility(EVisibility::Visible)
+				.Cursor(EMouseCursor::CardinalCross)
+			]
+
+			+SOverlay::Slot()
+// 			.VAlign(VAlign_Center)
+// 			.HAlign(HAlign_Center)
+			[
+				SNew(SVerticalBox)
+				+SVerticalBox::Slot()
+				.VAlign(VAlign_Top)
+				.HAlign(HAlign_Center)
+				[
+					SNew(SHorizontalBox)
+					+ SHorizontalBox::Slot()
+					.AutoWidth()
+					[
+						SNew(SOverlay)
+						+SOverlay::Slot()
+						[
+							// LEFT
+							SAssignNew(LeftNodeBox, SVerticalBox)
+						]
+						+SOverlay::Slot()
+						[
+							// RIGHT
+							SAssignNew(RightNodeBox, SVerticalBox)
+						]
+					]
+				]
+			]
+		];
 	// Create comment bubble
 	const FSlateColor CommentColor = GetDefault<UGraphEditorSettings>()->DefaultCommentNodeTitleColor;
 
@@ -393,7 +398,7 @@ void SGraphNodeKnot::UpdateGraphNode()
 
 const FSlateBrush* SGraphNodeKnot::GetShadowBrush(bool bSelected) const
 {
-	return bSelected ? ShadowBrushSelected : ShadowBrush;
+	return bSelected ? FEditorStyle::GetBrush(TEXT("Graph.Node.ShadowSelected")) : FEditorStyle::GetNoBrush();
 }
 
 TSharedPtr<SGraphPin> SGraphNodeKnot::CreatePinWidget(UEdGraphPin* Pin) const
diff --git a/Engine/Source/Editor/GraphEditor/Private/SGraphPin.cpp b/Engine/Source/Editor/GraphEditor/Private/SGraphPin.cpp
--- a/Engine/Source/Editor/GraphEditor/Private/SGraphPin.cpp
+++ b/Engine/Source/Editor/GraphEditor/Private/SGraphPin.cpp
@@ -317,7 +317,7 @@ FReply SGraphPin::OnPinMouseDown( const FGeometry& SenderGeometry, const FPointe
 
 	if (MouseEvent.GetEffectingButton() == EKeys::LeftMouseButton)
 	{
-		if (!GraphPinObj->bNotConnectable && IsEditingEnabled())
+		if (IsEditingEnabled())
 		{
 			if (MouseEvent.IsAltDown())
 			{
@@ -327,7 +327,7 @@ FReply SGraphPin::OnPinMouseDown( const FGeometry& SenderGeometry, const FPointe
 				return FReply::Handled();
 			}
 
-			auto OwnerNodePinned = OwnerNodePtr.Pin();
+			TSharedPtr<SGraphNode> OwnerNodePinned = OwnerNodePtr.Pin();
 			if (MouseEvent.IsControlDown() && (GraphPinObj->LinkedTo.Num() > 0))
 			{
 				// Get a reference to the owning panel widget
@@ -341,9 +341,8 @@ FReply SGraphPin::OnPinMouseDown( const FGeometry& SenderGeometry, const FPointe
 
 				// Construct a UEdGraphPin->SGraphPin mapping for the full pin set
 				TMap< FGraphPinHandle, TSharedRef<SGraphPin> > PinToPinWidgetMap;
-				for( TSet< TSharedRef<SWidget> >::TIterator ConnectorIt(AllPins); ConnectorIt; ++ConnectorIt )
+				for (const TSharedRef<SWidget>& SomePinWidget : AllPins)
 				{
-					const TSharedRef<SWidget>& SomePinWidget = *ConnectorIt;
 					const SGraphPin& PinWidget = static_cast<const SGraphPin&>(SomePinWidget.Get());
 
 					UEdGraphPin* GraphPin = PinWidget.GetPinObj();
@@ -354,7 +353,7 @@ FReply SGraphPin::OnPinMouseDown( const FGeometry& SenderGeometry, const FPointe
 				}
 
 				// Define a local struct to temporarily store lookup information for pins that we are currently linked to
-				struct LinkedToPinInfo
+				struct FLinkedToPinInfo
 				{
 					// Pin name string
 					FString PinName;
@@ -364,33 +363,30 @@ FReply SGraphPin::OnPinMouseDown( const FGeometry& SenderGeometry, const FPointe
 				};
 
 				// Build a lookup table containing information about the set of pins that we're currently linked to
-				TArray<LinkedToPinInfo> LinkedToPinInfoArray;
-				for( TArray<UEdGraphPin*>::TIterator LinkArrayIter(GetPinObj()->LinkedTo); LinkArrayIter; ++LinkArrayIter )
+				TArray<FLinkedToPinInfo> LinkedToPinInfoArray;
+				for (UEdGraphPin* Pin : GetPinObj()->LinkedTo)
 				{
-					if (auto PinWidget = PinToPinWidgetMap.Find(*LinkArrayIter))
+					if (auto PinWidget = PinToPinWidgetMap.Find(Pin))
 					{
 						check((*PinWidget)->OwnerNodePtr.IsValid());
 
-						LinkedToPinInfo PinInfo;
+						FLinkedToPinInfo PinInfo;
 						PinInfo.PinName = (*PinWidget)->GetPinObj()->PinName;
 						PinInfo.OwnerNodePtr = (*PinWidget)->OwnerNodePtr.Pin()->GetNodeObj();
-						LinkedToPinInfoArray.Add(PinInfo);
+						LinkedToPinInfoArray.Add(MoveTemp(PinInfo));
 					}
 				}
-				
-				
+
+
 				// Now iterate over our lookup table to find the instances of pin widgets that we had previously linked to
 				TArray<TSharedRef<SGraphPin>> PinArray;
-				for(auto LinkedToPinInfoIter = LinkedToPinInfoArray.CreateConstIterator(); LinkedToPinInfoIter; ++LinkedToPinInfoIter)
+				for (FLinkedToPinInfo PinInfo : LinkedToPinInfoArray)
 				{
-					LinkedToPinInfo PinInfo = *LinkedToPinInfoIter;
-					UEdGraphNode* OwnerNodeObj = PinInfo.OwnerNodePtr.Get();
-					if(OwnerNodeObj != NULL)
+					if (UEdGraphNode* OwnerNodeObj = PinInfo.OwnerNodePtr.Get())
 					{
-						for(auto PinIter = PinInfo.OwnerNodePtr.Get()->Pins.CreateConstIterator(); PinIter; ++PinIter)
+						for (UEdGraphPin* Pin : PinInfo.OwnerNodePtr.Get()->Pins)
 						{
-							UEdGraphPin* Pin = *PinIter;
-							if(Pin->PinName == PinInfo.PinName)
+							if (Pin->PinName == PinInfo.PinName)
 							{
 								if (TSharedRef<SGraphPin>* pWidget = PinToPinWidgetMap.Find(FGraphPinHandle(Pin)))
 								{
@@ -414,7 +410,7 @@ FReply SGraphPin::OnPinMouseDown( const FGeometry& SenderGeometry, const FPointe
 				const UEdGraphSchema* Schema = GraphPinObj->GetSchema();
 				Schema->BreakPinLinks(*GraphPinObj, true);
 
-				if(DragEvent.IsValid())
+				if (DragEvent.IsValid())
 				{
 					bIsMovingLinks = true;
 					return FReply::Handled().BeginDragDrop(DragEvent.ToSharedRef());
@@ -425,25 +421,26 @@ FReply SGraphPin::OnPinMouseDown( const FGeometry& SenderGeometry, const FPointe
 					return FReply::Handled();
 				}
 			}
-			
-			// Start a drag-drop on the pin
-			if (ensure(OwnerNodePinned.IsValid()))
-			{
-				TArray<TSharedRef<SGraphPin>> PinArray;
-				PinArray.Add(SharedThis(this));
 
-				return FReply::Handled().BeginDragDrop(SpawnPinDragEvent(OwnerNodePinned->GetOwnerPanel().ToSharedRef(), PinArray));
-			}
-			else
+			if (!GraphPinObj->bNotConnectable)
 			{
-				return FReply::Unhandled();
+				// Start a drag-drop on the pin
+				if (ensure(OwnerNodePinned.IsValid()))
+				{
+					TArray<TSharedRef<SGraphPin>> PinArray;
+					PinArray.Add(SharedThis(this));
+
+					return FReply::Handled().BeginDragDrop(SpawnPinDragEvent(OwnerNodePinned->GetOwnerPanel().ToSharedRef(), PinArray));
+				}
+				else
+				{
+					return FReply::Unhandled();
+				}
 			}
 		}
-		else
-		{
-			// It's not connectable, but we don't want anything above us to process this left click.
-			return FReply::Handled();
-		}
+
+		// It's not connectible, but we don't want anything above us to process this left click.
+		return FReply::Handled();
 	}
 	else
 	{
@@ -950,6 +947,10 @@ FSlateColor SGraphPin::GetPinColor() const
 		{
 			return FSlateColor(FLinearColor(0.9f, 0.2f, 0.15f));
 		}
+		if (GraphPinObj->bOrphanedPin)
+		{
+			return FSlateColor(FLinearColor::Red);
+		}
 		if (const UEdGraphSchema* Schema = GraphPinObj->GetSchema())
 		{
 			if (!GetPinObj()->GetOwningNode()->IsNodeEnabled() || !IsEditingEnabled())
@@ -975,11 +976,21 @@ FSlateColor SGraphPin::GetPinTextColor() const
 	// If there is no schema there is no owning node (or basically this is a deleted node)
 	if (UEdGraphNode* GraphNode = GraphPinObj->GetOwningNodeUnchecked())
 	{
-		if(!GraphNode->IsNodeEnabled() || !IsEditingEnabled())
+		const bool bDisabled = (!GraphNode->IsNodeEnabled() || !IsEditingEnabled());
+		if (GraphPinObj->bOrphanedPin)
+		{
+			FLinearColor PinColor = FLinearColor::Red;
+			if (bDisabled)
+			{
+				PinColor.A = 0.5f;
+			}
+			return PinColor;
+		}
+		else if (bDisabled)
 		{
 			return FLinearColor(1.0f, 1.0f, 1.0f, 0.5f);
 		}
-		else if(bUsePinColorForText)
+		if (bUsePinColorForText)
 		{
 			return GetPinColor();
 		}
@@ -1051,7 +1062,7 @@ EVisibility SGraphPin::GetDefaultValueVisibility() const
 		return EVisibility::Collapsed;
 	}
 
-	if (GraphPinObj->bNotConnectable)
+	if (GraphPinObj->bNotConnectable && !GraphPinObj->bOrphanedPin)
 	{
 		// The only reason this pin exists is to show something, so do so
 		return EVisibility::Visible;
diff --git a/Engine/Source/Editor/GraphEditor/Public/BlueprintConnectionDrawingPolicy.h b/Engine/Source/Editor/GraphEditor/Public/BlueprintConnectionDrawingPolicy.h
--- a/Engine/Source/Editor/GraphEditor/Public/BlueprintConnectionDrawingPolicy.h
+++ b/Engine/Source/Editor/GraphEditor/Public/BlueprintConnectionDrawingPolicy.h
@@ -51,6 +51,8 @@ protected:
 	float AttackWireThickness;
 	float SustainWireThickness;
 	float ReleaseWireThickness;
+	float DefaultDataWireThickness;
+	float DefaultExecutionWireThickness;
 
 	float TracePositionBonusPeriod;
 	float TracePositionExponent;
diff --git a/Engine/Source/Editor/GraphEditor/Public/GraphEditorSettings.h b/Engine/Source/Editor/GraphEditor/Public/GraphEditorSettings.h
--- a/Engine/Source/Editor/GraphEditor/Public/GraphEditorSettings.h
+++ b/Engine/Source/Editor/GraphEditor/Public/GraphEditorSettings.h
@@ -247,7 +247,15 @@ public:
 	FLinearColor DefaultCommentNodeTitleColor;
 
 public:
+	/** The thickness of a data wire */
+	UPROPERTY(EditAnywhere, config, Category=Tracing)
+	float DefaultDataWireThickness;
+
+	/** The thickness of an execution wire when not debugging */
+	UPROPERTY(EditAnywhere, config, Category=Tracing)
+	float DefaultExecutionWireThickness;
 
+	/** The color to display execution wires that were just executed */
 	UPROPERTY(EditAnywhere, config, Category=Tracing)
 	FLinearColor TraceAttackColor;
 
@@ -275,9 +283,11 @@ public:
 	UPROPERTY()
 	float TraceSustainHoldPeriod;
 
+	/** The color to fade to for execution wires on release */
 	UPROPERTY(EditAnywhere, config, Category=Tracing)
 	FLinearColor TraceReleaseColor;
 
+	/** The thickness to drop down to during release / for unexecuted wires when debugging */
 	UPROPERTY(EditAnywhere, config, Category=Tracing)
 	float TraceReleaseWireThickness;
 
diff --git a/Engine/Source/Editor/IntroTutorials/Private/EditorTutorial.cpp b/Engine/Source/Editor/IntroTutorials/Private/EditorTutorial.cpp
--- a/Engine/Source/Editor/IntroTutorials/Private/EditorTutorial.cpp
+++ b/Engine/Source/Editor/IntroTutorials/Private/EditorTutorial.cpp
@@ -13,6 +13,7 @@
 #include "IntroTutorials.h"
 #include "LevelEditor.h"
 #include "Widgets/Docking/SDockTab.h"
+#include "Misc/RuntimeErrors.h"
 
 #if WITH_EDITORONLY_DATA
 namespace
@@ -99,7 +100,10 @@ void UEditorTutorial::HandleTutorialClosed()
 
 void UEditorTutorial::OpenAsset(UObject* Asset)
 {
-	FAssetEditorManager::Get().OpenEditorForAsset(Asset);
+	if (ensureAsRuntimeWarning(Asset != nullptr))
+	{
+		FAssetEditorManager::Get().OpenEditorForAsset(Asset);
+	}
 }
 
 AActor* UEditorTutorial::GetActorReference(FString PathToActor)
