Copying //UE4/Dev-Framework to //UE4/Dev-Main (Source: //UE4/Dev-Framework @ 3628051)

diff --git a/Engine/Source/Editor/DetailCustomizations/Private/ComponentTransformDetails.cpp b/Engine/Source/Editor/DetailCustomizations/Private/ComponentTransformDetails.cpp
--- a/Engine/Source/Editor/DetailCustomizations/Private/ComponentTransformDetails.cpp
+++ b/Engine/Source/Editor/DetailCustomizations/Private/ComponentTransformDetails.cpp
@@ -112,16 +112,16 @@ TSharedRef<SWidget> FComponentTransformDetails::BuildTransformFieldLabel( ETrans
 
 	FUIAction SetRelativeLocationAction
 	(
-		FExecuteAction::CreateSP( this, &FComponentTransformDetails::OnSetRelativeTransform, TransformField ),
+		FExecuteAction::CreateSP( this, &FComponentTransformDetails::OnSetAbsoluteTransform, TransformField, false ),
 		FCanExecuteAction(),
-		FIsActionChecked::CreateSP( this, &FComponentTransformDetails::IsRelativeTransformChecked, TransformField )
+		FIsActionChecked::CreateSP( this, &FComponentTransformDetails::IsAbsoluteTransformChecked, TransformField, false )
 	);
 
 	FUIAction SetWorldLocationAction
 	(
-		FExecuteAction::CreateSP( this, &FComponentTransformDetails::OnSetWorldTransform, TransformField ),
+		FExecuteAction::CreateSP( this, &FComponentTransformDetails::OnSetAbsoluteTransform, TransformField, true ),
 		FCanExecuteAction(),
-		FIsActionChecked::CreateSP( this, &FComponentTransformDetails::IsWorldTransformChecked, TransformField )
+		FIsActionChecked::CreateSP( this, &FComponentTransformDetails::IsAbsoluteTransformChecked, TransformField, true )
 	);
 
 	MenuBuilder.BeginSection( TEXT("TransformType"), FText::Format( LOCTEXT("TransformType", "{0} Type"), Label ) );
@@ -189,68 +189,6 @@ FText FComponentTransformDetails::GetTransformFieldText( ETransformField::Type T
 	}
 }
 
-void FComponentTransformDetails::OnSetRelativeTransform(  ETransformField::Type TransformField )
-{
-	const bool bEnable = false;
-	switch (TransformField)
-	{
-	case ETransformField::Location:
-		OnToggleAbsoluteLocation(bEnable);
-		break;
-	case ETransformField::Rotation:
-		OnToggleAbsoluteRotation( bEnable );
-		break;
-	case ETransformField::Scale:
-		OnToggleAbsoluteScale( bEnable );
-		break;
-	default:
-		break;
-	}
-}
-
-bool FComponentTransformDetails::IsRelativeTransformChecked( ETransformField::Type TransformField ) const
-{
-	switch (TransformField)
-	{
-	case ETransformField::Location:
-		return !bAbsoluteLocation;
-		break;
-	case ETransformField::Rotation:
-		return !bAbsoluteRotation;
-		break;
-	case ETransformField::Scale:
-		return !bAbsoluteScale;
-		break;
-	default:
-		return false;
-		break;
-	}
-}
-
-void FComponentTransformDetails::OnSetWorldTransform( ETransformField::Type TransformField )
-{
-	const bool bEnable = true;
-	switch (TransformField)
-	{
-	case ETransformField::Location:
-		OnToggleAbsoluteLocation(bEnable);
-		break;
-	case ETransformField::Rotation:
-		OnToggleAbsoluteRotation(bEnable);
-		break;
-	case ETransformField::Scale:
-		OnToggleAbsoluteScale(bEnable);
-		break;
-	default:
-		break;
-	}
-}
-
-bool FComponentTransformDetails::IsWorldTransformChecked( ETransformField::Type TransformField ) const
-{
-	return !IsRelativeTransformChecked( TransformField );
-}
-
 bool FComponentTransformDetails::OnCanCopy( ETransformField::Type TransformField ) const
 {
 	// We can only copy values if the whole field is set.  If multiple values are defined we do not copy since we are unable to determine the value
@@ -310,9 +248,7 @@ void FComponentTransformDetails::OnPaste( ETransformField::Type TransformField )
 			if (Location.InitFromString(PastedText))
 			{
 				FScopedTransaction Transaction(LOCTEXT("PasteLocation", "Paste Location"));
-				OnSetLocation(Location.X, ETextCommit::Default, 0);
-				OnSetLocation(Location.Y, ETextCommit::Default, 1);
-				OnSetLocation(Location.Z, ETextCommit::Default, 2);
+				OnSetTransform(ETransformField::Location, EAxisList::All, Location, false, true);
 			}
 		}
 		break;
@@ -325,9 +261,7 @@ void FComponentTransformDetails::OnPaste( ETransformField::Type TransformField )
 			if (Rotation.InitFromString(PastedText))
 			{
 				FScopedTransaction Transaction(LOCTEXT("PasteRotation", "Paste Rotation"));
-				OnSetRotation(Rotation.Roll, ETextCommit::Default, 0);
-				OnSetRotation(Rotation.Pitch, ETextCommit::Default, 1);
-				OnSetRotation(Rotation.Yaw, ETextCommit::Default, 2);
+				OnSetTransform(ETransformField::Rotation, EAxisList::All, Rotation.Euler(), false, true);
 			}
 		}
 		break;
@@ -337,9 +271,7 @@ void FComponentTransformDetails::OnPaste( ETransformField::Type TransformField )
 			if (Scale.InitFromString(PastedText))
 			{
 				FScopedTransaction Transaction(LOCTEXT("PasteScale", "Paste Scale"));
-				OnSetScale(Scale.X, ETextCommit::Default, 0);
-				OnSetScale(Scale.Y, ETextCommit::Default, 1);
-				OnSetScale(Scale.Z, ETextCommit::Default, 2);
+				OnSetTransform(ETransformField::Scale, EAxisList::All, Scale, false, true);
 			}
 		}
 		break;
@@ -409,9 +341,9 @@ void FComponentTransformDetails::GenerateChildContent( IDetailChildrenBuilder& C
 				.bColorAxisLabels( true )
 				.AllowResponsiveLayout( true )
 				.IsEnabled( this, &FComponentTransformDetails::GetIsEnabled )
-				.OnXCommitted( this, &FComponentTransformDetails::OnSetLocation, 0 )
-				.OnYCommitted( this, &FComponentTransformDetails::OnSetLocation, 1 )
-				.OnZCommitted( this, &FComponentTransformDetails::OnSetLocation, 2 )
+				.OnXCommitted( this, &FComponentTransformDetails::OnSetTransformAxis, ETransformField::Location, EAxisList::X, true )
+				.OnYCommitted( this, &FComponentTransformDetails::OnSetTransformAxis, ETransformField::Location, EAxisList::Y, true )
+				.OnZCommitted( this, &FComponentTransformDetails::OnSetTransformAxis, ETransformField::Location, EAxisList::Z, true )
 				.Font( FontInfo )
 				.TypeInterface( TypeInterface )
 				.AllowSpin(false)
@@ -479,12 +411,12 @@ void FComponentTransformDetails::GenerateChildContent( IDetailChildrenBuilder& C
 				.IsEnabled( this, &FComponentTransformDetails::GetIsEnabled )
 				.OnBeginSliderMovement( this, &FComponentTransformDetails::OnBeginRotatonSlider )
 				.OnEndSliderMovement( this, &FComponentTransformDetails::OnEndRotationSlider )
-				.OnRollChanged( this, &FComponentTransformDetails::OnSetRotation, false, 0 )
-				.OnPitchChanged( this, &FComponentTransformDetails::OnSetRotation, false, 1 )
-				.OnYawChanged( this, &FComponentTransformDetails::OnSetRotation, false, 2 )
-				.OnRollCommitted( this, &FComponentTransformDetails::OnRotationCommitted, 0 )
-				.OnPitchCommitted( this, &FComponentTransformDetails::OnRotationCommitted, 1 )
-				.OnYawCommitted( this, &FComponentTransformDetails::OnRotationCommitted, 2 )
+				.OnRollChanged( this, &FComponentTransformDetails::OnSetTransformAxis, ETextCommit::Default, ETransformField::Rotation, EAxisList::X, false )
+				.OnPitchChanged( this, &FComponentTransformDetails::OnSetTransformAxis, ETextCommit::Default, ETransformField::Rotation, EAxisList::Y, false )
+				.OnYawChanged( this, &FComponentTransformDetails::OnSetTransformAxis, ETextCommit::Default, ETransformField::Rotation, EAxisList::Z, false )
+				.OnRollCommitted( this, &FComponentTransformDetails::OnSetTransformAxis, ETransformField::Rotation, EAxisList::X, true )
+				.OnPitchCommitted( this, &FComponentTransformDetails::OnSetTransformAxis, ETransformField::Rotation, EAxisList::Y, true )
+				.OnYawCommitted( this, &FComponentTransformDetails::OnSetTransformAxis, ETransformField::Rotation, EAxisList::Z, true )
 				.TypeInterface(TypeInterface)
 				.Font( FontInfo )
 			]
@@ -542,9 +474,9 @@ void FComponentTransformDetails::GenerateChildContent( IDetailChildrenBuilder& C
 				.bColorAxisLabels( true )
 				.AllowResponsiveLayout( true )
 				.IsEnabled( this, &FComponentTransformDetails::GetIsEnabled )
-				.OnXCommitted( this, &FComponentTransformDetails::OnSetScale, 0 )
-				.OnYCommitted( this, &FComponentTransformDetails::OnSetScale, 1 )
-				.OnZCommitted( this, &FComponentTransformDetails::OnSetScale, 2 )
+				.OnXCommitted( this, &FComponentTransformDetails::OnSetTransformAxis, ETransformField::Scale, EAxisList::X, true )
+				.OnYCommitted( this, &FComponentTransformDetails::OnSetTransformAxis, ETransformField::Scale, EAxisList::Y, true )
+				.OnZCommitted( this, &FComponentTransformDetails::OnSetTransformAxis, ETransformField::Scale, EAxisList::Z, true )
 				.ContextMenuExtenderX( this, &FComponentTransformDetails::ExtendXScaleContextMenu )
 				.ContextMenuExtenderY( this, &FComponentTransformDetails::ExtendYScaleContextMenu )
 				.ContextMenuExtenderZ( this, &FComponentTransformDetails::ExtendZScaleContextMenu )
@@ -653,218 +585,153 @@ FText FComponentTransformDetails::GetScaleText() const
 	return bAbsoluteScale ? LOCTEXT( "AbsoluteScale", "Absolute Scale" ) : LOCTEXT( "Scale", "Scale" );
 }
 
-void FComponentTransformDetails::OnToggleAbsoluteLocation( bool bEnable )
+void FComponentTransformDetails::OnSetAbsoluteTransform(ETransformField::Type TransformField, bool bAbsoluteEnabled)
 {
-	UProperty* AbsoluteLocationProperty = FindField<UProperty>( USceneComponent::StaticClass(), "bAbsoluteLocation" );
+	UBoolProperty* AbsoluteProperty = nullptr;
+	FText TransactionText;
+
+	switch (TransformField)
+	{
+	case ETransformField::Location:
+		AbsoluteProperty = FindField<UBoolProperty>(USceneComponent::StaticClass(), GET_MEMBER_NAME_CHECKED(USceneComponent, bAbsoluteLocation));
+		TransactionText = LOCTEXT("ToggleAbsoluteLocation", "Toggle Absolute Location");
+		break;
+	case ETransformField::Rotation:
+		AbsoluteProperty = FindField<UBoolProperty>(USceneComponent::StaticClass(), GET_MEMBER_NAME_CHECKED(USceneComponent, bAbsoluteRotation));
+		TransactionText = LOCTEXT("ToggleAbsoluteRotation", "Toggle Absolute Rotation");
+		break;
+	case ETransformField::Scale:
+		AbsoluteProperty = FindField<UBoolProperty>(USceneComponent::StaticClass(), GET_MEMBER_NAME_CHECKED(USceneComponent, bAbsoluteScale));
+		TransactionText = LOCTEXT("ToggleAbsoluteScale", "Toggle Absolute Scale");
+		break;
+	default:
+		return;
+	}
 
 	bool bBeganTransaction = false;
-	for( int32 ObjectIndex = 0; ObjectIndex < SelectedObjects.Num(); ++ObjectIndex )
+	TArray<UObject*> ModifiedObjects;
+	for (int32 ObjectIndex = 0; ObjectIndex < SelectedObjects.Num(); ++ObjectIndex)
 	{
 		TWeakObjectPtr<UObject> ObjectPtr = SelectedObjects[ObjectIndex];
-		if( ObjectPtr.IsValid() )
+		if (ObjectPtr.IsValid())
 		{
 			UObject* Object = ObjectPtr.Get();
-			USceneComponent* SceneComponent = GetSceneComponentFromDetailsObject( Object );
-			if( SceneComponent && SceneComponent->bAbsoluteLocation != bEnable )
+			USceneComponent* SceneComponent = GetSceneComponentFromDetailsObject(Object);
+			if (SceneComponent)
 			{
-				if( !bBeganTransaction )
+				bool bOldValue = TransformField == ETransformField::Location ? SceneComponent->bAbsoluteLocation : (TransformField == ETransformField::Rotation ? SceneComponent->bAbsoluteRotation : SceneComponent->bAbsoluteScale);
+
+				if (bOldValue == bAbsoluteEnabled)
+				{
+					// Already the desired value
+					continue;
+				}
+
+				if (!bBeganTransaction)
 				{
 					// NOTE: One transaction per change, not per actor
-					GEditor->BeginTransaction( LOCTEXT("ToggleAbsoluteLocation", "Toggle Absolute Location" ) );
+					GEditor->BeginTransaction(TransactionText);
 					bBeganTransaction = true;
 				}
 
-				FScopedSwitchWorldForObject WorldSwitcher( Object );
-				
-				if (Object->HasAnyFlags(RF_DefaultSubObject))
+				FScopedSwitchWorldForObject WorldSwitcher(Object);
+
+				if (SceneComponent->HasAnyFlags(RF_DefaultSubObject))
 				{
 					// Default subobjects must be included in any undo/redo operations
-					Object->SetFlags(RF_Transactional);
+					SceneComponent->SetFlags(RF_Transactional);
 				}
 
-				Object->PreEditChange( AbsoluteLocationProperty );
+				SceneComponent->PreEditChange(AbsoluteProperty);
 
-				if( NotifyHook )
+				if (NotifyHook)
 				{
-					NotifyHook->NotifyPreChange( AbsoluteLocationProperty );
+					NotifyHook->NotifyPreChange(AbsoluteProperty);
 				}
 
-				SceneComponent->bAbsoluteLocation = !SceneComponent->bAbsoluteLocation;
-				
-				// Update RelativeLocation to maintain/stabilize position when switching between relative and world.
-				if (SceneComponent->GetAttachParent())
+				switch (TransformField)
 				{
-					if (SceneComponent->bAbsoluteLocation)
-					{
-						SceneComponent->RelativeLocation = SceneComponent->GetComponentTransform().GetTranslation();
-					}
-					else
+				case ETransformField::Location:
+					SceneComponent->bAbsoluteLocation = bAbsoluteEnabled;
+
+					// Update RelativeLocation to maintain/stabilize position when switching between relative and world.
+					if (SceneComponent->GetAttachParent())
 					{
-						FTransform ParentToWorld = SceneComponent->GetAttachParent()->GetSocketTransform(SceneComponent->GetAttachSocketName());
-						FTransform RelativeTM = SceneComponent->GetComponentTransform().GetRelativeTransform(ParentToWorld);
-						SceneComponent->RelativeLocation = RelativeTM.GetTranslation();
+						if (SceneComponent->bAbsoluteLocation)
+						{
+							SceneComponent->RelativeLocation = SceneComponent->GetComponentTransform().GetTranslation();
+						}
+						else
+						{
+							FTransform ParentToWorld = SceneComponent->GetAttachParent()->GetSocketTransform(SceneComponent->GetAttachSocketName());
+							FTransform RelativeTM = SceneComponent->GetComponentTransform().GetRelativeTransform(ParentToWorld);
+							SceneComponent->RelativeLocation = RelativeTM.GetTranslation();
+						}
 					}
-				}
-				
-				FPropertyChangedEvent PropertyChangedEvent( AbsoluteLocationProperty );
-				Object->PostEditChangeProperty(PropertyChangedEvent);
-
-				// If it's a template, propagate the change out to any current instances of the object
-				if(SceneComponent->IsTemplate())
-				{
-					uint32 NewValue = SceneComponent->bAbsoluteLocation;
-					uint32 OldValue = !NewValue;
-					TSet<USceneComponent*> UpdatedInstances;
-					FComponentEditorUtils::PropagateDefaultValueChange(SceneComponent, AbsoluteLocationProperty, OldValue, NewValue, UpdatedInstances);
+					break;
+				case ETransformField::Rotation:
+					SceneComponent->bAbsoluteRotation = bAbsoluteEnabled;
+					break;
+				case ETransformField::Scale:
+					SceneComponent->bAbsoluteScale = bAbsoluteEnabled;
+					break;
 				}
 
-				if( NotifyHook )
-				{
-					NotifyHook->NotifyPostChange( PropertyChangedEvent, AbsoluteLocationProperty );
-				}
+				ModifiedObjects.Add(Object);
 			}
 		}
 	}
 
-	if( bBeganTransaction )
+	if (bBeganTransaction)
 	{
-		GEditor->EndTransaction();
+		FPropertyChangedEvent PropertyChangedEvent(AbsoluteProperty, EPropertyChangeType::ValueSet, (const TArray<const UObject*>*)&ModifiedObjects);
 
-		GUnrealEd->RedrawLevelEditingViewports();
-	}
-
-}
-
-void FComponentTransformDetails::OnToggleAbsoluteRotation( bool bEnable )
-{
-	UProperty* AbsoluteRotationProperty = FindField<UProperty>( USceneComponent::StaticClass(), "bAbsoluteRotation" );
-
-	bool bBeganTransaction = false;
-	for( int32 ObjectIndex = 0; ObjectIndex < SelectedObjects.Num(); ++ObjectIndex )
-	{
-		TWeakObjectPtr<UObject> ObjectPtr = SelectedObjects[ObjectIndex];
-		if( ObjectPtr.IsValid() )
+		for (UObject* Object : ModifiedObjects)
 		{
-			UObject* Object = ObjectPtr.Get();
-			USceneComponent* SceneComponent = GetSceneComponentFromDetailsObject( Object );
-			if( SceneComponent && SceneComponent->bAbsoluteRotation != bEnable )
+			USceneComponent* SceneComponent = GetSceneComponentFromDetailsObject(Object);
+			if (SceneComponent)
 			{
-				if( !bBeganTransaction )
-				{
-					// NOTE: One transaction per change, not per actor
-					GEditor->BeginTransaction( LOCTEXT("ToggleAbsoluteRotation", "Toggle Absolute Rotation" ) );
-					bBeganTransaction = true;
-				}
-
-				FScopedSwitchWorldForObject WorldSwitcher( Object );
-				
-				if (Object->HasAnyFlags(RF_DefaultSubObject))
-				{
-					// Default subobjects must be included in any undo/redo operations
-					Object->SetFlags(RF_Transactional);
-				}
-
-				Object->PreEditChange( AbsoluteRotationProperty );
-
-				if( NotifyHook )
-				{
-					NotifyHook->NotifyPreChange( AbsoluteRotationProperty );
-				}
-
-				SceneComponent->bAbsoluteRotation = !SceneComponent->bAbsoluteRotation;
-
-				FPropertyChangedEvent PropertyChangedEvent( AbsoluteRotationProperty );
-				Object->PostEditChangeProperty(PropertyChangedEvent);
+				SceneComponent->PostEditChangeProperty(PropertyChangedEvent);
 
 				// If it's a template, propagate the change out to any current instances of the object
-				if(SceneComponent->IsTemplate())
+				if (SceneComponent->IsTemplate())
 				{
-					uint32 NewValue = SceneComponent->bAbsoluteRotation;
-					uint32 OldValue = !NewValue;
+					bool NewValue = bAbsoluteEnabled;
+					bool OldValue = !NewValue;
 					TSet<USceneComponent*> UpdatedInstances;
-					FComponentEditorUtils::PropagateDefaultValueChange(SceneComponent, AbsoluteRotationProperty, OldValue, NewValue, UpdatedInstances);
-				}
-
-				if( NotifyHook )
-				{
-					NotifyHook->NotifyPostChange( PropertyChangedEvent, AbsoluteRotationProperty );
+					FComponentEditorUtils::PropagateDefaultValueChange(SceneComponent, AbsoluteProperty, OldValue, NewValue, UpdatedInstances);
 				}
 			}
 		}
-	}
 
-	if( bBeganTransaction )
-	{
+		if (NotifyHook)
+		{
+			NotifyHook->NotifyPostChange(PropertyChangedEvent, AbsoluteProperty);
+		}
+
 		GEditor->EndTransaction();
 
 		GUnrealEd->RedrawLevelEditingViewports();
 	}
 }
 
-void FComponentTransformDetails::OnToggleAbsoluteScale( bool bEnable )
+bool FComponentTransformDetails::IsAbsoluteTransformChecked(ETransformField::Type TransformField, bool bAbsoluteEnabled) const
 {
-	UProperty* AbsoluteScaleProperty = FindField<UProperty>( USceneComponent::StaticClass(), "bAbsoluteScale" );
-
-	bool bBeganTransaction = false;
-	for( int32 ObjectIndex = 0; ObjectIndex < SelectedObjects.Num(); ++ObjectIndex )
-	{
-		TWeakObjectPtr<UObject> ObjectPtr = SelectedObjects[ObjectIndex];
-		if( ObjectPtr.IsValid() )
-		{
-			UObject* Object = ObjectPtr.Get();
-			USceneComponent* SceneComponent = GetSceneComponentFromDetailsObject( Object );
-			if( SceneComponent && SceneComponent->bAbsoluteScale != bEnable )
-			{
-				if( !bBeganTransaction )
-				{
-					// NOTE: One transaction per change, not per actor
-					GEditor->BeginTransaction( LOCTEXT("ToggleAbsoluteScale", "Toggle Absolute Scale" ) );
-					bBeganTransaction = true;
-				}
-
-				FScopedSwitchWorldForObject WorldSwitcher( Object );
-				
-				if (Object->HasAnyFlags(RF_DefaultSubObject))
-				{
-					// Default subobjects must be included in any undo/redo operations
-					Object->SetFlags(RF_Transactional);
-				}
-
-				Object->PreEditChange( AbsoluteScaleProperty );
-
-				if( NotifyHook )
-				{
-					NotifyHook->NotifyPreChange( AbsoluteScaleProperty );
-				}
-
-				SceneComponent->bAbsoluteScale = !SceneComponent->bAbsoluteScale;
-
-				FPropertyChangedEvent PropertyChangedEvent( AbsoluteScaleProperty );
-				Object->PostEditChangeProperty(PropertyChangedEvent);
-
-				// If it's a template, propagate the change out to any current instances of the object
-				if(SceneComponent->IsTemplate())
-				{
-					uint32 NewValue = SceneComponent->bAbsoluteScale;
-					uint32 OldValue = !NewValue;
-					TSet<USceneComponent*> UpdatedInstances;
-					FComponentEditorUtils::PropagateDefaultValueChange(SceneComponent, AbsoluteScaleProperty, OldValue, NewValue, UpdatedInstances);
-				}
-
-				if( NotifyHook )
-				{
-					NotifyHook->NotifyPostChange( PropertyChangedEvent, AbsoluteScaleProperty );
-				}
-			}
-		}
-	}
-
-	if( bBeganTransaction )
+	switch (TransformField)
 	{
-		GEditor->EndTransaction();
-		GUnrealEd->RedrawLevelEditingViewports();
+	case ETransformField::Location:
+		return bAbsoluteLocation == bAbsoluteEnabled;
+		break;
+	case ETransformField::Rotation:
+		return bAbsoluteRotation == bAbsoluteEnabled;
+		break;
+	case ETransformField::Scale:
+		return bAbsoluteScale == bAbsoluteEnabled;
+		break;
+	default:
+		return false;
+		break;
 	}
-
 }
 
 struct FGetRootComponentArchetype
@@ -876,47 +743,68 @@ struct FGetRootComponentArchetype
 	}
 };
 
+EVisibility FComponentTransformDetails::GetLocationResetVisibility() const
+{
+	const USceneComponent* Archetype = FGetRootComponentArchetype::Get(SelectedObjects[0].Get());
+	const FVector Data = Archetype ? Archetype->RelativeLocation : FVector::ZeroVector;
+
+	// unset means multiple differing values, so show "Reset to Default" in that case
+	return CachedLocation.IsSet() && CachedLocation.X.GetValue() == Data.X && CachedLocation.Y.GetValue() == Data.Y && CachedLocation.Z.GetValue() == Data.Z ? EVisibility::Hidden : EVisibility::Visible;
+}
+
 FReply FComponentTransformDetails::OnLocationResetClicked()
 {
 	const FText TransactionName = LOCTEXT("ResetLocation", "Reset Location");
 	FScopedTransaction Transaction(TransactionName);
 
 	const USceneComponent* Archetype = FGetRootComponentArchetype::Get(SelectedObjects[0].Get());
-	const FVector Data = Archetype ? Archetype->RelativeLocation : FVector();
+	const FVector Data = Archetype ? Archetype->RelativeLocation : FVector::ZeroVector;
 
-	OnSetLocation(Data[0], ETextCommit::Default, 0);
-	OnSetLocation(Data[1], ETextCommit::Default, 1);
-	OnSetLocation(Data[2], ETextCommit::Default, 2);
+	OnSetTransform(ETransformField::Location, EAxisList::All, Data, false, true);
 
 	return FReply::Handled();
 }
 
+EVisibility FComponentTransformDetails::GetRotationResetVisibility() const
+{
+	const USceneComponent* Archetype = FGetRootComponentArchetype::Get(SelectedObjects[0].Get());
+	const FVector Data = Archetype ? Archetype->RelativeRotation.Euler() : FVector::ZeroVector;
+
+	// unset means multiple differing values, so show "Reset to Default" in that case
+	return CachedRotation.IsSet() && CachedRotation.X.GetValue() == Data.X && CachedRotation.Y.GetValue() == Data.Y && CachedRotation.Z.GetValue() == Data.Z ? EVisibility::Hidden : EVisibility::Visible;
+}
+
 FReply FComponentTransformDetails::OnRotationResetClicked()
 {
 	const FText TransactionName = LOCTEXT("ResetRotation", "Reset Rotation");
 	FScopedTransaction Transaction(TransactionName);
 
 	const USceneComponent* Archetype = FGetRootComponentArchetype::Get(SelectedObjects[0].Get());
-	const FRotator Data = Archetype ? Archetype->RelativeRotation : FRotator();
+	const FVector Data = Archetype ? Archetype->RelativeRotation.Euler() : FVector::ZeroVector;
 
-	OnSetRotation(Data.Roll, true, 0);
-	OnSetRotation(Data.Pitch, true, 1);
-	OnSetRotation(Data.Yaw, true, 2);
+	OnSetTransform(ETransformField::Rotation, EAxisList::All, Data, false, true);
 
 	return FReply::Handled();
 }
 
+EVisibility FComponentTransformDetails::GetScaleResetVisibility() const
+{
+	const USceneComponent* Archetype = FGetRootComponentArchetype::Get(SelectedObjects[0].Get());
+	const FVector Data = Archetype ? Archetype->RelativeScale3D : FVector(1.0f);
+
+	// unset means multiple differing values, so show "Reset to Default" in that case
+	return CachedScale.IsSet() && CachedScale.X.GetValue() == Data.X && CachedScale.Y.GetValue() == Data.Y && CachedScale.Z.GetValue() == Data.Z ? EVisibility::Hidden : EVisibility::Visible;
+}
+
 FReply FComponentTransformDetails::OnScaleResetClicked()
 {
 	const FText TransactionName = LOCTEXT("ResetScale", "Reset Scale");
 	FScopedTransaction Transaction(TransactionName);
 
 	const USceneComponent* Archetype = FGetRootComponentArchetype::Get(SelectedObjects[0].Get());
-	const FVector Data = Archetype ? Archetype->RelativeScale3D : FVector();
+	const FVector Data = Archetype ? Archetype->RelativeScale3D : FVector(1.0f);
 
-	ScaleObject(Data[0], 0, false, TransactionName);
-	ScaleObject(Data[1], 1, false, TransactionName);
-	ScaleObject(Data[2], 2, false, TransactionName);
+	OnSetTransform(ETransformField::Scale, EAxisList::All, Data, false, true);
 
 	return FReply::Handled();
 }
@@ -960,24 +848,24 @@ void FComponentTransformDetails::ExtendZScaleContextMenu( FMenuBuilder& MenuBuil
 void FComponentTransformDetails::OnXScaleMirrored()
 {
 	FSlateApplication::Get().ClearKeyboardFocus(EFocusCause::Mouse);
-	ScaleObject( 1.0f, 0, true, LOCTEXT( "MirrorActorScaleX", "Mirror actor scale X" ) );
+	FScopedTransaction Transaction(LOCTEXT("MirrorActorScaleX", "Mirror actor scale X"));
+	OnSetTransform(ETransformField::Scale, EAxisList::X, FVector(1.0f), true, true);
 }
 
 void FComponentTransformDetails::OnYScaleMirrored()
 {
 	FSlateApplication::Get().ClearKeyboardFocus(EFocusCause::Mouse);
-	ScaleObject(1.0f, 1, true, LOCTEXT("MirrorActorScaleY", "Mirror actor scale Y"));
+	FScopedTransaction Transaction(LOCTEXT("MirrorActorScaleY", "Mirror actor scale Y"));
+	OnSetTransform(ETransformField::Scale, EAxisList::X, FVector(1.0f), true, true);
 }
 
 void FComponentTransformDetails::OnZScaleMirrored()
 {
 	FSlateApplication::Get().ClearKeyboardFocus(EFocusCause::Mouse);
-	ScaleObject(1.0f, 2, true, LOCTEXT("MirrorActorScaleZ", "Mirror actor scale Z"));
+	FScopedTransaction Transaction(LOCTEXT("MirrorActorScaleZ", "Mirror actor scale Z"));
+	OnSetTransform(ETransformField::Scale, EAxisList::X, FVector(1.0f), true, true);
 }
 
-/**
- * Cache the entire transform at it is seen by the input boxes so we dont have to iterate over the selected actors multiple times                   
- */
 void FComponentTransformDetails::CacheTransform()
 {
 	FVector CurLoc;
@@ -998,7 +886,8 @@ void FComponentTransformDetails::CacheTransform()
 			if( SceneComponent )
 			{
 				Loc = SceneComponent->RelativeLocation;
-				Rot = (bEditingRotationInUI && !Object->IsTemplate()) ? ObjectToRelativeRotationMap.FindOrAdd(Object) : SceneComponent->RelativeRotation;
+				FRotator* FoundRotator = ObjectToRelativeRotationMap.Find(SceneComponent);
+				Rot = (bEditingRotationInUI && !Object->IsTemplate() && FoundRotator) ? *FoundRotator : SceneComponent->RelativeRotation;
 				Scale = SceneComponent->RelativeScale3D;
 
 				if( ObjectIndex == 0 )
@@ -1043,245 +932,364 @@ void FComponentTransformDetails::CacheTransform()
 	}
 }
 
+FVector FComponentTransformDetails::GetAxisFilteredVector(EAxisList::Type Axis, const FVector& NewValue, const FVector& OldValue)
+{
+	return FVector((Axis & EAxisList::X) ? NewValue.X : OldValue.X,
+		(Axis & EAxisList::Y) ? NewValue.Y : OldValue.Y,
+		(Axis & EAxisList::Z) ? NewValue.Z : OldValue.Z);
+}
 
-void FComponentTransformDetails::OnSetLocation( float NewValue, ETextCommit::Type CommitInfo, int32 Axis )
+void FComponentTransformDetails::OnSetTransform(ETransformField::Type TransformField, EAxisList::Type Axis, FVector NewValue, bool bMirror, bool bCommitted)
 {
+	if (!bCommitted && SelectedObjects.Num() > 1)
+	{
+		// Ignore interactive changes when we have more than one selected object
+		return;
+	}
+
+	FText TransactionText;
+	UProperty* ValueProperty = nullptr;
+	UProperty* AxisProperty = nullptr;
+	
+	switch (TransformField)
+	{
+	case ETransformField::Location:
+		TransactionText = LOCTEXT("OnSetLocation", "Set Location");
+		ValueProperty = FindField<UProperty>(USceneComponent::StaticClass(), GET_MEMBER_NAME_CHECKED(USceneComponent, RelativeLocation));
+		
+		// Only set axis property for single axis set
+		if (Axis == EAxisList::X)
+		{
+			AxisProperty = FindField<UFloatProperty>(TBaseStructure<FVector>::Get(), GET_MEMBER_NAME_CHECKED(FVector, X));
+		}
+		else if (Axis == EAxisList::Y)
+		{
+			AxisProperty = FindField<UFloatProperty>(TBaseStructure<FVector>::Get(), GET_MEMBER_NAME_CHECKED(FVector, Y));
+		}
+		else if (Axis == EAxisList::Z)
+		{
+			AxisProperty = FindField<UFloatProperty>(TBaseStructure<FVector>::Get(), GET_MEMBER_NAME_CHECKED(FVector, Z));
+		}
+		break;
+	case ETransformField::Rotation:
+		TransactionText = LOCTEXT("OnSetRotation", "Set Rotation");
+		ValueProperty = FindField<UProperty>(USceneComponent::StaticClass(), GET_MEMBER_NAME_CHECKED(USceneComponent, RelativeRotation));
+		
+		// Only set axis property for single axis set
+		if (Axis == EAxisList::X)
+		{
+			AxisProperty = FindField<UFloatProperty>(TBaseStructure<FRotator>::Get(), GET_MEMBER_NAME_CHECKED(FRotator, Roll));
+		}
+		else if (Axis == EAxisList::Y)
+		{
+			AxisProperty = FindField<UFloatProperty>(TBaseStructure<FRotator>::Get(), GET_MEMBER_NAME_CHECKED(FRotator, Pitch));
+		}
+		else if (Axis == EAxisList::Z)
+		{
+			AxisProperty = FindField<UFloatProperty>(TBaseStructure<FRotator>::Get(), GET_MEMBER_NAME_CHECKED(FRotator, Yaw));
+		}
+		break;
+	case ETransformField::Scale:
+		TransactionText = LOCTEXT("OnSetScale", "Set Scale");
+		ValueProperty = FindField<UProperty>(USceneComponent::StaticClass(), GET_MEMBER_NAME_CHECKED(USceneComponent, RelativeScale3D));
+
+		// If keep scale is set, don't set axis property
+		if (!bPreserveScaleRatio && Axis == EAxisList::X)
+		{
+			AxisProperty = FindField<UFloatProperty>(TBaseStructure<FVector>::Get(), GET_MEMBER_NAME_CHECKED(FVector, X));
+		}
+		else if (!bPreserveScaleRatio && Axis == EAxisList::Y)
+		{
+			AxisProperty = FindField<UFloatProperty>(TBaseStructure<FVector>::Get(), GET_MEMBER_NAME_CHECKED(FVector, Y));
+		}
+		else if (!bPreserveScaleRatio && Axis == EAxisList::Z)
+		{
+			AxisProperty = FindField<UFloatProperty>(TBaseStructure<FVector>::Get(), GET_MEMBER_NAME_CHECKED(FVector, Z));
+		}
+		break;
+	default:
+		return;
+	}
+
 	bool bBeganTransaction = false;
+	TArray<UObject*> ModifiedObjects;
 
-	for( int32 ObjectIndex = 0; ObjectIndex < SelectedObjects.Num(); ++ObjectIndex )
+	FPropertyChangedEvent PropertyChangedEvent(ValueProperty, !bCommitted ? EPropertyChangeType::Interactive : EPropertyChangeType::ValueSet, (const TArray<const UObject*>*)&ModifiedObjects);
+	FEditPropertyChain PropertyChain;
+
+	if (AxisProperty)
+	{
+		PropertyChain.AddHead(AxisProperty);
+	}
+	PropertyChain.AddHead(ValueProperty);
+	FPropertyChangedChainEvent PropertyChangedChainEvent(PropertyChain, PropertyChangedEvent);
+
+	for (int32 ObjectIndex = 0; ObjectIndex < SelectedObjects.Num(); ++ObjectIndex)
 	{
 		TWeakObjectPtr<UObject> ObjectPtr = SelectedObjects[ObjectIndex];
-		if( ObjectPtr.IsValid() )
+		if (ObjectPtr.IsValid())
 		{
 			UObject* Object = ObjectPtr.Get();
-			USceneComponent* SceneComponent = GetSceneComponentFromDetailsObject( Object );
-			if( SceneComponent )
+			USceneComponent* SceneComponent = GetSceneComponentFromDetailsObject(Object);
+			if (SceneComponent)
 			{
-				FVector OldRelativeLocation = SceneComponent->RelativeLocation;
-				FVector RelativeLocation = OldRelativeLocation;
-				float OldValue = RelativeLocation[Axis];
-				if( OldValue != NewValue )
-				{
-					if( !bBeganTransaction )
-					{
-						// Begin a transaction the first time an actors location is about to change.
-						// NOTE: One transaction per change, not per actor
-						if(Object->IsA<AActor>())
-						{
-							GEditor->BeginTransaction( LOCTEXT( "OnSetLocation", "Set actor location" ) );
-						}
-						else
-						{
-							GEditor->BeginTransaction( LOCTEXT( "OnSetLocation_ComponentDirect", "Modify Component(s)") );
-						}
-						
-						bBeganTransaction = true;
-					}
+				AActor* EditedActor = SceneComponent->GetOwner();
+				const bool bIsEditingTemplateObject = Object->IsTemplate();
+
+				FVector OldComponentValue;
+				FVector NewComponentValue;
 
-					if (Object->HasAnyFlags(RF_DefaultSubObject))
+				switch (TransformField)
+				{
+				case ETransformField::Location:
+					OldComponentValue = SceneComponent->RelativeLocation;
+					break;
+				case ETransformField::Rotation:
+					// Pull from the actual component or from the cache
+					OldComponentValue = SceneComponent->RelativeRotation.Euler();
+					if (bEditingRotationInUI && !bIsEditingTemplateObject && ObjectToRelativeRotationMap.Find(SceneComponent))
 					{
-						// Default subobjects must be included in any undo/redo operations
-						Object->SetFlags(RF_Transactional);
+						OldComponentValue = ObjectToRelativeRotationMap.Find(SceneComponent)->Euler();
 					}
+					break;
+				case ETransformField::Scale:
+					OldComponentValue = SceneComponent->RelativeScale3D;
+					break;
+				}
 
-					// Begin a new movement event which will broadcast delegates before and after the actor moves
-					FScopedObjectMovement ActorMoveEvent( Object );
-
-					FScopedSwitchWorldForObject WorldSwitcher( Object );
-
-					UProperty* RelativeLocationProperty = FindField<UProperty>( USceneComponent::StaticClass(), "RelativeLocation" );
-					Object->PreEditChange( RelativeLocationProperty );
+				// Set the incoming value
+				if (bMirror)
+				{
+					NewComponentValue = -OldComponentValue;
+				}
+				else
+				{
+					NewComponentValue = GetAxisFilteredVector(Axis, NewValue, OldComponentValue);
+				}
 
-					if( NotifyHook )
+				// If we're committing during a rotation edit then we need to force it
+				if (OldComponentValue != NewComponentValue || (bCommitted && bEditingRotationInUI))
+				{
+					if (!bBeganTransaction && bCommitted)
 					{
-						NotifyHook->NotifyPreChange( RelativeLocationProperty );
+						// Begin a transaction the first time an actors rotation is about to change.
+						// NOTE: One transaction per change, not per actor
+						GEditor->BeginTransaction(TransactionText);
+						bBeganTransaction = true;
 					}
 
-					RelativeLocation[Axis] = NewValue;
-					
-					if( SelectedActorInfo.NumSelected == 0 )
-					{
-						// HACK: Set directly if no actors are selected since this causes Rot->Quat->Rot conversion
-						// (recalculates relative rotation even though this is location)
-						SceneComponent->RelativeLocation = RelativeLocation;
-					}
-					else
+					FScopedSwitchWorldForObject WorldSwitcher(Object);
+
+					if (bCommitted)
 					{
-						SceneComponent->SetRelativeLocation( RelativeLocation );
-					}
+						if (!bIsEditingTemplateObject)
+						{
+							// Broadcast the first time an actor is about to move
+							GEditor->BroadcastBeginObjectMovement(*SceneComponent);
+							if (EditedActor && EditedActor->GetRootComponent() == SceneComponent)
+							{
+								GEditor->BroadcastBeginObjectMovement(*EditedActor);
+							}
+						}
 
-					FPropertyChangedEvent PropertyChangedEvent( RelativeLocationProperty );
-					Object->PostEditChangeProperty( PropertyChangedEvent );
+						if (SceneComponent->HasAnyFlags(RF_DefaultSubObject))
+						{
+							// Default subobjects must be included in any undo/redo operations
+							SceneComponent->SetFlags(RF_Transactional);
+						}
 
-					// If it's a template, propagate the change out to any current instances of the object
-					if(SceneComponent->IsTemplate())
-					{
-						TSet<USceneComponent*> UpdatedInstances;
-						FComponentEditorUtils::PropagateDefaultValueChange(SceneComponent, RelativeLocationProperty, OldRelativeLocation, RelativeLocation, UpdatedInstances);
+						// Have to downcast here because of function overloading and inheritance not playing nicely
+						// We don't call PreEditChange for non commit changes because most classes implement the version that doesn't check the interaction type
+						((UObject*)SceneComponent)->PreEditChange(PropertyChain);
+						if (EditedActor && EditedActor->GetRootComponent() == SceneComponent)
+						{
+							((UObject*)EditedActor)->PreEditChange(PropertyChain);
+						}
 					}
 
-					if( NotifyHook )
+					if (NotifyHook)
 					{
-						NotifyHook->NotifyPostChange( PropertyChangedEvent, RelativeLocationProperty );
+						NotifyHook->NotifyPreChange(ValueProperty);
 					}
-				}
-
-			}
-		}
-	}
-
-	if( bBeganTransaction )
-	{
-		GEditor->EndTransaction();
-	}
 
-	CacheTransform();
-
-	GUnrealEd->UpdatePivotLocationForSelection(true);
-	GUnrealEd->SetPivotMovedIndependently(false);
-	GUnrealEd->RedrawLevelEditingViewports();
-}
-
-void FComponentTransformDetails::OnSetRotation( float NewValue, bool bCommitted, int32 Axis )
-{
-	// OnSetRotation is sent from the slider or when the value changes and we dont have slider and the value is being typed.
-	// We should only change the value on commit when it is being typed
-	const bool bAllowSpin = SelectedObjects.Num() == 1;
-
-	if( bAllowSpin || bCommitted )
-	{
-		bool bBeganTransaction = false;
-		for( int32 ObjectIndex = 0; ObjectIndex < SelectedObjects.Num(); ++ObjectIndex )
-		{
-			TWeakObjectPtr<UObject> ObjectPtr = SelectedObjects[ObjectIndex];
-			if( ObjectPtr.IsValid() )
-			{
-				UObject* Object = ObjectPtr.Get();
-				USceneComponent* SceneComponent = GetSceneComponentFromDetailsObject( Object );
-				if( SceneComponent )
-				{
-					const bool bIsEditingTemplateObject = Object->IsTemplate();
-
-					FRotator CurrentComponentRotation = SceneComponent->RelativeRotation; // Intentionally make a copy, we don't want the FRotator& below to directly edit the component's values!
-					FRotator& RelativeRotation = (bEditingRotationInUI && !bIsEditingTemplateObject) ? ObjectToRelativeRotationMap.FindOrAdd(Object) : CurrentComponentRotation;
-					FRotator OldRelativeRotation = RelativeRotation;
-
-					float& ValueToChange = Axis == 0 ? RelativeRotation.Roll : Axis == 1 ? RelativeRotation.Pitch : RelativeRotation.Yaw;
-
-					if( ValueToChange != NewValue )
+					switch (TransformField)
 					{
-						if( !bBeganTransaction && bCommitted )
+					case ETransformField::Location:
 						{
-							// Begin a transaction the first time an actors rotation is about to change.
-							// NOTE: One transaction per change, not per actor
-							if(Object->IsA<AActor>())
+							if (!bIsEditingTemplateObject)
 							{
-								GEditor->BeginTransaction( LOCTEXT( "OnSetRotation", "Set actor rotation" ) );
+								// Update local cache for restoring later
+								ObjectToRelativeRotationMap.FindOrAdd(SceneComponent) = SceneComponent->RelativeRotation;
 							}
-							else
+
+							SceneComponent->SetRelativeLocation(NewComponentValue);
+
+							// Also forcibly set it as the cache may have changed it slightly
+							SceneComponent->RelativeLocation = NewComponentValue;
+
+							// If it's a template, propagate the change out to any current instances of the object
+							if (bIsEditingTemplateObject)
 							{
-								GEditor->BeginTransaction( LOCTEXT( "OnSetRotation_ComponentDirect", "Modify Component(s)") );
+								TSet<USceneComponent*> UpdatedInstances;
+								FComponentEditorUtils::PropagateDefaultValueChange(SceneComponent, ValueProperty, OldComponentValue, NewComponentValue, UpdatedInstances);
 							}
 
+							break;
+						}
+					case ETransformField::Rotation:
+						{
+							FRotator NewRotation = FRotator::MakeFromEuler(NewComponentValue);
+
 							if (!bIsEditingTemplateObject)
 							{
-								// Broadcast the first time an actor is about to move
-								GEditor->BroadcastBeginObjectMovement( *Object );
+								// Update local cache for restoring later
+								ObjectToRelativeRotationMap.FindOrAdd(SceneComponent) = NewRotation;
 							}
 
-							bBeganTransaction = true;
-						}
+							SceneComponent->SetRelativeRotation(NewRotation);
 
-						FScopedSwitchWorldForObject WorldSwitcher( Object );
+							// Also forcibly set it as the cache may have changed it slightly
+							SceneComponent->RelativeRotation = NewRotation;
 
-						UProperty* RelativeRotationProperty = FindField<UProperty>( USceneComponent::StaticClass(), "RelativeRotation" );
-						if( bCommitted && !bEditingRotationInUI )
-						{
-							if (Object->HasAnyFlags(RF_DefaultSubObject))
+							// If it's a template, propagate the change out to any current instances of the object
+							if (bIsEditingTemplateObject)
 							{
-								// Default subobjects must be included in any undo/redo operations
-								Object->SetFlags(RF_Transactional);
+								TSet<USceneComponent*> UpdatedInstances;
+								FComponentEditorUtils::PropagateDefaultValueChange(SceneComponent, ValueProperty, FRotator::MakeFromEuler(OldComponentValue), NewRotation, UpdatedInstances);
 							}
 
-							Object->PreEditChange( RelativeRotationProperty );
+							break;
 						}
-
-						if( NotifyHook )
+					case ETransformField::Scale:
 						{
-							NotifyHook->NotifyPreChange( RelativeRotationProperty );
-						}
+							if (bPreserveScaleRatio)
+							{
+								// If we set a single axis, scale the others
+								float Ratio = 0.0f;
+
+								switch (Axis)
+								{
+								case EAxisList::X:
+									// Account for the previous scale being zero.  Just set to the new value in that case?
+									Ratio = OldComponentValue.X == 0.0f ? NewComponentValue.X : NewComponentValue.X / OldComponentValue.X;
+									NewComponentValue.Y *= Ratio;
+									NewComponentValue.Z *= Ratio;
+									break;
+								case EAxisList::Y:
+									Ratio = OldComponentValue.Y == 0.0f ? NewComponentValue.Y : NewComponentValue.Y / OldComponentValue.Y;
+									NewComponentValue.X *= Ratio;
+									NewComponentValue.Z *= Ratio;
+									break;
+								case EAxisList::Z:
+									Ratio = OldComponentValue.Z == 0.0f ? NewComponentValue.Z : NewComponentValue.Z / OldComponentValue.Z;
+									NewComponentValue.X *= Ratio;
+									NewComponentValue.Y *= Ratio;
+								default:
+									// Do nothing, this set multiple axis at once
+									break;
+								}
+							}
 
-						ValueToChange = NewValue;
+							SceneComponent->SetRelativeScale3D(NewComponentValue);
 
-						if( !bIsEditingTemplateObject && SelectedActorInfo.NumSelected != 0 )
-						{
-							SceneComponent->SetRelativeRotation( RelativeRotation );
+							// If it's a template, propagate the change out to any current instances of the object
+							if (bIsEditingTemplateObject)
+							{
+								TSet<USceneComponent*> UpdatedInstances;
+								FComponentEditorUtils::PropagateDefaultValueChange(SceneComponent, ValueProperty, OldComponentValue, NewComponentValue, UpdatedInstances);
+							}
+
+							break;
 						}
+					}
 
-						// HACK: Set directly since to avoid Rot->Quat->Rot conversion issues
-						// (functions recalculate relative rotation from quat which can give an equivalent but different value than the user typed)
-						SceneComponent->RelativeRotation = RelativeRotation;
+					ModifiedObjects.Add(Object);
+				}
+			}
+		}
+	}
 
-						AActor* EditedActor = Cast<AActor>( Object );
-						if( !EditedActor && SceneComponent )
-						{
-							EditedActor = SceneComponent->GetOwner();
-						}
+	if (ModifiedObjects.Num())
+	{
+		for (UObject* Object : ModifiedObjects)
+		{
+			USceneComponent* SceneComponent = GetSceneComponentFromDetailsObject(Object);
+			USceneComponent* OldSceneComponent = SceneComponent;
 
-						if( EditedActor && !bIsEditingTemplateObject )
-						{
-							EditedActor->ReregisterAllComponents();
-						}
+			if (SceneComponent)
+			{
+				AActor* EditedActor = SceneComponent->GetOwner();
+				FString SceneComponentPath = SceneComponent->GetPathName(EditedActor);
+				
+				if (bCommitted)
+				{
+					// This can invalidate OldSceneComponent
+					// We don't call PostEditChange for non commit changes because most classes implement the version that doesn't check the interaction type
+					OldSceneComponent->PostEditChangeChainProperty(PropertyChangedChainEvent);
 
-						// If it's a template, propagate the change out to any current instances of the object
-						if(SceneComponent->IsTemplate())
-						{
-							TSet<USceneComponent*> UpdatedInstances;
-							FComponentEditorUtils::PropagateDefaultValueChange(SceneComponent, RelativeRotationProperty, OldRelativeRotation, RelativeRotation, UpdatedInstances);
-						}
+					SceneComponent = FindObject<USceneComponent>(EditedActor, *SceneComponentPath);
 
-						FPropertyChangedEvent PropertyChangedEvent( RelativeRotationProperty, !bCommitted && bEditingRotationInUI ? EPropertyChangeType::Interactive : EPropertyChangeType::ValueSet );
+					if (EditedActor && EditedActor->GetRootComponent() == SceneComponent)
+					{
+						EditedActor->PostEditChangeChainProperty(PropertyChangedChainEvent);
+						SceneComponent = FindObject<USceneComponent>(EditedActor, *SceneComponentPath);
+					}
+				}
 
-						if( NotifyHook )
-						{
-							NotifyHook->NotifyPostChange( PropertyChangedEvent, RelativeRotationProperty );
-						}
+				if (!Object->IsTemplate())
+				{
+					if (TransformField == ETransformField::Rotation || TransformField == ETransformField::Location)
+					{
+						FRotator* FoundRotator = ObjectToRelativeRotationMap.Find(OldSceneComponent);
 
-						if( bCommitted )
+						if (FoundRotator)
 						{
-							if( !bEditingRotationInUI )
-							{
-								Object->PostEditChangeProperty( PropertyChangedEvent );	
-							}
-					
-							if (!bIsEditingTemplateObject)
+							FQuat OldQuat = FoundRotator->GetDenormalized().Quaternion();
+							FQuat NewQuat = SceneComponent->RelativeRotation.GetDenormalized().Quaternion();
+
+							if (OldQuat.Equals(NewQuat))
 							{
-								// The actor is done moving
-								GEditor->BroadcastEndObjectMovement( *Object );
+								// Need to restore the manually set rotation as it was modified by quat conversion
+								SceneComponent->RelativeRotation = *FoundRotator;
 							}
 						}
 					}
+
+					if (bCommitted)
+					{
+						// Broadcast the first time an actor is about to move
+						GEditor->BroadcastBeginObjectMovement(*SceneComponent);
+						if (EditedActor && EditedActor->GetRootComponent() == SceneComponent)
+						{
+							GEditor->BroadcastBeginObjectMovement(*EditedActor);
+						}
+					}
 				}
 			}
 		}
 
-		if( bCommitted && bBeganTransaction )
+		if (NotifyHook)
 		{
-			GEditor->EndTransaction();
+			NotifyHook->NotifyPostChange(PropertyChangedEvent, ValueProperty);
 		}
+	}
 
-		GUnrealEd->UpdatePivotLocationForSelection();
-		GUnrealEd->SetPivotMovedIndependently(false);
-		// Redraw
-		GUnrealEd->RedrawLevelEditingViewports();
+	if (bCommitted && bBeganTransaction)
+	{
+		GEditor->EndTransaction();
+		CacheTransform();
 	}
+
+	GUnrealEd->UpdatePivotLocationForSelection();
+	GUnrealEd->SetPivotMovedIndependently(false);
+	// Redraw
+	GUnrealEd->RedrawLevelEditingViewports();
 }
 
-void FComponentTransformDetails::OnRotationCommitted(float NewValue, ETextCommit::Type CommitInfo, int32 Axis)
+void FComponentTransformDetails::OnSetTransformAxis(float NewValue, ETextCommit::Type CommitInfo, ETransformField::Type TransformField, EAxisList::Type Axis, bool bCommitted)
 {
-	OnSetRotation(NewValue, true, Axis);
-
-	CacheTransform();
+	FVector NewVector = GetAxisFilteredVector(Axis, FVector(NewValue), FVector::ZeroVector);
+	OnSetTransform(TransformField, Axis, NewVector, false, bCommitted);
 }
 
 void FComponentTransformDetails::OnBeginRotatonSlider()
@@ -1318,17 +1326,17 @@ void FComponentTransformDetails::OnBeginRotatonSlider()
 			{
 				FScopedSwitchWorldForObject WorldSwitcher( Object );
 				
-				if (Object->HasAnyFlags(RF_DefaultSubObject))
+				if (SceneComponent->HasAnyFlags(RF_DefaultSubObject))
 				{
 					// Default subobjects must be included in any undo/redo operations
-					Object->SetFlags(RF_Transactional);
+					SceneComponent->SetFlags(RF_Transactional);
 				}
 
-				UProperty* RelativeRotationProperty = FindField<UProperty>( USceneComponent::StaticClass(), "RelativeRotation" );
-				Object->PreEditChange( RelativeRotationProperty );
+				// Call modify but not PreEdit, we don't do the proper "Edit" until it's committed
+				SceneComponent->Modify();
 
 				// Add/update cached rotation value prior to slider interaction
-				ObjectToRelativeRotationMap.FindOrAdd(Object) = SceneComponent->RelativeRotation;
+				ObjectToRelativeRotationMap.FindOrAdd(SceneComponent) = SceneComponent->RelativeRotation;
 			}
 		}
 	}
@@ -1342,179 +1350,9 @@ void FComponentTransformDetails::OnBeginRotatonSlider()
 
 void FComponentTransformDetails::OnEndRotationSlider(float NewValue)
 {
+	// Commit gets called right before this, only need to end the transaction
 	bEditingRotationInUI = false;
-
-	for( int32 ObjectIndex = 0; ObjectIndex < SelectedObjects.Num(); ++ObjectIndex )
-	{
-		TWeakObjectPtr<UObject> ObjectPtr = SelectedObjects[ObjectIndex];
-		if( ObjectPtr.IsValid() )
-		{
-			UObject* Object = ObjectPtr.Get();
-			USceneComponent* SceneComponent = GetSceneComponentFromDetailsObject( Object );
-			if( SceneComponent )
-			{
-				FScopedSwitchWorldForObject WorldSwitcher( Object );
-
-				UProperty* RelativeRotationProperty = FindField<UProperty>( USceneComponent::StaticClass(), "RelativeRotation" );
-				FPropertyChangedEvent PropertyChangedEvent( RelativeRotationProperty );
-				Object->PostEditChangeProperty( PropertyChangedEvent );
-			}
-		}
-	}
-
 	GEditor->EndTransaction();
-
-	// Redraw
-	GUnrealEd->RedrawLevelEditingViewports();
-}
-
-void FComponentTransformDetails::OnSetScale( const float NewValue, ETextCommit::Type CommitInfo, int32 Axis )
-{
-	ScaleObject( NewValue, Axis, false, LOCTEXT( "OnSetScale", "Set actor scale" ) );
-}
-
-void FComponentTransformDetails::ScaleObject( float NewValue, int32 Axis, bool bMirror, const FText& TransactionSessionName )
-{
-	UProperty* RelativeScale3DProperty = FindField<UProperty>( USceneComponent::StaticClass(), "RelativeScale3D" );
-
-	bool bBeganTransaction = false;
-	for( int32 ObjectIndex = 0; ObjectIndex < SelectedObjects.Num(); ++ObjectIndex )
-	{
-		TWeakObjectPtr<UObject> ObjectPtr = SelectedObjects[ObjectIndex];
-		if( ObjectPtr.IsValid() )
-		{
-			UObject* Object = ObjectPtr.Get();
-			USceneComponent* SceneComponent = GetSceneComponentFromDetailsObject( Object );
-			if( SceneComponent )
-			{
-				FVector OldRelativeScale = SceneComponent->RelativeScale3D;
-				FVector RelativeScale = OldRelativeScale;
-				if( bMirror )
-				{
-					NewValue = -RelativeScale[Axis];
-				}
-				float OldValue = RelativeScale[Axis];
-				if( OldValue != NewValue )
-				{
-					if( !bBeganTransaction )
-					{
-						// Begin a transaction the first time an actors scale is about to change.
-						// NOTE: One transaction per change, not per actor
-						GEditor->BeginTransaction( TransactionSessionName );
-						bBeganTransaction = true;
-					}
-
-					FScopedSwitchWorldForObject WorldSwitcher( Object );
-					
-					if (Object->HasAnyFlags(RF_DefaultSubObject))
-					{
-						// Default subobjects must be included in any undo/redo operations
-						Object->SetFlags(RF_Transactional);
-					}
-
-					// Begin a new movement event which will broadcast delegates before and after the actor moves
-					FScopedObjectMovement ActorMoveEvent( Object );
-
-					Object->PreEditChange( RelativeScale3DProperty );
-
-					if( NotifyHook )
-					{
-						NotifyHook->NotifyPreChange( RelativeScale3DProperty );
-					}
-
-					// Set the new value for the corresponding axis
-					RelativeScale[Axis] = NewValue;
-
-					if( bPreserveScaleRatio )
-					{
-						// Account for the previous scale being zero.  Just set to the new value in that case?
-						float Ratio = OldValue == 0.0f ? NewValue : NewValue/OldValue;
-
-						// Change values on axes besides the one being directly changed
-						switch( Axis )
-						{
-						case 0:
-							RelativeScale.Y *= Ratio;
-							RelativeScale.Z *= Ratio;
-							break;
-						case 1:
-							RelativeScale.X *= Ratio;
-							RelativeScale.Z *= Ratio;
-							break;
-						case 2:
-							RelativeScale.X *= Ratio;
-							RelativeScale.Y *= Ratio;
-						}
-					}
-
-					SceneComponent->SetRelativeScale3D( RelativeScale );
-
-					// Build property chain so the actor knows whether we changed the X, Y or Z
-					FEditPropertyChain PropertyChain;
-
-					if (!bPreserveScaleRatio)
-					{
-						UStruct* VectorStruct = FindObjectChecked<UScriptStruct>(UObject::StaticClass()->GetOutermost(), TEXT("Vector"), false);
-
-						UProperty* VectorValueProperty = NULL;
-						switch( Axis )
-						{
-						case 0:
-							VectorValueProperty = FindField<UFloatProperty>(VectorStruct, TEXT("X"));
-							break;
-						case 1:
-							VectorValueProperty = FindField<UFloatProperty>(VectorStruct, TEXT("Y"));
-							break;
-						case 2:
-							VectorValueProperty = FindField<UFloatProperty>(VectorStruct, TEXT("Z"));
-						}
-
-						PropertyChain.AddHead(VectorValueProperty);
-					}
-					PropertyChain.AddHead(RelativeScale3DProperty);
-
-					FPropertyChangedEvent PropertyChangedEvent(RelativeScale3DProperty, EPropertyChangeType::ValueSet);
-					FPropertyChangedChainEvent PropertyChangedChainEvent(PropertyChain, PropertyChangedEvent);
-					Object->PostEditChangeChainProperty( PropertyChangedChainEvent );
-
-					// For backwards compatibility, as for some reason PostEditChangeChainProperty calls PostEditChangeProperty with the property set to "X" not "RelativeScale3D"
-					// (it does that for all vector properties, and I don't want to be the one to change that)
-					if (!bPreserveScaleRatio)
-					{
-						Object->PostEditChangeProperty( PropertyChangedEvent );
-					}
-					else
-					{
-						// If the other scale values have been updated, make sure we update the transform now (as the tick will be too late)
-						// so they appear correct when their EditedText is fetched from the delegate.
-						CacheTransform();
-					}
-
-					// If it's a template, propagate the change out to any current instances of the object
-					if(SceneComponent->IsTemplate())
-					{
-						TSet<USceneComponent*> UpdatedInstances;
-						FComponentEditorUtils::PropagateDefaultValueChange(SceneComponent, RelativeScale3DProperty, OldRelativeScale, RelativeScale, UpdatedInstances);
-					}
-
-					if( NotifyHook )
-					{
-						NotifyHook->NotifyPostChange( PropertyChangedEvent, RelativeScale3DProperty );
-					}
-				}
-			}
-		}
-	}
-
-	if( bBeganTransaction )
-	{
-		GEditor->EndTransaction();
-	}
-
-	CacheTransform();
-
-	// Redraw
-	GUnrealEd->RedrawLevelEditingViewports();
 }
 
 #undef LOCTEXT_NAMESPACE
diff --git a/Engine/Source/Editor/DetailCustomizations/Private/ComponentTransformDetails.h b/Engine/Source/Editor/DetailCustomizations/Private/ComponentTransformDetails.h
--- a/Engine/Source/Editor/DetailCustomizations/Private/ComponentTransformDetails.h
+++ b/Engine/Source/Editor/DetailCustomizations/Private/ComponentTransformDetails.h
@@ -56,30 +56,30 @@ private:
 	/** @return Whether the transform details panel should be enabled (editable) or not (read-only / greyed out) */
 	bool GetIsEnabled() const;
 
-	/**
-	 * Sets the selected object(s) translation (relative)
-	 *
-	 * @param NewValue The new translation value
-	 */
-	void OnSetLocation( float NewValue, ETextCommit::Type CommitInfo, int32 Axis );
+	/** Sets a vector based on two source vectors and an axis list */
+	FVector GetAxisFilteredVector(EAxisList::Type Axis, const FVector& NewValue, const FVector& OldValue);
 
 	/**
-	 * Sets the selected object(s) rotation (relative)
+	 * Sets the selected object(s) axis to passed in value
 	 *
-	 * @param NewValue		The new rotation value for the axis
-	 * @param bCommittted	true if the value was committed, false is the value comes from the slider (int32 due to delegate restrictions)
-	 * @param Axis			The axis of rotation which changed.  0 = Roll, 1 = Pitch, 2 = Yaw
+	 * @param TransformField	The field (location/rotation/scale) to modify
+	 * @param Axis				Bitfield of which axis to set, can be multiple
+	 * @param NewValue			The new vector values, it only uses the ones with specified axis
+	 * @param bMirror			If true, set the value to it's inverse instead of using NewValue
+	 * @param bCommittted		True if the value was committed, false is the value comes from the slider
 	 */
-	void OnSetRotation( float NewValue, bool bCommitted, int32 Axis );
+	void OnSetTransform(ETransformField::Type TransformField, EAxisList::Type Axis, FVector NewValue, bool bMirror, bool bCommitted);
 
 	/**
-	 * Sets the selected object(s) rotation (relative)
+	 * Sets a single axis value, called from UI
 	 *
+	 * @param TransformField	The field (location/rotation/scale) to modify
+	 * @param NewValue		The new translation value
 	 * @param CommitInfo	Whether or not this was committed from pressing enter or losing focus
-	 * @param Axis			The axis of rotation which changed.  0 = Roll, 1 = Pitch, 2 = Yaw
-	 * @param NewValue		The new rotation value for the axis
+	 * @param Axis				Bitfield of which axis to set, can be multiple
+	 * @param bCommittted	true if the value was committed, false is the value comes from the slider
 	 */
-	void OnRotationCommitted( float NewValue, ETextCommit::Type CommitInfo, int32 Axis );
+	void OnSetTransformAxis(float NewValue, ETextCommit::Type CommitInfo, ETransformField::Type TransformField, EAxisList::Type Axis, bool bCommitted);
 
 	/**
 	 * Called when the one of the axis sliders for object rotation begins to change for the first time 
@@ -91,24 +91,6 @@ private:
 	 */
 	void OnEndRotationSlider(float NewValue);
 
-	/**
-	 * Sets the selected object(s) scale (relative)
-	 *
-	 * @param NewValue The new scale value
-	 */
-	void OnSetScale( float NewValue, ETextCommit::Type CommitInfo, int32 Axis );
-
-	/**
-	 * Sets the selected object(s) scale or mirrors them (relative)
-	 * Helper function for OnSetScale and On*ScaleMirrored.
-	 *
-	 * @param NewValue					The new scale value - unused if mirroring
-	 * @param Axis						The axis we are manipulating
-	 * @param bMirror					Whether to mirror the axis
-	 * @param TransactionSessionName	the name of the undo/redo transaction
-	 */
-	void ScaleObject( float NewValue, int32 Axis, bool bMirror, const FText& TransactionSessionName );
-
 	/** @return Icon to use in the preserve scale ratio check box */
 	const FSlateBrush* GetPreserveScaleRatioImage() const;
 
@@ -151,46 +133,15 @@ private:
 	 */
 	FText GetScaleText() const;
 
-	/**
-	 * Enables or disables absolute location
-	 *
-	 * @param bEnable True to enable, false to disable (use relative location)
-	 */
-	void OnToggleAbsoluteLocation( bool bEnable );
-
-	/**
-	 * Enables or disables absolute rotation
-	 *
-	 * @param bEnable True to enable, false to disable (use relative rotation)
-	 */
-	void OnToggleAbsoluteRotation( bool bEnable );
-
-	/**
-	 * Enables or disables absolute scale
-	 *
-	 * @param bEnable True to enable, false to disable (use relative scale)
-	 */
-	void OnToggleAbsoluteScale( bool bEnable );
-
 	/**
 	 * Sets relative transform on the specified field
 	 *
 	 * @param The field that should be set to relative
 	 */
-	void OnSetRelativeTransform( ETransformField::Type TransformField );
+	void OnSetAbsoluteTransform( ETransformField::Type TransformField, bool bAbsoluteEnabled);
 
-	/**
-	 * Sets world transform on the specified field
-	 *
-	 * @param The field that should be set to world
-	 */
-	void OnSetWorldTransform( ETransformField::Type TransformField );
-
-	/** @return true if relative transform is enabled for the specified field */
-	bool IsRelativeTransformChecked( ETransformField::Type TransformField ) const;
-
-	/** @return true if world transform is enabled for the specified field */
-	bool IsWorldTransformChecked( ETransformField::Type TransformField ) const;
+	/** @return true if Absolute flag of transform type matches passed in bCheckAbsolute*/
+	bool IsAbsoluteTransformChecked( ETransformField::Type TransformField, bool bAbsoluteEnabled=true) const;
 
 	/** @return true of copy is enabled for the specified field */
 	bool OnCanCopy( ETransformField::Type TransformField ) const;
@@ -245,11 +196,7 @@ private:
 	/** @return The Z component of location */
 	TOptional<float> GetLocationZ() const { return CachedLocation.Z; }
 	/** @return The visibility of the "Reset to Default" button for the location component */
-	EVisibility GetLocationResetVisibility() const
-	{
-		// unset means multiple differing values, so show "Reset to Default" in that case
-		return CachedLocation.IsSet() && CachedLocation.X.GetValue() == 0.0f && CachedLocation.Y.GetValue() == 0.0f && CachedLocation.Z.GetValue() == 0.0f ? EVisibility::Hidden : EVisibility::Visible;
-	}
+	EVisibility GetLocationResetVisibility() const;
 
 	/** @return The X component of rotation */
 	TOptional<float> GetRotationX() const { return CachedRotation.X; }
@@ -258,11 +205,7 @@ private:
 	/** @return The Z component of rotation */
 	TOptional<float> GetRotationZ() const { return CachedRotation.Z; }
 	/** @return The visibility of the "Reset to Default" button for the rotation component */
-	EVisibility GetRotationResetVisibility() const
-	{
-		// unset means multiple differing values, so show "Reset to Default" in that case
-		return CachedRotation.IsSet() && CachedRotation.X.GetValue() == 0.0f && CachedRotation.Y.GetValue() == 0.0f && CachedRotation.Z.GetValue() == 0.0f ? EVisibility::Hidden : EVisibility::Visible;
-	}
+	EVisibility GetRotationResetVisibility() const;
 
 	/** @return The X component of scale */
 	TOptional<float> GetScaleX() const { return CachedScale.X; }
@@ -271,11 +214,7 @@ private:
 	/** @return The Z component of scale */
 	TOptional<float> GetScaleZ() const { return CachedScale.Z; }
 	/** @return The visibility of the "Reset to Default" button for the scale component */
-	EVisibility GetScaleResetVisibility() const
-	{
-		// unset means multiple differing values, so show "Reset to Default" in that case
-		return CachedScale.IsSet() && CachedScale.X.GetValue() == 1.0f && CachedScale.Y.GetValue() == 1.0f && CachedScale.Z.GetValue() == 1.0f ? EVisibility::Hidden : EVisibility::Visible;
-	}
+	EVisibility GetScaleResetVisibility() const;
 
 	/** Cache a single unit to display all location comonents in */
 	void CacheCommonLocationUnits();
diff --git a/Engine/Source/Editor/DetailCustomizations/Private/FbxImportUIDetails.cpp b/Engine/Source/Editor/DetailCustomizations/Private/FbxImportUIDetails.cpp
--- a/Engine/Source/Editor/DetailCustomizations/Private/FbxImportUIDetails.cpp
+++ b/Engine/Source/Editor/DetailCustomizations/Private/FbxImportUIDetails.cpp
@@ -203,7 +203,7 @@ void FFbxImportUIDetails::CustomizeDetails( IDetailLayoutBuilder& DetailBuilder
 	for(TSharedPtr<IPropertyHandle> Handle : ExtraProperties)
 	{
 		FString ImportTypeMetaData = Handle->GetMetaData(TEXT("ImportType"));
-		FString CategoryMetaData = Handle->GetMetaData(TEXT("ImportCategory"));
+		const FString& CategoryMetaData = Handle->GetMetaData(TEXT("ImportCategory"));
 		if(IsImportTypeMetaDataValid(ImportType, ImportTypeMetaData))
 		{
 			// Decide on category
@@ -272,7 +272,7 @@ void FFbxImportUIDetails::CustomizeDetails( IDetailLayoutBuilder& DetailBuilder
 
 		for(TSharedPtr<IPropertyHandle> Handle : ExtraProperties)
 		{
-			FString CategoryMetaData = Handle->GetMetaData(TEXT("ImportCategory"));
+			const FString& CategoryMetaData = Handle->GetMetaData(TEXT("ImportCategory"));
 			if(Handle->GetProperty()->GetOuter() == UFbxAnimSequenceImportData::StaticClass()
 			   && CategoryMetaData.IsEmpty())
 			{
diff --git a/Engine/Source/Editor/DetailCustomizations/Private/GuidStructCustomization.cpp b/Engine/Source/Editor/DetailCustomizations/Private/GuidStructCustomization.cpp
--- a/Engine/Source/Editor/DetailCustomizations/Private/GuidStructCustomization.cpp
+++ b/Engine/Source/Editor/DetailCustomizations/Private/GuidStructCustomization.cpp
@@ -100,8 +100,9 @@ void FGuidStructCustomization::SetGuidValue( const FGuid& Guid )
 {
 	for (uint32 ChildIndex = 0; ChildIndex < 4; ++ChildIndex)
 	{
+		// Mark first 3 as interactive changes so the post edit doesn't go off until the last one
 		TSharedRef<IPropertyHandle> ChildHandle = PropertyHandle->GetChildHandle(ChildIndex).ToSharedRef();
-		ChildHandle->SetValue((int32)Guid[ChildIndex]);
+		ChildHandle->SetValue((int32)Guid[ChildIndex], ChildIndex != 3 ? EPropertyValueSetFlags::InteractiveChange : EPropertyValueSetFlags::DefaultFlags);
 	}
 }
 
diff --git a/Engine/Source/Editor/DetailCustomizations/Private/MarginCustomization.cpp b/Engine/Source/Editor/DetailCustomizations/Private/MarginCustomization.cpp
--- a/Engine/Source/Editor/DetailCustomizations/Private/MarginCustomization.cpp
+++ b/Engine/Source/Editor/DetailCustomizations/Private/MarginCustomization.cpp
@@ -20,7 +20,7 @@ void FMarginStructCustomization::CustomizeHeader( TSharedRef<class IPropertyHand
 {
 	StructPropertyHandle = InStructPropertyHandle;
 
-	const FString UVSpaceString( StructPropertyHandle->GetProperty()->GetMetaData( TEXT( "UVSpace" ) ) );
+	const FString& UVSpaceString( StructPropertyHandle->GetProperty()->GetMetaData( TEXT( "UVSpace" ) ) );
 	bIsMarginUsingUVSpace = UVSpaceString.Len() > 0 && UVSpaceString == TEXT( "true" );
 
 	uint32 NumChildren;
diff --git a/Engine/Source/Editor/DetailCustomizations/Private/MeshMergingSettingsCustomization.cpp b/Engine/Source/Editor/DetailCustomizations/Private/MeshMergingSettingsCustomization.cpp
--- a/Engine/Source/Editor/DetailCustomizations/Private/MeshMergingSettingsCustomization.cpp
+++ b/Engine/Source/Editor/DetailCustomizations/Private/MeshMergingSettingsCustomization.cpp
@@ -30,7 +30,7 @@ void FMeshMergingSettingsObjectCustomization::CustomizeDetails(IDetailLayoutBuil
 	FName CategoryMetaData("Category");
 	for (TSharedRef<IPropertyHandle> Property: SimpleDefaultProperties)
 	{
-		const FString CategoryName = Property->GetMetaData(CategoryMetaData);
+		const FString& CategoryName = Property->GetMetaData(CategoryMetaData);
 
 		IDetailCategoryBuilder& CategoryBuilder = LayoutBuilder.EditCategory(*CategoryName);
 		IDetailPropertyRow& PropertyRow = CategoryBuilder.AddProperty(Property);
@@ -59,7 +59,7 @@ void FMeshMergingSettingsObjectCustomization::CustomizeDetails(IDetailLayoutBuil
 
 	for (TSharedRef<IPropertyHandle> Property : SimpleDefaultProperties)
 	{
-		const FString CategoryName = Property->GetMetaData(CategoryMetaData);
+		const FString& CategoryName = Property->GetMetaData(CategoryMetaData);
 
 		IDetailCategoryBuilder& CategoryBuilder = LayoutBuilder.EditCategory(*CategoryName);
 		IDetailPropertyRow& PropertyRow = CategoryBuilder.AddProperty(Property);
diff --git a/Engine/Source/Editor/DetailCustomizations/Private/SoftObjectPathCustomization.cpp b/Engine/Source/Editor/DetailCustomizations/Private/SoftObjectPathCustomization.cpp
--- a/Engine/Source/Editor/DetailCustomizations/Private/SoftObjectPathCustomization.cpp
+++ b/Engine/Source/Editor/DetailCustomizations/Private/SoftObjectPathCustomization.cpp
@@ -11,7 +11,7 @@ void FSoftObjectPathCustomization::CustomizeHeader( TSharedRef<IPropertyHandle>
 {
 	StructPropertyHandle = InStructPropertyHandle;
 
-	FString ClassFilterString = StructPropertyHandle->GetMetaData("AllowedClasses");
+	const FString& ClassFilterString = StructPropertyHandle->GetMetaData("AllowedClasses");
 	if( !ClassFilterString.IsEmpty() )
 	{
 		TArray<FString> CustomClassFilterNames;
diff --git a/Engine/Source/Editor/DetailCustomizations/Private/TextureLODSettingsDetails.cpp b/Engine/Source/Editor/DetailCustomizations/Private/TextureLODSettingsDetails.cpp
--- a/Engine/Source/Editor/DetailCustomizations/Private/TextureLODSettingsDetails.cpp
+++ b/Engine/Source/Editor/DetailCustomizations/Private/TextureLODSettingsDetails.cpp
@@ -82,7 +82,7 @@ void FTextureLODGroupLayout::AddToAvailableMipGenSettings(TextureMipGenSettings
 void FTextureLODGroupLayout::GenerateHeaderRowContent(FDetailWidgetRow& NodeRow)
 {
 	UEnum* TextureGroupEnum = FindObject<UEnum>(NULL, TEXT("/Script/Engine.TextureGroup"));
-	const FString LODGroupName = TextureGroupEnum->GetMetaData(TEXT("DisplayName"), LodGroup->Group);
+	const FString& LODGroupName = TextureGroupEnum->GetMetaData(TEXT("DisplayName"), LodGroup->Group);
 
 	NodeRow.NameContent()
 		[
@@ -367,7 +367,7 @@ FText FTextureLODGroupLayout::GetMipFilterComboBoxContent() const
 TSharedRef<SWidget> FTextureLODGroupLayout::MakeMipGenSettingsComboWidget(TSharedPtr<TextureMipGenSettings> InItem)
 {
 	UEnum* TextureGroupEnum = FindObject<UEnum>(NULL, TEXT("/Script/Engine.TextureMipGenSettings"));
-	const FString MipGenSettingsName = TextureGroupEnum->GetMetaData(TEXT("DisplayName"), *InItem.Get());
+	const FString& MipGenSettingsName = TextureGroupEnum->GetMetaData(TEXT("DisplayName"), *InItem.Get());
 
 	return SNew(STextBlock).Text(FText::FromString(MipGenSettingsName)).Font(IDetailLayoutBuilder::GetDetailFont());
 }
@@ -390,7 +390,7 @@ FText FTextureLODGroupLayout::GetMipGenSettingsComboBoxToolTip() const
 FText FTextureLODGroupLayout::GetMipGenSettingsComboBoxContent() const
 {
 	UEnum* TextureGroupEnum = FindObject<UEnum>(NULL, TEXT("/Script/Engine.TextureMipGenSettings"));
-	const FString MipGenSettingsName = TextureGroupEnum->GetMetaData(TEXT("DisplayName"), LodGroup->MipGenSettings);
+	const FString& MipGenSettingsName = TextureGroupEnum->GetMetaData(TEXT("DisplayName"), LodGroup->MipGenSettings);
 	
 	return FText::FromString(MipGenSettingsName);
 }
diff --git a/Engine/Source/Editor/EditorStyle/Private/SlateEditorStyle.cpp b/Engine/Source/Editor/EditorStyle/Private/SlateEditorStyle.cpp
--- a/Engine/Source/Editor/EditorStyle/Private/SlateEditorStyle.cpp
+++ b/Engine/Source/Editor/EditorStyle/Private/SlateEditorStyle.cpp
@@ -747,7 +747,9 @@ void FSlateEditorStyle::FStyle::SetupGeneralStyles()
 	Set( "MarqueeSelection", new BORDER_BRUSH( "Old/DashedBorder", FMargin(6.0f/32.0f) ) );
 
 	Set( "GenericLock", new IMAGE_BRUSH( "Icons/padlock_locked_16x", Icon16x16 ) );
+	Set( "GenericLock.Small", new IMAGE_BRUSH( "Icons/padlock_locked_16x", Icon16x16 ) );
 	Set( "GenericUnlock", new IMAGE_BRUSH( "Icons/padlock_unlocked_16x", Icon16x16 ) );
+	Set( "GenericUnlock.Small", new IMAGE_BRUSH( "Icons/padlock_unlocked_16x", Icon16x16 ) );
 
 	Set( "GenericPlay", new IMAGE_BRUSH( "Icons/generic_play_16x", Icon16x16 ) );
 	Set( "GenericPause", new IMAGE_BRUSH( "Icons/generic_pause_16x", Icon16x16 ) );
@@ -5526,14 +5528,29 @@ void FSlateEditorStyle::FStyle::SetupPersonaStyle()
 
 	}
 
+	//Find Results
+	{
+		Set("FindResults.FindInBlueprints", FTextBlockStyle(NormalText)
+			.SetFont(TTF_FONT("Fonts/FontAwesome", 10))
+			.SetColorAndOpacity(FLinearColor(0.f, 0.f, 0.f))
+		);
+
+		Set("FindResults.LockButton_Locked", new IMAGE_BRUSH("Icons/padlock_locked_16x", Icon16x16));
+		Set("FindResults.LockButton_Unlocked", new IMAGE_BRUSH("Icons/padlock_unlocked_16x", Icon16x16));
+	}
+
 	//Blueprint Diff
 	{
-		Set( "BlueprintDif.HasGraph", new IMAGE_BRUSH( "/Icons/blueprint_Dif_has_graph_8x", Icon8x8 )  );
-		Set( "BlueprintDif.MissingGraph", new IMAGE_BRUSH( "/Icons/blueprint_Dif_missing_graph_8x", Icon8x8 )  );
-		Set( "BlueprintDif.NextDiff", new IMAGE_BRUSH( "/Icons/diff_next_40x", Icon16x16 )  );
-		Set( "BlueprintDif.PrevDiff", new IMAGE_BRUSH( "/Icons/diff_prev_40x", Icon16x16 )  );
+		Set("BlueprintDif.HasGraph", new IMAGE_BRUSH("/Icons/blueprint_Dif_has_graph_8x", Icon8x8));
+		Set("BlueprintDif.HasGraph.Small", new IMAGE_BRUSH("/Icons/blueprint_Dif_has_graph_8x", Icon8x8));
+		Set("BlueprintDif.MissingGraph", new IMAGE_BRUSH("/Icons/blueprint_Dif_missing_graph_8x", Icon8x8));
+		Set("BlueprintDif.MissingGraph.Small", new IMAGE_BRUSH("/Icons/blueprint_Dif_missing_graph_8x", Icon8x8));
+		Set("BlueprintDif.NextDiff", new IMAGE_BRUSH("/Icons/diff_next_40x", Icon16x16));
+		Set("BlueprintDif.NextDiff.Small", new IMAGE_BRUSH("/Icons/diff_next_40x", Icon16x16));
+		Set("BlueprintDif.PrevDiff", new IMAGE_BRUSH("/Icons/diff_prev_40x", Icon16x16));
+		Set("BlueprintDif.PrevDiff.Small", new IMAGE_BRUSH("/Icons/diff_prev_40x", Icon16x16));
 
-		Set( "BlueprintDif.ItalicText", 
+		Set("BlueprintDif.ItalicText", 
 			FTextBlockStyle(NormalText)
 			.SetFont(TTF_FONT("Fonts/Roboto-Italic", 10))
 			.SetColorAndOpacity(FLinearColor(.7f, .7f, .7f))
diff --git a/Engine/Source/Editor/GameplayTasksEditor/Private/K2Node_LatentGameplayTaskCall.cpp b/Engine/Source/Editor/GameplayTasksEditor/Private/K2Node_LatentGameplayTaskCall.cpp
--- a/Engine/Source/Editor/GameplayTasksEditor/Private/K2Node_LatentGameplayTaskCall.cpp
+++ b/Engine/Source/Editor/GameplayTasksEditor/Private/K2Node_LatentGameplayTaskCall.cpp
@@ -185,7 +185,7 @@ void UK2Node_LatentGameplayTaskCall::CreatePinsForClass(UClass* InClass)
 	{
 		UFunction* ProxyFunction = ProxyFactoryClass->FindFunctionByName(ProxyFactoryFunctionName);
 
-		FString IgnorePropertyListStr = ProxyFunction->GetMetaData(FName(TEXT("HideSpawnParms")));
+		const FString& IgnorePropertyListStr = ProxyFunction->GetMetaData(FName(TEXT("HideSpawnParms")));
 	
 		if (!IgnorePropertyListStr.IsEmpty())
 		{
diff --git a/Engine/Source/Editor/GraphEditor/Private/KismetNodes/KismetNodeInfoContext.cpp b/Engine/Source/Editor/GraphEditor/Private/KismetNodes/KismetNodeInfoContext.cpp
--- a/Engine/Source/Editor/GraphEditor/Private/KismetNodes/KismetNodeInfoContext.cpp
+++ b/Engine/Source/Editor/GraphEditor/Private/KismetNodes/KismetNodeInfoContext.cpp
@@ -45,7 +45,7 @@ FKismetNodeInfoContext::FKismetNodeInfoContext(UEdGraph* SourceGraph)
 				// if the node has a specific "world context" pin, attempt to get the value set for that first
 				if (Function->HasMetaData(FBlueprintMetadata::MD_WorldContext))
 				{
-					FString const WorldContextPinName = Function->GetMetaData(FBlueprintMetadata::MD_WorldContext);
+					const FString& WorldContextPinName = Function->GetMetaData(FBlueprintMetadata::MD_WorldContext);
 					if (UEdGraphPin* ContextPin = FunctionNode->FindPin(WorldContextPinName))
 					{
 						if (UObjectPropertyBase* ContextProperty = Cast<UObjectPropertyBase>(ClassDebugData.FindClassPropertyForPin(ContextPin)))
diff --git a/Engine/Source/Editor/HardwareTargeting/Private/HardwareTargetingModule.cpp b/Engine/Source/Editor/HardwareTargeting/Private/HardwareTargetingModule.cpp
--- a/Engine/Source/Editor/HardwareTargeting/Private/HardwareTargetingModule.cpp
+++ b/Engine/Source/Editor/HardwareTargeting/Private/HardwareTargetingModule.cpp
@@ -127,7 +127,7 @@ static FName HardwareTargetingConsoleVariableMetaFName(TEXT("ConsoleVariable"));
 	bool bModified = SettingsObject->PropertyName != (TargetValue); \
 	UProperty* Property = FindFieldChecked<UProperty>(Class::StaticClass(), GET_MEMBER_NAME_CHECKED(Class, PropertyName)); \
 	if (!Builder.bReadOnly) { \
-		FString CVarName = Property->GetMetaData(HardwareTargetingConsoleVariableMetaFName); \
+		const FString& CVarName = Property->GetMetaData(HardwareTargetingConsoleVariableMetaFName); \
 		if (!CVarName.IsEmpty()) { IConsoleVariable* CVar = IConsoleManager::Get().FindConsoleVariable(*CVarName); \
 			if (CVar) { CVar->Set(TargetValue, ECVF_SetByProjectSetting); } } \
 		SettingsObject->PropertyName = (TargetValue); } \
diff --git a/Engine/Source/Editor/Kismet/Kismet.Build.cs b/Engine/Source/Editor/Kismet/Kismet.Build.cs
--- a/Engine/Source/Editor/Kismet/Kismet.Build.cs
+++ b/Engine/Source/Editor/Kismet/Kismet.Build.cs
@@ -54,8 +54,9 @@ public class Kismet : ModuleRules
 				"HotReload",
                 "BlueprintNativeCodeGen",
                 "UMGEditor",
-                "UMG" // for SBlueprintDiff
-			}
+                "UMG", // for SBlueprintDiff
+                "WorkspaceMenuStructure"
+            }
 			);
 
         DynamicallyLoadedModuleNames.AddRange(
diff --git a/Engine/Source/Editor/Kismet/Private/BlueprintActionMenuUtils.cpp b/Engine/Source/Editor/Kismet/Private/BlueprintActionMenuUtils.cpp
--- a/Engine/Source/Editor/Kismet/Private/BlueprintActionMenuUtils.cpp
+++ b/Engine/Source/Editor/Kismet/Private/BlueprintActionMenuUtils.cpp
@@ -181,7 +181,7 @@ static bool BlueprintActionMenuUtilsImpl::IsUnexposedMemberAction(FBlueprintActi
 		{
 			if (UProperty* Property = Cast<UProperty>(Binding.Get()))
 			{
-				FString const ExposedCategoryMetadata = Property->GetMetaData(FBlueprintMetadata::MD_ExposeFunctionCategories);
+				const FString& ExposedCategoryMetadata = Property->GetMetaData(FBlueprintMetadata::MD_ExposeFunctionCategories);
 				if (ExposedCategoryMetadata.IsEmpty())
 				{
 					continue;
@@ -193,7 +193,7 @@ static bool BlueprintActionMenuUtilsImpl::IsUnexposedMemberAction(FBlueprintActi
 			}
 		}
 
-		FString FunctionCategory = Function->GetMetaData(FBlueprintMetadata::MD_FunctionCategory);
+		const FString& FunctionCategory = Function->GetMetaData(FBlueprintMetadata::MD_FunctionCategory);
 		bIsFliteredOut = !AllExposedCategories.Contains(FunctionCategory);
 	}
 	return bIsFliteredOut;
diff --git a/Engine/Source/Editor/Kismet/Private/BlueprintCompilationManager.cpp b/Engine/Source/Editor/Kismet/Private/BlueprintCompilationManager.cpp
--- a/Engine/Source/Editor/Kismet/Private/BlueprintCompilationManager.cpp
+++ b/Engine/Source/Editor/Kismet/Private/BlueprintCompilationManager.cpp
@@ -31,6 +31,8 @@
 #include "WidgetBlueprint.h"
 #include "Kismet2/KismetDebugUtilities.h"
 
+#define LOCTEXT_NAMESPACE "BlueprintCompilationManager"
+
 /*
 	BLUEPRINT COMPILATION MANAGER IMPLEMENTATION NOTES
 
@@ -85,7 +87,7 @@ struct FBlueprintCompilationManagerImpl : public FGCObject
 
 	void QueueForCompilation(const FBPCompileRequest& CompileJob);
 	void CompileSynchronouslyImpl(const FBPCompileRequest& Request);
-	void FlushCompilationQueueImpl(TArray<UObject*>* ObjLoaded, bool bSuppressBroadcastCompiled);
+	void FlushCompilationQueueImpl(TArray<UObject*>* ObjLoaded, bool bSuppressBroadcastCompiled, TArray<UBlueprint*>* BlueprintsCompiled);
 	void FlushReinstancingQueueImpl();
 	bool HasBlueprintsToCompile() const;
 	bool IsGeneratedClassLayoutReady() const;
@@ -159,6 +161,7 @@ void FBlueprintCompilationManagerImpl::CompileSynchronouslyImpl(const FBPCompile
 	const bool bSkipGarbageCollection		= (Request.CompileOptions & EBlueprintCompileOptions::SkipGarbageCollection		) != EBlueprintCompileOptions::None;
 	const bool bBatchCompile				= (Request.CompileOptions & EBlueprintCompileOptions::BatchCompile				) != EBlueprintCompileOptions::None;
 	const bool bSkipReinstancing			= (Request.CompileOptions & EBlueprintCompileOptions::SkipReinstancing			) != EBlueprintCompileOptions::None;
+	const bool bSkipSaving					= (Request.CompileOptions & EBlueprintCompileOptions::SkipSave					) != EBlueprintCompileOptions::None;
 
 	// Wipe the PreCompile log, any generated messages are now irrelevant
 	Request.BPToCompile->PreCompileLog.Reset();
@@ -181,10 +184,11 @@ void FBlueprintCompilationManagerImpl::CompileSynchronouslyImpl(const FBPCompile
 
 	ensure(QueuedRequests.Num() == 0);
 	QueuedRequests.Add(Request);
-	// We suppress normal compilation boradcasts because the old code path 
+	// We suppress normal compilation broadcasts because the old code path 
 	// did this after GC and we want to match the old behavior:
 	const bool bSuppressBroadcastCompiled = true;
-	FlushCompilationQueueImpl(nullptr, bSuppressBroadcastCompiled);
+	TArray<UBlueprint*> CompiledBlueprints;
+	FlushCompilationQueueImpl(nullptr, bSuppressBroadcastCompiled, &CompiledBlueprints);
 	FlushReinstancingQueueImpl();
 	
 	if ( GEditor )
@@ -201,21 +205,29 @@ void FBlueprintCompilationManagerImpl::CompileSynchronouslyImpl(const FBPCompile
 	
 	if (!bBatchCompile)
 	{
+		for(UBlueprint* BP : CompiledBlueprints)
+		{
+			BP->BroadcastCompiled();
+		}
+
 		if(GEditor)
 		{
 			GEditor->BroadcastBlueprintCompiled();	
 		}
 	}
 
-	if(CompiledBlueprintsToSave.Num() > 0)
+	if (CompiledBlueprintsToSave.Num() > 0)
 	{
-		TArray<UPackage*> PackagesToSave;
-		for(UBlueprint* BP : CompiledBlueprintsToSave)
+		if (!bSkipSaving)
 		{
-			PackagesToSave.Add(BP->GetOutermost());
+			TArray<UPackage*> PackagesToSave;
+			for (UBlueprint* BP : CompiledBlueprintsToSave)
+			{
+				PackagesToSave.Add(BP->GetOutermost());
+			}
+
+			FEditorFileUtils::PromptForCheckoutAndSave(PackagesToSave, /*bCheckDirty =*/true, /*bPromptToSave =*/false);
 		}
-	
-		FEditorFileUtils::PromptForCheckoutAndSave(PackagesToSave, /*bCheckDirty =*/true, /*bPromptToSave =*/false);
 		CompiledBlueprintsToSave.Empty();
 	}
 }
@@ -296,7 +308,7 @@ struct FReinstancingJob
 	TSharedPtr<FKismetCompilerContext> Compiler;
 };
 	
-void FBlueprintCompilationManagerImpl::FlushCompilationQueueImpl(TArray<UObject*>* ObjLoaded, bool bSuppressBroadcastCompiled)
+void FBlueprintCompilationManagerImpl::FlushCompilationQueueImpl(TArray<UObject*>* ObjLoaded, bool bSuppressBroadcastCompiled, TArray<UBlueprint*>* BlueprintsCompiled)
 {
 	TGuardValue<bool> GuardTemplateNameFlag(GCompilingBlueprint, true);
 	ensure(bGeneratedClassLayoutReady);
@@ -310,32 +322,14 @@ void FBlueprintCompilationManagerImpl::FlushCompilationQueueImpl(TArray<UObject*
 	{ // begin GTimeCompiling scope 
 		FScopedDurationTimer SetupTimer(GTimeCompiling); 
 
-		// STAGE I: Add any child blueprints that were not compiled, we will need to run a link pass on them:
+		// STAGE I: Add any related blueprints that were not compiled, then add any children so that they will be relinked:
+		TArray<UBlueprint*> BlueprintsToRecompile;
 		for(const FBPCompileRequest& ComplieJob : QueuedRequests)
 		{
-			if(UClass* OldSkeletonClass = ComplieJob.BPToCompile->SkeletonGeneratedClass)
-			{
-				TArray<UClass*> SkeletonClassesToReparentList;
-				GetDerivedClasses(OldSkeletonClass, SkeletonClassesToReparentList);
-		
-				for(UClass* ChildClass : SkeletonClassesToReparentList)
-				{
-					if(UBlueprint* ChildBlueprint = UBlueprint::GetBlueprintFromClass(ChildClass))
-					{
-						if(!IsQueuedForCompilation(ChildBlueprint))
-						{
-							ChildBlueprint->bQueuedForCompilation = true;
-							ensure(ChildBlueprint->bHasBeenRegenerated);
-							CurrentlyCompilingBPs.Add(FCompilerData(ChildBlueprint, ECompilationManagerJobType::RelinkOnly, nullptr, EBlueprintCompileOptions::None, false));
-						}
-					}
-				}
-			}
-			
 			// Add any dependent blueprints for a bytecode compile, this is needed because we 
 			// have no way to keep bytecode safe when a function is renamed or parameters are
 			// added or removed - strictly speaking we only need to do this when function 
-			// parameters changed, but that's a somewhat dubious optmization - ideally this
+			// parameters changed, but that's a somewhat dubious optimization - ideally this
 			// work *never* needs to happen:
 			if(!FBlueprintEditorUtils::IsInterfaceBlueprint(ComplieJob.BPToCompile))
 			{
@@ -347,6 +341,7 @@ void FBlueprintCompilationManagerImpl::FlushCompilationQueueImpl(TArray<UObject*
 					{
 						DependentBlueprint->bQueuedForCompilation = true;
 						CurrentlyCompilingBPs.Add(FCompilerData(DependentBlueprint, ECompilationManagerJobType::Normal, nullptr, EBlueprintCompileOptions::None, true));
+						BlueprintsToRecompile.Add(DependentBlueprint);
 					}
 				}
 			}
@@ -416,9 +411,34 @@ void FBlueprintCompilationManagerImpl::FlushCompilationQueueImpl(TArray<UObject*
 			else
 			{
 				CurrentlyCompilingBPs.Add(FCompilerData(QueuedBP, ECompilationManagerJobType::Normal, QueuedJob.ClientResultsLog, QueuedJob.CompileOptions, false));
+				BlueprintsToRecompile.Add(QueuedBP);
+			}
+		}
+			
+		for(UBlueprint* BP : BlueprintsToRecompile)
+		{
+			// make sure all children are at least re-linked:
+			if(UClass* OldSkeletonClass = BP->SkeletonGeneratedClass)
+			{
+				TArray<UClass*> SkeletonClassesToReparentList;
+				GetDerivedClasses(OldSkeletonClass, SkeletonClassesToReparentList);
+		
+				for(UClass* ChildClass : SkeletonClassesToReparentList)
+				{
+					if(UBlueprint* ChildBlueprint = UBlueprint::GetBlueprintFromClass(ChildClass))
+					{
+						if(!IsQueuedForCompilation(ChildBlueprint))
+						{
+							ChildBlueprint->bQueuedForCompilation = true;
+							ensure(ChildBlueprint->bHasBeenRegenerated);
+							CurrentlyCompilingBPs.Add(FCompilerData(ChildBlueprint, ECompilationManagerJobType::RelinkOnly, nullptr, EBlueprintCompileOptions::None, false));
+						}
+					}
+				}
 			}
 		}
 
+		BlueprintsToRecompile.Empty();
 		QueuedRequests.Empty();
 
 		// STAGE III: Sort into correct compilation order. We want to compile root types before their derived (child) types:
@@ -689,6 +709,10 @@ void FBlueprintCompilationManagerImpl::FlushCompilationQueueImpl(TArray<UObject*
 					FKismetCompilerContext& CompilerContext = *(CompilerData.Compiler);
 					CompilerContext.CompileClassLayout(EInternalCompilerFlags::PostponeLocalsGenerationUntilPhaseTwo);
 				}
+				else
+				{
+					CompilerData.ActiveResultsLog->Error(*LOCTEXT("KismetCompileError_MalformedParentClasss", "Blueprint @@ has missing or NULL parent class.").ToString(), BP);
+				}
 			}
 			else if(CompilerData.Compiler.IsValid() && BP->GeneratedClass)
 			{
@@ -833,12 +857,7 @@ void FBlueprintCompilationManagerImpl::FlushCompilationQueueImpl(TArray<UObject*
 				{
 					// for interface changes, auto-refresh nodes on any dependent blueprints
 					// note: RefreshAllNodes() will internally send a change notification event to the dependent blueprint
-					if (bIsInterface)
-					{
-						ensure(!Dependent->bIsRegeneratingOnLoad);
-						FBlueprintEditorUtils::RefreshAllNodes(Dependent);
-					}
-					else if(!BP->bIsRegeneratingOnLoad)
+					if(!BP->bIsRegeneratingOnLoad)
 					{
 						// Some logic (e.g. UObject::ProcessInternal) uses this flag to suppress warnings:
 						TGuardValue<bool> ReinstancingGuard(GIsReinstancing, true);
@@ -868,9 +887,18 @@ void FBlueprintCompilationManagerImpl::FlushCompilationQueueImpl(TArray<UObject*
 		{
 			if(CompilerData.ShouldCompileClassFunctions())
 			{
-				// Some logic (e.g. UObject::ProcessInternal) uses this flag to suppress warnings:
-				TGuardValue<bool> ReinstancingGuard(GIsReinstancing, true);
-				CompilerData.BP->BroadcastCompiled();
+				if(BlueprintsCompiled)
+				{
+					BlueprintsCompiled->Add(CompilerData.BP);
+				}
+				
+				if(!bSuppressBroadcastCompiled)
+				{
+					// Some logic (e.g. UObject::ProcessInternal) uses this flag to suppress warnings:
+					TGuardValue<bool> ReinstancingGuard(GIsReinstancing, true);
+					CompilerData.BP->BroadcastCompiled();
+				}
+
 				continue;
 			}
 
@@ -1181,8 +1209,11 @@ void FBlueprintCompilationManagerImpl::ReinstanceBatch(TArray<FReinstancingJob>&
 						Obj->HasAnyFlags(RF_ArchetypeObject|RF_InheritableComponentTemplate)
 						|| Obj->GetTypedOuter<UBlueprintGeneratedClass>()
 						|| Obj->GetTypedOuter<UBlueprint>();
-					// remove if this is not an archetype or its an archetype in the transient package:
-					return !bIsArchetype || Obj->GetOutermost() == GetTransientPackage(); 
+					// remove if this is not an archetype or its already in the transient package, note
+					// that things that are not directly outered to the transient package will be 
+					// 'reinst'd', this is specifically to handle components, which need to be up to date
+					// on the REINST_ actor class:
+					return !bIsArchetype || Obj->GetOuter() == GetTransientPackage(); 
 				}
 			);
 
@@ -1597,7 +1628,7 @@ UClass* FBlueprintCompilationManagerImpl::FastGenerateSkeletonClass(UBlueprint*
 	UField** CurrentFieldStorageLocation = &Ret->Children;
 	
 	// Helper function for making UFunctions generated for 'event' nodes, e.g. custom event and timelines
-	const auto MakeEventFunction = [&CurrentFieldStorageLocation, MakeFunction]( FName InName, EFunctionFlags ExtraFnFlags, const TArray<UEdGraphPin*>& InputPins, UFunction* InSourceFN )
+	const auto MakeEventFunction = [&CurrentFieldStorageLocation, MakeFunction, Schema]( FName InName, EFunctionFlags ExtraFnFlags, const TArray<UEdGraphPin*>& InputPins, UFunction* InSourceFN, bool bInCallInEditor )
 	{
 		UField** CurrentParamStorageLocation = nullptr;
 
@@ -1615,6 +1646,24 @@ UClass* FBlueprintCompilationManagerImpl::FastGenerateSkeletonClass(UBlueprint*
 
 		if(NewFunction)
 		{
+			for (UEdGraphPin* InputPin : InputPins)
+			{
+				// No defaults for object/class pins
+				if(	!Schema->IsMetaPin(*InputPin) && 
+					(InputPin->PinType.PinCategory != UEdGraphSchema_K2::PC_Object) && 
+					(InputPin->PinType.PinCategory != UEdGraphSchema_K2::PC_Class) && 
+					(InputPin->PinType.PinCategory != UEdGraphSchema_K2::PC_Interface) && 
+					!InputPin->DefaultValue.IsEmpty() )
+				{
+					NewFunction->SetMetaData(*InputPin->PinName, *InputPin->DefaultValue);
+				}
+			}
+
+			if(bInCallInEditor)
+			{
+				NewFunction->SetMetaData(FBlueprintMetadata::MD_CallInEditor, TEXT( "true" ));
+			}
+
 			NewFunction->Bind();
 			NewFunction->StaticLink(true);
 		}
@@ -1649,11 +1698,18 @@ UClass* FBlueprintCompilationManagerImpl::FastGenerateSkeletonClass(UBlueprint*
 		Graph->GetNodesOfClass(EventNodes);
 		for( UK2Node_Event* Event : EventNodes )
 		{
+			bool bCallInEditor = false;
+			if(UK2Node_CustomEvent* CustomEvent = Cast<UK2Node_CustomEvent>(Event))
+			{
+				bCallInEditor = CustomEvent->bCallInEditor;
+			}
+
 			MakeEventFunction(
 				CompilerContext.GetEventStubFunctionName(Event), 
 				(EFunctionFlags)Event->FunctionFlags, 
 				Event->Pins, 
-				Event->FindEventSignatureFunction()
+				Event->FindEventSignatureFunction(),
+				bCallInEditor
 			);
 		}
 	}
@@ -1662,11 +1718,11 @@ UClass* FBlueprintCompilationManagerImpl::FastGenerateSkeletonClass(UBlueprint*
 	{
 		for(int32 EventTrackIdx=0; EventTrackIdx<Timeline->EventTracks.Num(); EventTrackIdx++)
 		{
-			MakeEventFunction(Timeline->GetEventTrackFunctionName(EventTrackIdx), EFunctionFlags::FUNC_None, TArray<UEdGraphPin*>(), nullptr);
+			MakeEventFunction(Timeline->GetEventTrackFunctionName(EventTrackIdx), EFunctionFlags::FUNC_None, TArray<UEdGraphPin*>(), nullptr, false);
 		}
 		
-		MakeEventFunction(Timeline->GetUpdateFunctionName(), EFunctionFlags::FUNC_None, TArray<UEdGraphPin*>(), nullptr);
-		MakeEventFunction(Timeline->GetFinishedFunctionName(), EFunctionFlags::FUNC_None, TArray<UEdGraphPin*>(), nullptr);
+		MakeEventFunction(Timeline->GetUpdateFunctionName(), EFunctionFlags::FUNC_None, TArray<UEdGraphPin*>(), nullptr, false);
+		MakeEventFunction(Timeline->GetFinishedFunctionName(), EFunctionFlags::FUNC_None, TArray<UEdGraphPin*>(), nullptr, false);
 	}
 
 	CompilerContext.NewClass = Ret;
@@ -1905,7 +1961,7 @@ void FBlueprintCompilationManager::FlushCompilationQueue(TArray<UObject*>* ObjLo
 {
 	if(BPCMImpl)
 	{
-		BPCMImpl->FlushCompilationQueueImpl(ObjLoaded, false);
+		BPCMImpl->FlushCompilationQueueImpl(ObjLoaded, false, nullptr);
 
 		// we can't support save on compile when reinstancing is deferred:
 		BPCMImpl->CompiledBlueprintsToSave.Empty();
@@ -1947,4 +2003,4 @@ bool FBlueprintCompilationManager::IsGeneratedClassLayoutReady()
 	return BPCMImpl->IsGeneratedClassLayoutReady();
 }
 
-
+#undef LOCTEXT_NAMESPACE
diff --git a/Engine/Source/Editor/Kismet/Private/BlueprintDetailsCustomization.cpp b/Engine/Source/Editor/Kismet/Private/BlueprintDetailsCustomization.cpp
--- a/Engine/Source/Editor/Kismet/Private/BlueprintDetailsCustomization.cpp
+++ b/Engine/Source/Editor/Kismet/Private/BlueprintDetailsCustomization.cpp
@@ -5145,11 +5145,11 @@ bool FBlueprintGlobalOptionsDetails::CanDeprecateBlueprint() const
 	if (UBlueprint* Blueprint = GetBlueprintObj())
 	{
 		// If the parent is deprecated, we cannot modify deprecation on this Blueprint
-		if (Blueprint->ParentClass->HasAnyClassFlags(CLASS_Deprecated))
+		if (Blueprint->ParentClass && Blueprint->ParentClass->HasAnyClassFlags(CLASS_Deprecated))
 		{
 			return false;
 		}
-		
+
 		return true;
 	}
 
@@ -5196,7 +5196,7 @@ void FBlueprintGlobalOptionsDetails::CustomizeDetails(IDetailLayoutBuilder& Deta
 		for (TFieldIterator<UProperty> PropertyIt(Blueprint->GetClass(), EFieldIteratorFlags::IncludeSuper); PropertyIt; ++PropertyIt)
 		{
 			UProperty* Property = *PropertyIt;
-			FString Category = Property->GetMetaData(TEXT("Category"));
+			const FString& Category = Property->GetMetaData(TEXT("Category"));
 
 			if (Category != TEXT("BlueprintOptions") && Category != TEXT("ClassOptions"))
 			{
diff --git a/Engine/Source/Editor/Kismet/Private/BlueprintEditor.cpp b/Engine/Source/Editor/Kismet/Private/BlueprintEditor.cpp
--- a/Engine/Source/Editor/Kismet/Private/BlueprintEditor.cpp
+++ b/Engine/Source/Editor/Kismet/Private/BlueprintEditor.cpp
@@ -572,6 +572,13 @@ bool FBlueprintEditor::IsInAScriptingMode() const
 
 bool FBlueprintEditor::OnRequestClose()
 {
+	// Also close the Find Results tab if we're not in full edit mode and the option to host Global Find Results is enabled.
+	TSharedPtr<SDockTab> FindResultsTab = TabManager->FindExistingLiveTab(FBlueprintEditorTabs::FindResultsID);
+	if (FindResultsTab.IsValid() && !IsInAScriptingMode() && GetDefault<UBlueprintEditorSettings>()->bHostFindInBlueprintsInGlobalTab)
+	{
+		FindResultsTab->RequestCloseTab();
+	}
+
 	bEditorMarkedAsClosed = true;
 	return FWorkflowCentricApplication::OnRequestClose();
 }
@@ -806,8 +813,23 @@ void FBlueprintEditor::ClearSelectionStateFor(FName SelectionOwner)
 
 void FBlueprintEditor::SummonSearchUI(bool bSetFindWithinBlueprint, FString NewSearchTerms, bool bSelectFirstResult)
 {
-	TabManager->InvokeTab(FBlueprintEditorTabs::FindResultsID);
-	FindResults->FocusForUse(bSetFindWithinBlueprint, NewSearchTerms, bSelectFirstResult);
+	TSharedPtr<SFindInBlueprints> FindResultsToUse;
+
+	if (bSetFindWithinBlueprint
+		|| !GetDefault<UBlueprintEditorSettings>()->bHostFindInBlueprintsInGlobalTab)
+	{
+		FindResultsToUse = FindResults;
+		TabManager->InvokeTab(FBlueprintEditorTabs::FindResultsID);
+	}
+	else
+	{
+		FindResultsToUse = FFindInBlueprintSearchManager::Get().GetGlobalFindResults();
+	}
+
+	if (FindResultsToUse.IsValid())
+	{
+		FindResultsToUse->FocusForUse(bSetFindWithinBlueprint, NewSearchTerms, bSelectFirstResult);
+	}
 }
 
 void FBlueprintEditor::SummonFindAndReplaceUI()
@@ -2101,6 +2123,12 @@ void FBlueprintEditor::PostLayoutBlueprintEditorInitialization()
 			TabManager->InvokeTab(FBlueprintEditorTabs::CompilerResultsID);
 		}
 	}
+
+	if (!GetDefault<UBlueprintEditorSettings>()->bHostFindInBlueprintsInGlobalTab)
+	{
+		// Close any docked global FiB tabs that may have been restored with a saved layout.
+		FFindInBlueprintSearchManager::Get().CloseOrphanedGlobalFindResultsTabs(TabManager);
+	}
 }
 
 void FBlueprintEditor::SetupViewForBlueprintEditingMode()
@@ -2599,8 +2627,7 @@ bool FBlueprintEditor::CanGenerateNativeCode() const
 
 void FBlueprintEditor::FindInBlueprint_Clicked()
 {
-	TabManager->InvokeTab(FBlueprintEditorTabs::FindResultsID);
-	FindResults->FocusForUse(true);
+	SummonSearchUI(true);
 }
 
 void FBlueprintEditor::ReparentBlueprint_Clicked()
@@ -3264,9 +3291,8 @@ void FBlueprintEditor::DeleteUnusedVariables_OnClicked()
 void FBlueprintEditor::FindInBlueprints_OnClicked()
 {
 	SetCurrentMode(FBlueprintEditorApplicationModes::StandardBlueprintEditorMode);
-
-	TabManager->InvokeTab(FBlueprintEditorTabs::FindResultsID);
-	FindResults->FocusForUse(false);
+	
+	SummonSearchUI(false);
 }
 
 void FBlueprintEditor::ClearAllBreakpoints()
@@ -6782,7 +6808,7 @@ UEdGraph* FBlueprintEditor::CollapseSelectionToMacro(TSharedPtr<SGraphEditor> In
 
 void FBlueprintEditor::ExpandNode(UEdGraphNode* InNodeToExpand, UEdGraph* InSourceGraph, TSet<UEdGraphNode*>& OutExpandedNodes)
 {
-	UEdGraph* DestinationGraph = InNodeToExpand->GetGraph();
+ 	UEdGraph* DestinationGraph = InNodeToExpand->GetGraph();
 	UEdGraph* SourceGraph = InSourceGraph;
 	check(SourceGraph);
 
@@ -6866,7 +6892,7 @@ void FBlueprintEditor::ExpandNode(UEdGraphNode* InNodeToExpand, UEdGraph* InSour
 	}
 
 	const UEdGraphSchema_K2* K2Schema = GetDefault<UEdGraphSchema_K2>();
-	K2Schema->CollapseGatewayNode(Cast<UK2Node>(InNodeToExpand), Entry, Result);
+	K2Schema->CollapseGatewayNode(Cast<UK2Node>(InNodeToExpand), Entry, Result, nullptr, &OutExpandedNodes);
 
 	if(Entry)
 	{
diff --git a/Engine/Source/Editor/Kismet/Private/BlueprintEditorCommands.cpp b/Engine/Source/Editor/Kismet/Private/BlueprintEditorCommands.cpp
--- a/Engine/Source/Editor/Kismet/Private/BlueprintEditorCommands.cpp
+++ b/Engine/Source/Editor/Kismet/Private/BlueprintEditorCommands.cpp
@@ -107,7 +107,7 @@ namespace NodeSpawnInfoHelpers
 		TArray<FString> ExcludedEventNames;
 		if( InBlueprint->ParentClass->HasMetaData(*ExclusionListKeyName) )
 		{
-			const FString ExcludedEventNameString = InBlueprint->ParentClass->GetMetaData(*ExclusionListKeyName);
+			const FString& ExcludedEventNameString = InBlueprint->ParentClass->GetMetaData(*ExclusionListKeyName);
 			ExcludedEventNameString.ParseIntoArray(ExcludedEventNames, TEXT(","), true);
 		}
 
diff --git a/Engine/Source/Editor/Kismet/Private/BlueprintEditorTabFactories.cpp b/Engine/Source/Editor/Kismet/Private/BlueprintEditorTabFactories.cpp
--- a/Engine/Source/Editor/Kismet/Private/BlueprintEditorTabFactories.cpp
+++ b/Engine/Source/Editor/Kismet/Private/BlueprintEditorTabFactories.cpp
@@ -20,7 +20,7 @@
 #include "SMyBlueprint.h"
 #include "SReplaceNodeReferences.h"
 #include "Widgets/Input/SHyperlink.h"
-
+#include "BlueprintEditorSettings.h"
 
 #define LOCTEXT_NAMESPACE "BlueprintEditor"
 
@@ -377,7 +377,15 @@ FFindResultsSummoner::FFindResultsSummoner(TSharedPtr<class FAssetEditorToolkit>
 	bIsSingleton = true;
 
 	ViewMenuDescription = LOCTEXT("FindResultsView", "Find Results");
-	ViewMenuTooltip = LOCTEXT("FindResultsView_ToolTip", "Show find results for searching in this blueprint or all blueprints");
+
+	if (GetDefault<UBlueprintEditorSettings>()->bHostFindInBlueprintsInGlobalTab)
+	{
+		ViewMenuTooltip = LOCTEXT("FindResultsView_ToolTip", "Show find results for searching in this blueprint");
+	}
+	else
+	{
+		ViewMenuTooltip = LOCTEXT("FindResultsView_ToolTip", "Show find results for searching in this blueprint or all blueprints");
+	}
 }
 
 TSharedRef<SWidget> FFindResultsSummoner::CreateTabBody(const FWorkflowTabSpawnInfo& Info) const
