Fix a number of localization issues

diff --git a/Engine/Source/Editor/DetailCustomizations/Private/ComponentTransformDetails.cpp b/Engine/Source/Editor/DetailCustomizations/Private/ComponentTransformDetails.cpp
--- a/Engine/Source/Editor/DetailCustomizations/Private/ComponentTransformDetails.cpp
+++ b/Engine/Source/Editor/DetailCustomizations/Private/ComponentTransformDetails.cpp
@@ -1311,7 +1311,7 @@ void FComponentTransformDetails::OnBeginRotatonSlider()
 			{
 				if(Object->IsA<AActor>())
 				{
-					GEditor->BeginTransaction( LOCTEXT( "OnSetRotation", "Set actor rotation" ) );
+					GEditor->BeginTransaction( LOCTEXT( "OnSetRotation", "Set Rotation" ) );
 				}
 				else
 				{
@@ -1344,7 +1344,7 @@ void FComponentTransformDetails::OnBeginRotatonSlider()
 	// Just in case we couldn't start a new transaction for some reason
 	if(!bBeganTransaction)
 	{
-		GEditor->BeginTransaction( LOCTEXT( "OnSetRotation", "Set actor rotation" ) );
+		GEditor->BeginTransaction( LOCTEXT( "OnSetRotation", "Set Rotation" ) );
 	}	
 }
 
diff --git a/Engine/Source/Editor/Kismet/Private/BlueprintEditorTabFactories.cpp b/Engine/Source/Editor/Kismet/Private/BlueprintEditorTabFactories.cpp
--- a/Engine/Source/Editor/Kismet/Private/BlueprintEditorTabFactories.cpp
+++ b/Engine/Source/Editor/Kismet/Private/BlueprintEditorTabFactories.cpp
@@ -384,7 +384,7 @@ FFindResultsSummoner::FFindResultsSummoner(TSharedPtr<class FAssetEditorToolkit>
 	}
 	else
 	{
-		ViewMenuTooltip = LOCTEXT("FindResultsView_ToolTip", "Show find results for searching in this blueprint or all blueprints");
+		ViewMenuTooltip = LOCTEXT("FindResultsViewAllBlueprints_ToolTip", "Show find results for searching in this blueprint or all blueprints");
 	}
 }
 
diff --git a/Engine/Source/Editor/StaticMeshEditor/Private/StaticMeshEditorTools.cpp b/Engine/Source/Editor/StaticMeshEditor/Private/StaticMeshEditorTools.cpp
--- a/Engine/Source/Editor/StaticMeshEditor/Private/StaticMeshEditorTools.cpp
+++ b/Engine/Source/Editor/StaticMeshEditor/Private/StaticMeshEditorTools.cpp
@@ -3512,7 +3512,7 @@ FText FLevelOfDetailSettingsLayout::GetLODCustomModeNameContent(int32 LODIndex)
 	{
 		return LOCTEXT("GetLODCustomModeNameContent", "Custom");
 	}
-	return FText::Format(LOCTEXT("GetLODCustomModeNameContent", "LOD{0}"), LODIndex);
+	return FText::Format(LOCTEXT("GetLODModeNameContent", "LOD{0}"), LODIndex);
 }
 
 ECheckBoxState FLevelOfDetailSettingsLayout::IsLODCustomModeCheck(int32 LODIndex) const
diff --git a/Engine/Source/Editor/UnrealEd/Private/Commandlets/GatherTextFromSourceCommandlet.cpp b/Engine/Source/Editor/UnrealEd/Private/Commandlets/GatherTextFromSourceCommandlet.cpp
--- a/Engine/Source/Editor/UnrealEd/Private/Commandlets/GatherTextFromSourceCommandlet.cpp
+++ b/Engine/Source/Editor/UnrealEd/Private/Commandlets/GatherTextFromSourceCommandlet.cpp
@@ -1421,7 +1421,8 @@ void UGatherTextFromSourceCommandlet::FCommandMacroDescriptor::TryParse(const FS
 		TArray<FString> Arguments;
 		if (ParseArgsFromMacro(StripCommentsFromToken(Text, Context), Arguments, Context))
 		{
-			if (Arguments.Num() != 5)
+			// 5 or 6 arguments depending on whether or not there are one or two bindings
+			if (Arguments.Num() != 5 || Arguments.Num() != 6)
 			{
 				UE_LOG(LogGatherTextFromSourceCommandlet, Warning, TEXT("Too many arguments in command %s macro in %s(%d):%s"), *GetToken(), *Context.Filename, Context.LineNumber, *FLocTextHelper::SanitizeLogOutput(Context.LineText));
 			}
diff --git a/Engine/Source/Runtime/Engine/Private/UnrealEngine.cpp b/Engine/Source/Runtime/Engine/Private/UnrealEngine.cpp
--- a/Engine/Source/Runtime/Engine/Private/UnrealEngine.cpp
+++ b/Engine/Source/Runtime/Engine/Private/UnrealEngine.cpp
@@ -8600,7 +8600,7 @@ void DrawStatsHUD( UWorld* World, FViewport* Viewport, FCanvas* Canvas, UCanvas*
 #if ENABLE_LOW_LEVEL_MEM_TRACKER
 			if (FLowLevelMemTracker::Get().IsEnabled() && !FPlatformMemory::IsDebugMemoryEnabled())
 			{
-				SmallTextItem.Text = LOCTEXT("MEMPROFILINGWARNING", "LLM enabled without Debug Memory enabled!");
+				SmallTextItem.Text = LOCTEXT("MEMPROFILINGWARNINGLLM", "LLM enabled without Debug Memory enabled!");
 				Canvas->DrawItem(SmallTextItem, FVector2D(MessageX, MessageY));
 				MessageY += FontSizeY;
 			}
