Fix opening color picker causing values to change each time.  Was due to conversion between srgb and linear color that would happen each time.  Note that the color picker is still aways in linear space for now so some conversion always happens when editing an FColor

diff --git a/Engine/Source/Editor/DetailCustomizations/Private/ColorStructCustomization.cpp b/Engine/Source/Editor/DetailCustomizations/Private/ColorStructCustomization.cpp
--- a/Engine/Source/Editor/DetailCustomizations/Private/ColorStructCustomization.cpp
+++ b/Engine/Source/Editor/DetailCustomizations/Private/ColorStructCustomization.cpp
@@ -185,13 +185,13 @@ void FColorStructCustomization::CreateColorPicker(bool bUseAlpha)
 		{
 			FLinearColor Color;
 			Color.InitFromString(PerObjectValues[ObjectIndex]);
-			SavedPreColorPickerColors.Add(Color);	
+			SavedPreColorPickerColors.Add(FLinearOrSrgbColor(Color));
 		}
 		else
 		{
 			FColor Color;
 			Color.InitFromString(PerObjectValues[ObjectIndex]);
-			SavedPreColorPickerColors.Add(FLinearColor(Color));
+			SavedPreColorPickerColors.Add(FLinearOrSrgbColor(Color));
 		}
 	}
 
@@ -239,13 +239,13 @@ TSharedRef<SColorPicker> FColorStructCustomization::CreateInlineColorPicker(TWea
 		{
 			FLinearColor Color;
 			Color.InitFromString(PerObjectValues[ObjectIndex]);
-			SavedPreColorPickerColors.Add(Color);	
+			SavedPreColorPickerColors.Add(FLinearOrSrgbColor(Color));	
 		}
 		else
 		{
 			FColor Color;
 			Color.InitFromString(PerObjectValues[ObjectIndex]);
-			SavedPreColorPickerColors.Add(FLinearColor(Color));
+			SavedPreColorPickerColors.Add(FLinearOrSrgbColor(Color));
 		}
 	}
 
@@ -296,12 +296,11 @@ void FColorStructCustomization::OnColorPickerCancelled(FLinearColor OriginalColo
 	{
 		if (bIsLinearColor)
 		{
-			PerObjectColors.Add(SavedPreColorPickerColors[ColorIndex].ToString());
+			PerObjectColors.Add(SavedPreColorPickerColors[ColorIndex].GetLinear().ToString());
 		}
 		else
 		{
-			const bool bSRGB = true;
-			FColor Color = SavedPreColorPickerColors[ColorIndex].ToFColor(bSRGB);
+			FColor Color = SavedPreColorPickerColors[ColorIndex].GetSrgb();
 			PerObjectColors.Add(Color.ToString());
 		}
 	}
@@ -350,46 +349,28 @@ FPropertyAccess::Result FColorStructCustomization::GetColorAsLinear(FLinearColor
 	// Default to full alpha in case the alpha component is disabled.
 	OutColor.A = 1.0f;
 
-	// Get each color component 
-	for (int32 ChildIndex = 0; ChildIndex < SortedChildHandles.Num(); ++ChildIndex)
-	{
-		FPropertyAccess::Result ValueResult = FPropertyAccess::Fail;
+	FString StringValue;
+	FPropertyAccess::Result Result = StructPropertyHandle->GetValueAsFormattedString(StringValue);
 
+	if(Result == FPropertyAccess::Success)
+	{
 		if (bIsLinearColor)
 		{
-			float ComponentValue = 0;
-			ValueResult = SortedChildHandles[ChildIndex]->GetValue(ComponentValue);
-
-			OutColor.Component(ChildIndex) = ComponentValue;
+			OutColor.InitFromString(StringValue);
 		}
 		else
 		{
-			uint8 ComponentValue = 0;
-			ValueResult = SortedChildHandles[ChildIndex]->GetValue(ComponentValue);
-
-			// Convert the FColor to a linear equivalent
-			OutColor.Component(ChildIndex) = ComponentValue/255.0f;
-		}
-
-		switch(ValueResult)
-		{
-		case FPropertyAccess::MultipleValues:
-			{
-				// Default the color to white if we've got multiple values selected
-				OutColor = FLinearColor::White;
-				return FPropertyAccess::MultipleValues;
-			}
-
-		case FPropertyAccess::Fail:
-			return FPropertyAccess::Fail;
-
-		default:
-			break;
+			FColor SrgbColor;
+			SrgbColor.InitFromString(StringValue);
+			OutColor = FLinearColor(SrgbColor);
 		}
 	}
+	else if(Result == FPropertyAccess::MultipleValues)
+	{
+		OutColor = FLinearColor::White;
+	}
 
-	// If we've got this far, we have to have successfully read a color with a single value
-	return FPropertyAccess::Success;
+	return Result;
 }
 
 
diff --git a/Engine/Source/Editor/DetailCustomizations/Private/ColorStructCustomization.h b/Engine/Source/Editor/DetailCustomizations/Private/ColorStructCustomization.h
--- a/Engine/Source/Editor/DetailCustomizations/Private/ColorStructCustomization.h
+++ b/Engine/Source/Editor/DetailCustomizations/Private/ColorStructCustomization.h
@@ -108,9 +108,26 @@ protected:
 	virtual void GetSortedChildren(TSharedRef<IPropertyHandle> StructPropertyHandle, TArray< TSharedRef<IPropertyHandle> >& OutChildren) override;
 
 protected:
+	/** Stores a linear or srb color without converting between the two. Only one is valid at a time */
+	struct FLinearOrSrgbColor
+	{
+		FLinearOrSrgbColor(const FLinearColor& InLinearColor)
+			: LinearColor(InLinearColor)
+		{}
+
+		FLinearOrSrgbColor(const FColor& InSrgbColor)
+			: SrgbColor(InSrgbColor)
+		{}
+
+		FLinearColor GetLinear() const { return LinearColor; }
+		FColor GetSrgb() const { return SrgbColor; }
+	private:
+		FLinearColor LinearColor;
+		FColor SrgbColor;
+	};
 
 	/** Saved per struct colors in case the user clicks cancel in the color picker */
-	TArray<FLinearColor> SavedPreColorPickerColors;
+	TArray<FLinearOrSrgbColor> SavedPreColorPickerColors;
 
 	/** Color struct handle */
 	TSharedPtr<IPropertyHandle> StructPropertyHandle;
