Fix attenuation settings not being disabled when overriden for force feedback and non-english builds

diff --git a/Engine/Source/Editor/DetailCustomizations/Private/AttenuationSettingsCustomizations.cpp b/Engine/Source/Editor/DetailCustomizations/Private/AttenuationSettingsCustomizations.cpp
--- a/Engine/Source/Editor/DetailCustomizations/Private/AttenuationSettingsCustomizations.cpp
+++ b/Engine/Source/Editor/DetailCustomizations/Private/AttenuationSettingsCustomizations.cpp
@@ -51,42 +51,12 @@ void FBaseAttenuationSettingsCustomization::CustomizeHeader(TSharedRef<class IPr
 
 TSharedPtr<IPropertyHandle> FBaseAttenuationSettingsCustomization::GetOverrideAttenuationHandle(TSharedRef<IPropertyHandle> StructPropertyHandle)
 {
-	// Look for override attenuation handle in parent. 
-	// This allows us to disable properties in sound cue/audio component/ambient actor details.
-	uint32 NumChildren;
 	TSharedPtr< IPropertyHandle > ParentHandle = StructPropertyHandle->GetParentHandle();
-	FText DisplayName = ParentHandle->GetPropertyDisplayName();
-	if (DisplayName.ToString() == FString(TEXT("AttenuationSettings")))
+	if (TSharedPtr< IPropertyHandle > GrandParentHandle = ParentHandle->GetParentHandle())
 	{
-		ParentHandle = ParentHandle->GetParentHandle();
-		ParentHandle->GetNumChildren(NumChildren);
-		bool bFoundOverrideProperty = false;
-
-		for (uint32 ChildIndex = 0; ChildIndex < NumChildren; ++ChildIndex)
-		{
-			TSharedRef<IPropertyHandle> ChildHandle = ParentHandle->GetChildHandle(ChildIndex).ToSharedRef();
-			DisplayName = ChildHandle->GetPropertyDisplayName();
-
-			if (DisplayName.ToString() == FString(TEXT("Attenuation")))
-			{
-				uint32 NumAttenChildren;
-				ChildHandle->GetNumChildren(NumAttenChildren);
-
-				for (uint32 i = 0; i < NumAttenChildren; ++i)
-				{
-					TSharedRef<IPropertyHandle> PropHandle = ChildHandle->GetChildHandle(i).ToSharedRef();
-					DisplayName = PropHandle->GetPropertyDisplayName();
-					if (DisplayName.ToString() == FString(TEXT("Override Attenuation")))
-					{
-						return PropHandle;
-					}
-				}
-			}
-		}
+		ParentHandle = GrandParentHandle;
 	}
-
-	// A override attenuation handle wasn't found
-	return nullptr;
+	return ParentHandle->GetChildHandle(TEXT("bOverrideAttenuation"), true);
 }
 
 void FBaseAttenuationSettingsCustomization::CustomizeChildren(TSharedRef<IPropertyHandle> StructPropertyHandle, class IDetailChildrenBuilder& ChildBuilder, IPropertyTypeCustomizationUtils& StructCustomizationUtils)
