Fixed incorrect flags being used when applying frame rates on details panels

The 'Interactive' flag is used to signify that the transaction is externally managed (ie, dragging a slider), so the transaction was never being closed after this function was called

diff --git a/Engine/Source/Editor/DetailCustomizations/Private/FrameRateCustomization.cpp b/Engine/Source/Editor/DetailCustomizations/Private/FrameRateCustomization.cpp
--- a/Engine/Source/Editor/DetailCustomizations/Private/FrameRateCustomization.cpp
+++ b/Engine/Source/Editor/DetailCustomizations/Private/FrameRateCustomization.cpp
@@ -73,7 +73,7 @@ void FFrameRateCustomization::SetFrameRate(FFrameRate NewFrameRate)
 
 		FString TextValue;
 		StructProperty->Struct->ExportText(TextValue, &NewFrameRate, PreviousFrameRate, nullptr, EPropertyPortFlags::PPF_None, nullptr);
-		ensure(StructPropertyHandle->SetValueFromFormattedString(TextValue, EPropertyValueSetFlags::InteractiveChange) == FPropertyAccess::Result::Success);
+		ensure(StructPropertyHandle->SetValueFromFormattedString(TextValue, EPropertyValueSetFlags::DefaultFlags) == FPropertyAccess::Result::Success);
 	}
 }
 
