[NO NON-PII COMMIT MESSAGE]

diff --git a/Engine/Source/Editor/DetailCustomizations/Private/MaterialExpressionLandscapeGrassCustomization.cpp b/Engine/Source/Editor/DetailCustomizations/Private/MaterialExpressionLandscapeGrassCustomization.cpp
--- a/Engine/Source/Editor/DetailCustomizations/Private/MaterialExpressionLandscapeGrassCustomization.cpp
+++ b/Engine/Source/Editor/DetailCustomizations/Private/MaterialExpressionLandscapeGrassCustomization.cpp
@@ -32,11 +32,28 @@ void FMaterialExpressionLandscapeGrassInputCustomization::CustomizeHeader(TShare
 
 void FMaterialExpressionLandscapeGrassInputCustomization::CustomizeChildren(TSharedRef<IPropertyHandle> StructPropertyHandle, IDetailChildrenBuilder& StructBuilder, IPropertyTypeCustomizationUtils& StructCustomizationUtils)
 {
-	TSharedPtr<IPropertyHandle> GrassTypeHandle = StructPropertyHandle->GetChildHandle(GET_MEMBER_NAME_CHECKED(FGrassInput, GrassType));
+	uint32 NumChildren = 0;
+	StructPropertyHandle->GetNumChildren(NumChildren);
+
+	TSharedPtr<IPropertyHandle> GrassTypeHandle;
+
+	for (uint32 i = 0; i < NumChildren; ++i)
+	{
+		TSharedPtr<IPropertyHandle> PropertyHandle = StructPropertyHandle->GetChildHandle(i);
+
+		if (PropertyHandle->GetProperty()->GetFName() == GET_MEMBER_NAME_CHECKED(FGrassInput, GrassType))
+		{
+			GrassTypeHandle = PropertyHandle;
+		}
+		else
+		{
+			StructBuilder.AddProperty(PropertyHandle.ToSharedRef());
+		}
+	}
 
 	TArray<UObject*> OwningObjects;
 	StructPropertyHandle->GetOuterObjects(OwningObjects);
-	
+
 	if (OwningObjects.Num() == 1)
 	{
 		MaterialNode = CastChecked<UMaterialExpressionLandscapeGrassOutput>(OwningObjects[0]);
