Fix for slate brush handle not invalidating on change

diff --git a/Engine/Source/Editor/DetailCustomizations/Private/SlateBrushCustomization.cpp b/Engine/Source/Editor/DetailCustomizations/Private/SlateBrushCustomization.cpp
--- a/Engine/Source/Editor/DetailCustomizations/Private/SlateBrushCustomization.cpp
+++ b/Engine/Source/Editor/DetailCustomizations/Private/SlateBrushCustomization.cpp
@@ -1180,6 +1180,16 @@ private:
 		FPropertyAccess::Result Result = ResourceObjectProperty->GetValue(ResourceObject);
 		if ( Result == FPropertyAccess::Success )
 		{
+			TArray<void*> RawBrushData;
+			ResourceObjectProperty->GetParentHandle()->AccessRawData(RawBrushData);
+			for (int32 BrushIndex = 0; BrushIndex < RawBrushData.Num(); BrushIndex++)
+			{
+				FSlateBrush* TemporaryBrush = static_cast<FSlateBrush*>(RawBrushData[BrushIndex]);
+				if (TemporaryBrush)
+				{
+					TemporaryBrush->InvalidateResourceHandle();
+				}
+			}
 			FVector2D CachedTextureSize;
 
 			TArray<void*> RawData;
diff --git a/Engine/Source/Runtime/SlateCore/Public/Styling/SlateBrush.h b/Engine/Source/Runtime/SlateCore/Public/Styling/SlateBrush.h
--- a/Engine/Source/Runtime/SlateCore/Public/Styling/SlateBrush.h
+++ b/Engine/Source/Runtime/SlateCore/Public/Styling/SlateBrush.h
@@ -320,6 +320,13 @@ public:
 		return ResourceHandle;
 	}
 
+#if WITH_EDITOR
+	void InvalidateResourceHandle()
+	{
+		ResourceHandle = FSlateResourceHandle();
+	};
+#endif
+
 private:
 	void UpdateRenderingResource() const;
 	bool CanRenderResourceObject(UObject* InResourceObject) const;
