Editor: Change to broadcast end object movement at the end of setting the transform in the details panel.

diff --git a/Engine/Source/Editor/DetailCustomizations/Private/ComponentTransformDetails.cpp b/Engine/Source/Editor/DetailCustomizations/Private/ComponentTransformDetails.cpp
--- a/Engine/Source/Editor/DetailCustomizations/Private/ComponentTransformDetails.cpp
+++ b/Engine/Source/Editor/DetailCustomizations/Private/ComponentTransformDetails.cpp
@@ -1254,11 +1254,11 @@ void FComponentTransformDetails::OnSetTransform(ETransformField::Type TransformF
 
 					if (bCommitted)
 					{
-						// Broadcast the first time an actor is about to move
-						GEditor->BroadcastBeginObjectMovement(*SceneComponent);
+						// Broadcast when the actor is done moving
+						GEditor->BroadcastEndObjectMovement(*SceneComponent);
 						if (EditedActor && EditedActor->GetRootComponent() == SceneComponent)
 						{
-							GEditor->BroadcastBeginObjectMovement(*EditedActor);
+							GEditor->BroadcastEndObjectMovement(*EditedActor);
 						}
 					}
 				}
