Fix the PerPlatformPropertiesWidget to be readable when there is many platform

diff --git a/Engine/Source/Editor/DetailCustomizations/Private/PerPlatformPropertyCustomization.cpp b/Engine/Source/Editor/DetailCustomizations/Private/PerPlatformPropertyCustomization.cpp
--- a/Engine/Source/Editor/DetailCustomizations/Private/PerPlatformPropertyCustomization.cpp
+++ b/Engine/Source/Editor/DetailCustomizations/Private/PerPlatformPropertyCustomization.cpp
@@ -19,13 +19,15 @@
 template<typename PerPlatformType>
 void FPerPlatformPropertyCustomization<PerPlatformType>::CustomizeHeader(TSharedRef<IPropertyHandle> StructPropertyHandle, FDetailWidgetRow& HeaderRow, IPropertyTypeCustomizationUtils& StructCustomizationUtils)
 {
+	int32 PlatformNumber = PlatformInfo::GetAllPlatformGroupNames().Num();
+
 	HeaderRow.NameContent()
 	[
 		StructPropertyHandle->CreatePropertyNameWidget()
 	]
 	.ValueContent()
 	.MinDesiredWidth(CalcDesiredWidth(StructPropertyHandle))
-	.MaxDesiredWidth(CalcDesiredWidth(StructPropertyHandle))
+	.MaxDesiredWidth((float)(PlatformNumber + 1)*125.0f)
 	[
 		SNew(SPerPlatformPropertiesWidget)
 		.OnGenerateWidget(this, &FPerPlatformPropertyCustomization<PerPlatformType>::GetWidget, StructPropertyHandle)
@@ -140,43 +142,35 @@ bool FPerPlatformPropertyCustomization<PerPlatformType>::AddPlatformOverride(FNa
 			}
 		}
 	}
-
 	return false;
 }
 
 template<typename PerPlatformType>
 bool FPerPlatformPropertyCustomization<PerPlatformType>::RemovePlatformOverride(FName PlatformGroupName, TSharedRef<IPropertyHandle> StructPropertyHandle)
 {
-	TSharedPtr<IPropertyHandle>	PerPlatformProperty = StructPropertyHandle->GetChildHandle(FName("PerPlatform"));
-	if (PerPlatformProperty.IsValid())
+	TSharedPtr<IPropertyHandle>	MapProperty = StructPropertyHandle->GetChildHandle(FName("PerPlatform"));
+	if (MapProperty.IsValid())
 	{
-		TSharedPtr<IPropertyHandleMap> MapProperty = PerPlatformProperty->AsMap();
-		if (MapProperty.IsValid())
+		TArray<const void*> RawData;
+		MapProperty->AccessRawData(RawData);
+		for (const void* Data : RawData)
 		{
-			uint32 NumChildren = 0;
-			PerPlatformProperty->GetNumChildren(NumChildren);
-			for (uint32 ChildIdx = 0; ChildIdx < NumChildren; ChildIdx++)
+			TMap<FName, typename PerPlatformType::ValueType>* PerPlatformMap = (TMap<FName, typename PerPlatformType::ValueType>*)(Data);
+			check(PerPlatformMap);
+			TArray<FName> KeyArray;
+			PerPlatformMap->GenerateKeyArray(KeyArray);
+			for (FName PlatformName : KeyArray)
 			{
-				TSharedPtr<IPropertyHandle> ChildProperty = PerPlatformProperty->GetChildHandle(ChildIdx);
-				if (ChildProperty.IsValid())
+				if (PlatformName == PlatformGroupName)
 				{
-					TSharedPtr<IPropertyHandle> KeyProperty = ChildProperty->GetKeyHandle();
-					if (KeyProperty.IsValid())
-					{
-						FName KeyName;
-						if (KeyProperty->GetValue(KeyName) == FPropertyAccess::Success && KeyName == PlatformGroupName)
-						{
-							MapProperty->DeleteItem(ChildIdx);
-							return true;
-						}
-					}
+					PerPlatformMap->Remove(PlatformName);
+					return true;
 				}
 			}
 		}
-
-		return true;
 	}
 	return false;
+
 }
 
 template<typename PerPlatformType>
diff --git a/Engine/Source/Editor/PropertyEditor/Public/SPerPlatformPropertiesWidget.h b/Engine/Source/Editor/PropertyEditor/Public/SPerPlatformPropertiesWidget.h
--- a/Engine/Source/Editor/PropertyEditor/Public/SPerPlatformPropertiesWidget.h
+++ b/Engine/Source/Editor/PropertyEditor/Public/SPerPlatformPropertiesWidget.h
@@ -58,8 +58,8 @@ public:
 
 	void ConstructChildren()
 	{
-		TSharedPtr<SHorizontalBox> H;
-		TSharedPtr<SVerticalBox> V;
+		TSharedPtr<SHorizontalBox> HorizontalBox;
+		TSharedPtr<SVerticalBox> VerticalBox;
 
 		TArray<FName> PlatformOverrides = PlatformOverrideNames.Get();
 		LastPlatformOverrideNames = PlatformOverrides.Num();
@@ -67,10 +67,11 @@ public:
 		.VAlign(VAlign_Fill)
 		.HAlign(HAlign_Fill)
 		[
-			SAssignNew(H, SHorizontalBox)
+			SAssignNew(HorizontalBox, SHorizontalBox)
 			+SHorizontalBox::Slot()
+			.FillWidth(1.0f)
 			[
-				SAssignNew(V, SVerticalBox)
+				SAssignNew(VerticalBox, SVerticalBox)
 				+ SVerticalBox::Slot()
 				.AutoHeight()
 				[
@@ -85,7 +86,7 @@ public:
 		if (OnGenerateWidget.IsBound())
 		{
 			// Default control
-			V->AddSlot()
+			VerticalBox->AddSlot()
 			.AutoHeight()
 			.VAlign(EVerticalAlignment::VAlign_Bottom)
 			[
@@ -115,7 +116,7 @@ public:
 
 			if (bAddedMenuItem)
 			{
-				H->AddSlot()
+				HorizontalBox->AddSlot()
 				.AutoWidth()
 				[
 					SNew(SComboButton)
@@ -133,8 +134,9 @@ public:
 
 			for (FName PlatformName : PlatformOverrides)
 			{
-				H->AddSlot()
+				HorizontalBox->AddSlot()
 				.Padding(FMargin(5.0f, 0.0f, 0.0f, 0.0f))
+				.AutoWidth()
 				[
 					SNew(SVerticalBox)
 					+ SVerticalBox::Slot()
@@ -143,6 +145,7 @@ public:
 					[
 						SNew(SHorizontalBox)
 						+ SHorizontalBox::Slot()
+						.FillWidth(1.0f)
 						[
 							SNew(STextBlock)
 							.Font(IDetailLayoutBuilder::GetDetailFont())
@@ -151,9 +154,10 @@ public:
 
 						+ SHorizontalBox::Slot()
 						.HAlign(EHorizontalAlignment::HAlign_Right)
+						.AutoWidth()
 						[
 							SNew(SButton)
-							.ContentPadding(FMargin(1.0f, 1.0f, 1.0f, 1.0f))
+							.ContentPadding(FMargin(1.0f, 0.0f, 1.0f, 0.0f))
 							.ButtonStyle(FEditorStyle::Get(), "NoBorder")
 							.OnClicked(this, &SPerPlatformPropertiesWidget::RemovePlatform, PlatformName)
 							.ToolTipText(FText::Format(NSLOCTEXT("SPerPlatformPropertiesWidget", "RemoveOverrideFor", "Remove Override for {0}"), FText::AsCultureInvariant(PlatformName.ToString())))
@@ -178,7 +182,7 @@ public:
 		}
 		else
 		{
-			V->AddSlot()
+			VerticalBox->AddSlot()
 			.AutoHeight()
 			[
 				SNew(STextBlock)
diff --git a/Engine/Source/Editor/StaticMeshEditor/Private/StaticMeshEditorTools.cpp b/Engine/Source/Editor/StaticMeshEditor/Private/StaticMeshEditorTools.cpp
--- a/Engine/Source/Editor/StaticMeshEditor/Private/StaticMeshEditorTools.cpp
+++ b/Engine/Source/Editor/StaticMeshEditor/Private/StaticMeshEditorTools.cpp
@@ -2778,6 +2778,8 @@ void FLevelOfDetailSettingsLayout::AddToDetailsPanel( IDetailLayoutBuilder& Deta
 			.OnSelectionChanged(this, &FLevelOfDetailSettingsLayout::OnImportLOD)
 		];
 
+	int32 PlatformNumber = PlatformInfo::GetAllPlatformGroupNames().Num();
+
 	LODSettingsCategory.AddCustomRow( LOCTEXT("MinLOD", "Minimum LOD") )
 	.NameContent()
 	[
@@ -2786,6 +2788,8 @@ void FLevelOfDetailSettingsLayout::AddToDetailsPanel( IDetailLayoutBuilder& Deta
 		.Text(LOCTEXT("MinLOD", "Minimum LOD"))
 	]
 	.ValueContent()
+	.MinDesiredWidth((float)(StaticMesh->MinLOD.PerPlatform.Num() + 1)*125.0f)
+	.MaxDesiredWidth((float)(PlatformNumber + 1)*125.0f)
 	[
 		SNew(SPerPlatformPropertiesWidget)
 		.OnGenerateWidget(this, &FLevelOfDetailSettingsLayout::GetMinLODWidget)
@@ -3061,6 +3065,8 @@ void FLevelOfDetailSettingsLayout::AddLODLevelCategories( IDetailLayoutBuilder&
 			SectionSettingsWidgets[ LODIndex ] = MakeShareable( new FMeshSectionSettingsLayout( StaticMeshEditor, LODIndex, LodCategories, &CustomLODEditMode) );
 			SectionSettingsWidgets[ LODIndex ]->AddToCategory( LODCategory );
 
+			int32 PlatformNumber = PlatformInfo::GetAllPlatformGroupNames().Num();
+
 			LODCategory.AddCustomRow(( LOCTEXT("ScreenSizeRow", "ScreenSize")))
 			.NameContent()
 			[
@@ -3069,8 +3075,8 @@ void FLevelOfDetailSettingsLayout::AddLODLevelCategories( IDetailLayoutBuilder&
 				.Text(LOCTEXT("ScreenSizeName", "Screen Size"))
 			]
 			.ValueContent()
-			.MinDesiredWidth(FLevelOfDetailSettingsLayout::GetScreenSizeWidgetWidth(LODIndex))
-			.MaxDesiredWidth(FLevelOfDetailSettingsLayout::GetScreenSizeWidgetWidth(LODIndex))
+			.MinDesiredWidth(GetScreenSizeWidgetWidth(LODIndex))
+			.MaxDesiredWidth((float)(PlatformNumber + 1)*125.0f)
 			[
 				SNew(SPerPlatformPropertiesWidget)
 				.OnGenerateWidget(this, &FLevelOfDetailSettingsLayout::GetLODScreenSizeWidget, LODIndex)
@@ -3173,6 +3179,7 @@ TSharedRef<SWidget> FLevelOfDetailSettingsLayout::GetLODScreenSizeWidget(FName P
 {
 	return SNew(SSpinBox<float>)
 		.Font(IDetailLayoutBuilder::GetDetailFont())
+		.MinDesiredWidth(60.0f)
 		.MinValue(0.0f)
 		.MaxValue(WORLD_MAX)
 		.SliderExponent(2.0f)
