Fix CIS errors from merge.

#rnx

diff --git a/Engine/Source/Editor/DetailCustomizations/Private/LinuxTargetSettingsDetails.cpp b/Engine/Source/Editor/DetailCustomizations/Private/LinuxTargetSettingsDetails.cpp
--- a/Engine/Source/Editor/DetailCustomizations/Private/LinuxTargetSettingsDetails.cpp
+++ b/Engine/Source/Editor/DetailCustomizations/Private/LinuxTargetSettingsDetails.cpp
@@ -154,7 +154,7 @@ void FLinuxTargetSettingsDetails::CustomizeDetails( IDetailLayoutBuilder& Detail
 	// Setup the supported/targeted RHI property view
 	ITargetPlatform* TargetPlatform = FModuleManager::GetModuleChecked<ITargetPlatformModule>("LinuxTargetPlatform").GetTargetPlatforms()[0];
 	TargetShaderFormatsDetails = MakeShareable(new FShaderFormatsPropertyDetails(&DetailBuilder));
-	TargetShaderFormatsDetails->CreateTargetShaderFormatsPropertyView(TargetPlatform, GetFriendlyNameFromRHIName);
+	TargetShaderFormatsDetails->CreateTargetShaderFormatsPropertyView(TargetPlatform, GetFriendlyNameFromLinuxRHIName);
 
 	// Next add the splash image customization
 	const FText EditorSplashDesc(LOCTEXT("EditorSplashLabel", "Editor Splash"));
diff --git a/Engine/Source/Editor/DetailCustomizations/Private/WindowsTargetSettingsDetails.cpp b/Engine/Source/Editor/DetailCustomizations/Private/WindowsTargetSettingsDetails.cpp
--- a/Engine/Source/Editor/DetailCustomizations/Private/WindowsTargetSettingsDetails.cpp
+++ b/Engine/Source/Editor/DetailCustomizations/Private/WindowsTargetSettingsDetails.cpp
@@ -163,7 +163,7 @@ void FWindowsTargetSettingsDetails::CustomizeDetails( IDetailLayoutBuilder& Deta
 	// Setup the supported/targeted RHI property view
 	ITargetPlatform* TargetPlatform = FModuleManager::GetModuleChecked<ITargetPlatformModule>("WindowsTargetPlatform").GetTargetPlatforms()[0];
 	TargetShaderFormatsDetails = MakeShareable(new FShaderFormatsPropertyDetails(&DetailBuilder));
-	TargetShaderFormatsDetails->CreateTargetShaderFormatsPropertyView(TargetPlatform, GetFriendlyNameFromRHIName);
+	TargetShaderFormatsDetails->CreateTargetShaderFormatsPropertyView(TargetPlatform, GetFriendlyNameFromWindowsRHIName);
 
 	TSharedRef<IPropertyHandle> MinOSProperty = DetailBuilder.GetProperty("MinimumOSVersion");
 	IDetailCategoryBuilder& OSInfoCategory = DetailBuilder.EditCategory(TEXT("OS Info"));
