UE-70595 - fix copy right notice in ssearchablecombobox

diff --git a/Engine/Source/Editor/DetailCustomizations/Private/SSearchableComboBox.cpp b/Engine/Source/Editor/DetailCustomizations/Private/SSearchableComboBox.cpp
--- a/Engine/Source/Editor/DetailCustomizations/Private/SSearchableComboBox.cpp
+++ b/Engine/Source/Editor/DetailCustomizations/Private/SSearchableComboBox.cpp
@@ -1,4 +1,4 @@
-// Copyright 1998-2018 Epic Games, Inc. All Rights Reservekd.
+// Copyright 1998-2019 Epic Games, Inc. All Rights Reservekd.
 
 #include "SSearchableComboBox.h"
 #include "DetailLayoutBuilder.h"
diff --git a/Engine/Source/Editor/DetailCustomizations/Private/SSearchableComboBox.h b/Engine/Source/Editor/DetailCustomizations/Private/SSearchableComboBox.h
--- a/Engine/Source/Editor/DetailCustomizations/Private/SSearchableComboBox.h
+++ b/Engine/Source/Editor/DetailCustomizations/Private/SSearchableComboBox.h
@@ -1,4 +1,4 @@
-// Copyright 1998-2018 Epic Games, Inc. All Rights Reservekd.
+// Copyright 1998-2019 Epic Games, Inc. All Rights Reservekd.
 
 #pragma once
 
