Material Instance editor. Adding better behavior for overriding a material with shading model from material expression.

#rnx

diff --git a/Engine/Source/Editor/DetailCustomizations/Private/MaterialShadingModelCustomization.cpp b/Engine/Source/Editor/DetailCustomizations/Private/MaterialShadingModelCustomization.cpp
--- a/Engine/Source/Editor/DetailCustomizations/Private/MaterialShadingModelCustomization.cpp
+++ b/Engine/Source/Editor/DetailCustomizations/Private/MaterialShadingModelCustomization.cpp
@@ -10,6 +10,7 @@
 #include "DetailWidgetRow.h"
 #include "IDetailPropertyRow.h"
 #include "PropertyCustomizationHelpers.h"
+#include "MaterialEditor/MaterialEditorInstanceConstant.h"
 
 TSharedRef<IPropertyTypeCustomization> FMaterialShadingModelCustomization::MakeInstance() 
 {
@@ -27,8 +28,8 @@ void FMaterialShadingModelCustomization::CustomizeHeader(TSharedRef<IPropertyHan
 		TArray<UObject*> OuterObjects;
 		PropertyHandle->GetOuterObjects(OuterObjects);
 		
-		// Don't hide MSM_FromMaterialExpression on a UMaterial
-		if (OuterObjects[0]->IsA<UMaterial>())
+		// Don't hide MSM_FromMaterialExpression on a UMaterial or on UMaterialEditorInstanceConstant
+		if (OuterObjects[0]->IsA<UMaterial>() || OuterObjects[0]->IsA<UMaterialEditorInstanceConstant>())
 		{
 			bShouldHideFromMaterialExpression = false;
 		}
diff --git a/Engine/Source/Runtime/Engine/Private/Materials/MaterialInstance.cpp b/Engine/Source/Runtime/Engine/Private/Materials/MaterialInstance.cpp
--- a/Engine/Source/Runtime/Engine/Private/Materials/MaterialInstance.cpp
+++ b/Engine/Source/Runtime/Engine/Private/Materials/MaterialInstance.cpp
@@ -2435,8 +2435,13 @@ void UMaterialInstance::UpdateOverridableBaseProperties()
 
 	if (BasePropertyOverrides.bOverride_ShadingModel)
 	{
-		// Can only override using one of the actual shading models
-		if (BasePropertyOverrides.ShadingModel != MSM_FromMaterialExpression)
+		if (BasePropertyOverrides.ShadingModel == MSM_FromMaterialExpression)
+		{
+			// Can't override using MSM_FromMaterialExpression, simply fall back to parent
+			ShadingModels = Parent->GetShadingModels();
+			bIsShadingModelFromMaterialExpression = Parent->IsShadingModelFromMaterialExpression();
+		}
+		else
 		{
 			// It's only possible to override using a single shading model
 			ShadingModels = FMaterialShadingModelField(BasePropertyOverrides.ShadingModel);
