Adding missing include

#rnx

diff --git a/Engine/Source/Editor/DetailCustomizations/Private/MaterialShadingModelCustomization.cpp b/Engine/Source/Editor/DetailCustomizations/Private/MaterialShadingModelCustomization.cpp
--- a/Engine/Source/Editor/DetailCustomizations/Private/MaterialShadingModelCustomization.cpp
+++ b/Engine/Source/Editor/DetailCustomizations/Private/MaterialShadingModelCustomization.cpp
@@ -11,6 +11,7 @@
 #include "IDetailPropertyRow.h"
 #include "PropertyCustomizationHelpers.h"
 #include "MaterialEditor/MaterialEditorInstanceConstant.h"
+#include "Materials/Material.h"
 
 TSharedRef<IPropertyTypeCustomization> FMaterialShadingModelCustomization::MakeInstance() 
 {
