Tweak desired width of the EMaterialShadingModel combobox to fit longer names

#rnx

diff --git a/Engine/Source/Editor/DetailCustomizations/Private/MaterialShadingModelCustomization.cpp b/Engine/Source/Editor/DetailCustomizations/Private/MaterialShadingModelCustomization.cpp
--- a/Engine/Source/Editor/DetailCustomizations/Private/MaterialShadingModelCustomization.cpp
+++ b/Engine/Source/Editor/DetailCustomizations/Private/MaterialShadingModelCustomization.cpp
@@ -52,6 +52,8 @@ void FMaterialShadingModelCustomization::CustomizeHeader(TSharedRef<IPropertyHan
 		PropertyHandle->CreatePropertyNameWidget()
 	]
 	.ValueContent()
+	.MinDesiredWidth(150.0f) // To fit the longer names
+	.MaxDesiredWidth(0.0f) // Don't care
 	[
 		PropertyCustomizationHelpers::MakePropertyComboBox(PropertyHandle)
 	];
