Fixed issue where Slate brush resource widget could change the material domain on an engine material.

diff --git a/Engine/Source/Editor/DetailCustomizations/Private/SlateBrushCustomization.cpp b/Engine/Source/Editor/DetailCustomizations/Private/SlateBrushCustomization.cpp
--- a/Engine/Source/Editor/DetailCustomizations/Private/SlateBrushCustomization.cpp
+++ b/Engine/Source/Editor/DetailCustomizations/Private/SlateBrushCustomization.cpp
@@ -1141,10 +1141,16 @@ class SBrushResourceObjectBox : public SCompoundWidget
 					+ SVerticalBox::Slot()
 					.HAlign( HAlign_Left )
 					[
-						SNew( SHyperlink )
+						SAssignNew(ChangeDomainLink, SHyperlink )
 						.Text( NSLOCTEXT("FSlateBrushStructCustomization", "ChangeMaterialDomain_ErrorMessage", "Change the Material Domain?" ) )
 						.OnNavigate( this, &SBrushResourceObjectBox::OnErrorLinkClicked )
 					]
+					+ SVerticalBox::Slot()
+					.HAlign( HAlign_Left )
+					[
+						SAssignNew(IsEngineMaterialError, STextBlock)
+						.Text( NSLOCTEXT("FSlateBrushStructCustomization", "IsEngineMaterialErrorText", "Assign a parent UI Material" ) )
+					]
 				]
 			]
 		];
@@ -1161,6 +1167,20 @@ class SBrushResourceObjectBox : public SCompoundWidget
 			if( BaseMaterial && !BaseMaterial->IsUIMaterial() )
 			{
 				ResourceError->SetVisibility( EVisibility::Visible );
+
+				// Special engine materials cannot change domain. This typically occurs when
+				// the user creates or assigns a material instance with no parent material.
+				// In this case, we warn the user rather than offer to change the domain.
+				if (BaseMaterial->bUsedAsSpecialEngineMaterial)
+				{
+					ChangeDomainLink->SetVisibility( EVisibility::Collapsed );
+					IsEngineMaterialError->SetVisibility( EVisibility::Visible );
+				}
+				else
+				{
+					ChangeDomainLink->SetVisibility( EVisibility::Visible );
+					IsEngineMaterialError->SetVisibility( EVisibility::Collapsed );
+				}
 			}
 			else
 			{
@@ -1242,6 +1262,8 @@ private:
 	TSharedPtr<IPropertyHandle> ResourceObjectProperty;
 	TSharedPtr<IPropertyHandle> ImageSizeProperty;
 	TSharedPtr<SBrushResourceError> ResourceError;
+	TSharedPtr<SHyperlink> ChangeDomainLink;
+	TSharedPtr<STextBlock> IsEngineMaterialError;
 };
 
 // FSlateBrushStructCustomization
