Replace TEXT by a GET_MEMBER_NAME_CHECKED

#rnx

diff --git a/Engine/Source/Editor/DetailCustomizations/Private/LandscapeSettingsDetails.cpp b/Engine/Source/Editor/DetailCustomizations/Private/LandscapeSettingsDetails.cpp
--- a/Engine/Source/Editor/DetailCustomizations/Private/LandscapeSettingsDetails.cpp
+++ b/Engine/Source/Editor/DetailCustomizations/Private/LandscapeSettingsDetails.cpp
@@ -48,7 +48,7 @@ void FLandscapeSettingsDetails::CustomizeDetails( IDetailLayoutBuilder& DetailBu
 
 void FLandscapeSettingsDetails::OnEditorExperimentalSettingsChanged(FName PropertyName)
 {
-	if (DetailLayoutBuilder && PropertyName == TEXT("bLandscapeLayerSystem"))
+	if (DetailLayoutBuilder && PropertyName == GET_MEMBER_NAME_CHECKED(UEditorExperimentalSettings, bLandscapeLayerSystem))
 	{
 		DetailLayoutBuilder->ForceRefreshDetails();
 	}
