Add search field to AnimSequence RetargetSource.

diff --git a/Engine/Source/Editor/DetailCustomizations/Private/AnimSequenceDetails.cpp b/Engine/Source/Editor/DetailCustomizations/Private/AnimSequenceDetails.cpp
--- a/Engine/Source/Editor/DetailCustomizations/Private/AnimSequenceDetails.cpp
+++ b/Engine/Source/Editor/DetailCustomizations/Private/AnimSequenceDetails.cpp
@@ -99,7 +99,7 @@ void FAnimSequenceDetails::CustomizeDetails(IDetailLayoutBuilder& DetailBuilder)
 	]
 	.ValueContent()
 	[
-		SAssignNew(RetargetSourceComboBox, SComboBox< TSharedPtr<FString> >)
+		SAssignNew(RetargetSourceComboBox, SSearchableComboBox)
 		.OptionsSource(&RetargetSourceComboList)
 		.OnGenerateWidget(this, &FAnimSequenceDetails::MakeRetargetSourceComboWidget)
 		.OnSelectionChanged(this, &FAnimSequenceDetails::OnRetargetSourceChanged)
diff --git a/Engine/Source/Editor/DetailCustomizations/Private/AnimSequenceDetails.h b/Engine/Source/Editor/DetailCustomizations/Private/AnimSequenceDetails.h
--- a/Engine/Source/Editor/DetailCustomizations/Private/AnimSequenceDetails.h
+++ b/Engine/Source/Editor/DetailCustomizations/Private/AnimSequenceDetails.h
@@ -9,7 +9,7 @@
 #include "Widgets/DeclarativeSyntaxSupport.h"
 #include "Widgets/SWidget.h"
 #include "Widgets/SCompoundWidget.h"
-#include "Widgets/Input/SComboBox.h"
+#include "SSearchableComboBox.h"
 #include "Animation/Skeleton.h"
 #include "PreviewScene.h"
 #include "IDetailCustomization.h"
@@ -54,7 +54,7 @@ private:
 	// retarget source handler
 	TSharedPtr<IPropertyHandle> RetargetSourceNameHandler;
 
-	TSharedPtr<class SComboBox< TSharedPtr<FString> > > RetargetSourceComboBox;
+	TSharedPtr<class SSearchableComboBox> RetargetSourceComboBox;
 	TArray< TSharedPtr< FString > >						RetargetSourceComboList;
 
 	TSharedRef<SWidget> MakeRetargetSourceComboWidget( TSharedPtr<FString> InItem );
