Fixed crash in FSkeletalMeshComponentDetails::UpdateAnimationCategory when a custom skeletal mesh component hides the animation category via the HideCategories UCLASS specifier

#5802

diff --git a/Engine/Source/Editor/DetailCustomizations/Private/SkeletalMeshComponentDetails.cpp b/Engine/Source/Editor/DetailCustomizations/Private/SkeletalMeshComponentDetails.cpp
--- a/Engine/Source/Editor/DetailCustomizations/Private/SkeletalMeshComponentDetails.cpp
+++ b/Engine/Source/Editor/DetailCustomizations/Private/SkeletalMeshComponentDetails.cpp
@@ -11,6 +11,7 @@
 #include "Animation/AnimInstance.h"
 #include "Animation/AnimBlueprint.h"
 #include "Editor.h"
+#include "EditorCategoryUtils.h"
 #include "DetailLayoutBuilder.h"
 #include "IDetailPropertyRow.h"
 #include "DetailCategoryBuilder.h"
@@ -96,8 +97,14 @@ void FSkeletalMeshComponentDetails::CustomizeDetails(IDetailLayoutBuilder& Detai
 	UpdatePhysicsCategory(DetailBuilder);
 }
 
-void FSkeletalMeshComponentDetails::UpdateAnimationCategory( IDetailLayoutBuilder& DetailBuilder )
+void FSkeletalMeshComponentDetails::UpdateAnimationCategory(IDetailLayoutBuilder& DetailBuilder)
 {
+	// Custom skeletal mesh components may hide the animation category, so we won't assume it's visible
+	if (DetailBuilder.GetBaseClass() && FEditorCategoryUtils::IsCategoryHiddenFromClass(DetailBuilder.GetBaseClass(), "Animation"))
+	{
+		return;
+	}
+
 	UpdateSkeletonNameAndPickerVisibility();
 
 	IDetailCategoryBuilder& AnimationCategory = DetailBuilder.EditCategory("Animation", FText::GetEmpty(), ECategoryPriority::Important);
