Fix Landscape Detail Customization showing duplicate property

diff --git a/Engine/Source/Editor/DetailCustomizations/Private/DetailCustomizations.cpp b/Engine/Source/Editor/DetailCustomizations/Private/DetailCustomizations.cpp
--- a/Engine/Source/Editor/DetailCustomizations/Private/DetailCustomizations.cpp
+++ b/Engine/Source/Editor/DetailCustomizations/Private/DetailCustomizations.cpp
@@ -383,7 +383,6 @@ void FDetailCustomizationsModule::RegisterObjectCustomizations()
 	RegisterCustomClassLayout("MotionControllerComponent", FOnGetDetailCustomizationInstance::CreateStatic(&FMotionControllerDetails::MakeInstance));
 
 	RegisterCustomClassLayout("Landscape", FOnGetDetailCustomizationInstance::CreateStatic(&FLandscapeUIDetails::MakeInstance));
-	RegisterCustomClassLayout("LandscapeProxy", FOnGetDetailCustomizationInstance::CreateStatic(&FLandscapeUIDetails::MakeInstance));
 	RegisterCustomClassLayout("LandscapeSettings", FOnGetDetailCustomizationInstance::CreateStatic(&FLandscapeSettingsDetails::MakeInstance));
 }
 
diff --git a/Engine/Source/Editor/DetailCustomizations/Private/LandscapeUIDetails.cpp b/Engine/Source/Editor/DetailCustomizations/Private/LandscapeUIDetails.cpp
--- a/Engine/Source/Editor/DetailCustomizations/Private/LandscapeUIDetails.cpp
+++ b/Engine/Source/Editor/DetailCustomizations/Private/LandscapeUIDetails.cpp
@@ -59,29 +59,29 @@ void FLandscapeUIDetails::CustomizeDetails( IDetailLayoutBuilder& DetailBuilder
 			DetailBuilder.HideProperty(PropertyHandle);
 			if(GetMutableDefault<UEditorExperimentalSettings>()->bLandscapeLayerSystem)
 			{
-				IDetailCategoryBuilder& Category = DetailBuilder.EditCategory(PropertyHandle->GetDefaultCategoryName());
-				Category.AddCustomRow(FText(), true)
+				const FText DisplayAndFilterText(NSLOCTEXT("UnrealEd", "LandscapeToggleLayerName", "Enable Layer System"));
+				DetailBuilder.AddCustomRowToCategory(PropertyHandle, DisplayAndFilterText)
 				.NameContent()
 				[
-					PropertyHandle->CreatePropertyNameWidget(NSLOCTEXT("UnrealEd", "LandscapeToggleLayerName", "Enable Layer System"))
+					PropertyHandle->CreatePropertyNameWidget(DisplayAndFilterText)
 				]
 				.ValueContent()
-					[
-						SNew(SCheckBox)
-						.ToolTipText(NSLOCTEXT("UnrealEd", "LandscapeToggleLayerToolTip", "Toggle whether or not to support layers on this Landscape and its streaming proxies. Toggling this will clear the undo stack."))
-						.Type(ESlateCheckBoxType::CheckBox)
-						.IsChecked_Lambda([=]() 
-						{ 
-							return Landscape->CanHaveLayersContent() ? ECheckBoxState::Checked : ECheckBoxState::Unchecked;
-						})
-						.OnCheckStateChanged_Lambda([=](ECheckBoxState NewState) 
-						{ 
-							bool bChecked = (NewState == ECheckBoxState::Checked);
-							if (Landscape->CanHaveLayersContent() != bChecked)
-							{
-								ToggleCanHaveLayersContent(Landscape);
-							}
-						})
+				[
+					SNew(SCheckBox)
+					.ToolTipText(NSLOCTEXT("UnrealEd", "LandscapeToggleLayerToolTip", "Toggle whether or not to support layers on this Landscape and its streaming proxies. Toggling this will clear the undo stack."))
+					.Type(ESlateCheckBoxType::CheckBox)
+					.IsChecked_Lambda([=]() 
+					{ 
+						return Landscape->CanHaveLayersContent() ? ECheckBoxState::Checked : ECheckBoxState::Unchecked;
+					})
+					.OnCheckStateChanged_Lambda([=](ECheckBoxState NewState) 
+					{ 
+						bool bChecked = (NewState == ECheckBoxState::Checked);
+						if (Landscape->CanHaveLayersContent() != bChecked)
+						{
+							ToggleCanHaveLayersContent(Landscape);
+						}
+					})
 				];
 			}
 			
diff --git a/Engine/Source/Runtime/Landscape/Classes/Landscape.h b/Engine/Source/Runtime/Landscape/Classes/Landscape.h
--- a/Engine/Source/Runtime/Landscape/Classes/Landscape.h
+++ b/Engine/Source/Runtime/Landscape/Classes/Landscape.h
@@ -355,7 +355,7 @@ private:
 public:
 
 #if WITH_EDITORONLY_DATA
-	UPROPERTY(EditAnywhere, Category=Experimental, AdvancedDisplay)
+	UPROPERTY(EditAnywhere, Category=Experimental)
 	bool bCanHaveLayersContent = false;
 
 	DECLARE_EVENT(ALandscape, FLandscapeBPCustomBrushChangedDelegate);
