- Landscape: Add New Information Category in Details Customization that displays some important Landscape Details (Component Count, Resolution, Overall Resolution...)

#rnx

diff --git a/Engine/Source/Editor/DetailCustomizations/Private/DetailCustomizations.cpp b/Engine/Source/Editor/DetailCustomizations/Private/DetailCustomizations.cpp
--- a/Engine/Source/Editor/DetailCustomizations/Private/DetailCustomizations.cpp
+++ b/Engine/Source/Editor/DetailCustomizations/Private/DetailCustomizations.cpp
@@ -129,6 +129,7 @@
 #include "MotionControllerDetails.h"
 #include "MotionControllerPinFactory.h"
 #include "LandscapeUIDetails.h"
+#include "LandscapeProxyUIDetails.h"
 #include "LandscapeSettingsDetails.h"
 #include "PerPlatformPropertyCustomization.h"
 #include "SkeletalMeshReductionSettingsDetails.h"
@@ -383,6 +384,7 @@ void FDetailCustomizationsModule::RegisterObjectCustomizations()
 	RegisterCustomClassLayout("MotionControllerComponent", FOnGetDetailCustomizationInstance::CreateStatic(&FMotionControllerDetails::MakeInstance));
 
 	RegisterCustomClassLayout("Landscape", FOnGetDetailCustomizationInstance::CreateStatic(&FLandscapeUIDetails::MakeInstance));
+	RegisterCustomClassLayout("LandscapeProxy", FOnGetDetailCustomizationInstance::CreateStatic(&FLandscapeProxyUIDetails::MakeInstance));
 	RegisterCustomClassLayout("LandscapeSettings", FOnGetDetailCustomizationInstance::CreateStatic(&FLandscapeSettingsDetails::MakeInstance));
 }
 
diff --git a/Engine/Source/Editor/DetailCustomizations/Private/LandscapeProxyUIDetails.cpp b/Engine/Source/Editor/DetailCustomizations/Private/LandscapeProxyUIDetails.cpp
new file mode 100644
--- /dev/null
+++ b/Engine/Source/Editor/DetailCustomizations/Private/LandscapeProxyUIDetails.cpp
@@ -0,0 +1,125 @@
+// Copyright 1998-2019 Epic Games, Inc. All Rights Reserved.
+
+#include "LandscapeProxyUIDetails.h"
+#include "Widgets/DeclarativeSyntaxSupport.h"
+#include "PropertyHandle.h"
+#include "DetailLayoutBuilder.h"
+#include "Runtime/Landscape/Classes/Landscape.h"
+#include "Runtime/Landscape/Classes/LandscapeInfo.h"
+#include "Settings/EditorExperimentalSettings.h"
+#include "DetailCategoryBuilder.h"
+#include "DetailWidgetRow.h"
+#include "Widgets/Input/SButton.h"
+#include "Widgets/Text/STextBlock.h"
+#include "Widgets/Input/SCheckBox.h"
+#include "Misc/MessageDialog.h"
+#include "Editor.h"
+
+#define LOCTEXT_NAMESPACE "FLandscapeProxyUIDetails"
+
+FLandscapeProxyUIDetails::FLandscapeProxyUIDetails()
+{
+}
+
+TSharedRef<IDetailCustomization> FLandscapeProxyUIDetails::MakeInstance()
+{
+	return MakeShareable( new FLandscapeProxyUIDetails);
+}
+
+void FLandscapeProxyUIDetails::CustomizeDetails( IDetailLayoutBuilder& DetailBuilder )
+{
+	TArray<TWeakObjectPtr<UObject>> EditingObjects;
+	DetailBuilder.GetObjectsBeingCustomized(EditingObjects);
+
+	auto GenerateTextWidget = [](const FText& InText, bool bInBold = false) -> TSharedRef<SWidget>
+	{
+		return SNew(STextBlock)
+				.Font(bInBold? IDetailLayoutBuilder::GetDetailFontBold() : IDetailLayoutBuilder::GetDetailFont())
+				.Text(InText);
+	};
+
+	if (EditingObjects.Num() == 1)
+	{
+		if (ALandscapeProxy* Proxy = Cast<ALandscapeProxy>(EditingObjects[0]))
+		{
+			if (ULandscapeInfo* LandscapeInfo = Proxy->GetLandscapeInfo())
+			{
+				IDetailCategoryBuilder& CategoryBuilder = DetailBuilder.EditCategory("Information", FText::GetEmpty(), ECategoryPriority::Important);
+				
+				FText RowDisplayText = LOCTEXT("LandscapeComponentResolution", "Component Resolution (Verts)");
+				CategoryBuilder.AddCustomRow(RowDisplayText)
+				.NameContent()
+				[
+					GenerateTextWidget(RowDisplayText)
+				]
+				.ValueContent()
+				[
+					GenerateTextWidget(FText::Format(LOCTEXT("LandscapeComponentResolutionValue", "{0} x {0}"), Proxy->ComponentSizeQuads+1), true) // Verts
+				];
+
+				RowDisplayText = LOCTEXT("LandscapeComponentCount", "Component Count");
+				CategoryBuilder.AddCustomRow(RowDisplayText)
+				.NameContent()
+				[
+					GenerateTextWidget(RowDisplayText)
+				]
+				.ValueContent()
+				[
+					GenerateTextWidget(FText::Format(LOCTEXT("LandscapeComponentCountValue", "{0}"), Proxy->LandscapeComponents.Num()), true)
+				];
+
+				RowDisplayText = LOCTEXT("LandscapeComponentSubsections", "Component Subsections");
+				CategoryBuilder.AddCustomRow(RowDisplayText)
+				.NameContent()
+				[
+					GenerateTextWidget(RowDisplayText)
+				]
+				.ValueContent()
+				[
+					GenerateTextWidget(FText::Format(LOCTEXT("LandscapeComponentSubSectionsValue", "{0} x {0}"), Proxy->NumSubsections), true)
+				];
+
+				FIntRect Rect = Proxy->GetBoundingRect();
+				FIntPoint Size = Rect.Size();
+				RowDisplayText = LOCTEXT("LandscapeResolution", "Resolution (Verts)");
+				CategoryBuilder.AddCustomRow(RowDisplayText)
+				.NameContent()
+				[
+					GenerateTextWidget(RowDisplayText)
+				]
+				.ValueContent()
+				[
+					GenerateTextWidget(FText::Format(LOCTEXT("LandscapeResolutionValue", "{0} x {1}"), Size.X+1, Size.Y+1), true)
+				];
+
+				int32 LandscapeCount = LandscapeInfo->Proxies.Num() + (LandscapeInfo->LandscapeActor.Get() ? 1 : 0);
+				RowDisplayText = LOCTEXT("LandscapeCount", "Landscape Count");
+				CategoryBuilder.AddCustomRow(RowDisplayText)
+				.NameContent()
+				[
+					GenerateTextWidget(RowDisplayText)
+				]
+				.ValueContent()
+				[
+					GenerateTextWidget(FText::Format(LOCTEXT("LandscapeCountValue", "{0}"), LandscapeCount), true)
+				];
+
+				LandscapeInfo->GetLandscapeExtent(Rect.Min.X, Rect.Min.Y, Rect.Max.X, Rect.Max.Y);
+				Size = Rect.Size();
+				RowDisplayText = LOCTEXT("LandscapeOverallResolution", "Overall Resolution (Verts)");
+				CategoryBuilder.AddCustomRow(RowDisplayText)
+				.NameContent()
+				[
+					GenerateTextWidget(RowDisplayText)
+				]
+				.ValueContent()
+				[
+					GenerateTextWidget(FText::Format(LOCTEXT("LandscapeOveralResolutionValue", "{0} x {1}"), Size.X + 1, Size.Y + 1), true)
+				];
+			}
+
+		}
+	}
+}
+
+#undef LOCTEXT_NAMESPACE
diff --git a/Engine/Source/Editor/DetailCustomizations/Private/LandscapeProxyUIDetails.h b/Engine/Source/Editor/DetailCustomizations/Private/LandscapeProxyUIDetails.h
new file mode 100644
--- /dev/null
+++ b/Engine/Source/Editor/DetailCustomizations/Private/LandscapeProxyUIDetails.h
@@ -0,0 +1,29 @@
+// Copyright 1998-2019 Epic Games, Inc. All Rights Reserved.
+
+#pragma once
+
+#include "CoreMinimal.h"
+#include "UObject/WeakObjectPtr.h"
+#include "Types/SlateEnums.h"
+#include "IDetailCustomization.h"
+#include "Internationalization/Text.h"
+
+class IDetailLayoutBuilder;
+class ALandscapeProxy;
+
+class FLandscapeProxyUIDetails : public IDetailCustomization
+{
+public:
+	/** Makes a new instance of this detail layout class for a specific detail view requesting it */
+	static TSharedRef<IDetailCustomization> MakeInstance();
+
+	/** IDetailCustomization interface */
+	virtual void CustomizeDetails( IDetailLayoutBuilder& DetailBuilder ) override;
+	/** End IDetailCustomization interface */
+
+private:
+	/** Use MakeInstance to create an instance of this class */
+	FLandscapeProxyUIDetails();
+	
+	IDetailLayoutBuilder* DetailLayoutBuilder = nullptr;
+};
diff --git a/Engine/Source/Editor/DetailCustomizations/Private/LandscapeUIDetails.cpp b/Engine/Source/Editor/DetailCustomizations/Private/LandscapeUIDetails.cpp
--- a/Engine/Source/Editor/DetailCustomizations/Private/LandscapeUIDetails.cpp
+++ b/Engine/Source/Editor/DetailCustomizations/Private/LandscapeUIDetails.cpp
@@ -15,6 +15,8 @@
 #include "Misc/MessageDialog.h"
 #include "Editor.h"
 
+#define LOCTEXT_NAMESPACE "FLandscapeUIDetails"
+
 FLandscapeUIDetails::FLandscapeUIDetails()
 {
 }
@@ -40,51 +42,51 @@ void FLandscapeUIDetails::CustomizeDetails( IDetailLayoutBuilder& DetailBuilder
 	if (EditingObjects.Num() == 1)
 	{
 		ALandscape* Landscape = Cast<ALandscape>(EditingObjects[0].Get());
+		if (Landscape == nullptr)
+		{
+			return;
+		}
 
-		if (Landscape != nullptr && Landscape->NumSubsections == 1)
+		if (Landscape->NumSubsections == 1)
 		{
 			TSharedRef<IPropertyHandle> ComponentScreenSizeToUseSubSectionsProp = DetailBuilder.GetProperty(GET_MEMBER_NAME_CHECKED(ALandscapeProxy, ComponentScreenSizeToUseSubSections));
 			DetailBuilder.HideProperty(ComponentScreenSizeToUseSubSectionsProp);
 		}
 
-		if (Landscape != nullptr)
+		if (!DetailLayoutBuilder)
 		{
-			if (!DetailLayoutBuilder)
-			{
-				DetailLayoutBuilder = &DetailBuilder;
-				GetMutableDefault<UEditorExperimentalSettings>()->OnSettingChanged().AddSP(this, &FLandscapeUIDetails::OnEditorExperimentalSettingsChanged);
-			}
-
-			TSharedRef<IPropertyHandle> PropertyHandle = DetailBuilder.GetProperty(GET_MEMBER_NAME_CHECKED(ALandscape, bCanHaveLayersContent));
-			DetailBuilder.HideProperty(PropertyHandle);
-			if(GetMutableDefault<UEditorExperimentalSettings>()->bLandscapeLayerSystem)
-			{
-				const FText DisplayAndFilterText(NSLOCTEXT("UnrealEd", "LandscapeToggleLayerName", "Enable Layer System"));
-				DetailBuilder.AddCustomRowToCategory(PropertyHandle, DisplayAndFilterText)
-				.NameContent()
-				[
-					PropertyHandle->CreatePropertyNameWidget(DisplayAndFilterText)
-				]
-				.ValueContent()
-				[
-					SNew(SCheckBox)
-					.ToolTipText(NSLOCTEXT("UnrealEd", "LandscapeToggleLayerToolTip", "Toggle whether or not to support layers on this Landscape and its streaming proxies. Toggling this will clear the undo stack."))
-					.Type(ESlateCheckBoxType::CheckBox)
-					.IsChecked_Lambda([=]() 
-					{ 
-						return Landscape->CanHaveLayersContent() ? ECheckBoxState::Checked : ECheckBoxState::Unchecked;
-					})
-					.OnCheckStateChanged_Lambda([=](ECheckBoxState NewState) 
-					{ 
-						bool bChecked = (NewState == ECheckBoxState::Checked);
-						if (Landscape->CanHaveLayersContent() != bChecked)
-						{
-							ToggleCanHaveLayersContent(Landscape);
-						}
-					})
-				];
-			}
-			
+			DetailLayoutBuilder = &DetailBuilder;
+			GetMutableDefault<UEditorExperimentalSettings>()->OnSettingChanged().AddSP(this, &FLandscapeUIDetails::OnEditorExperimentalSettingsChanged);
+		}
+			   
+		TSharedRef<IPropertyHandle> PropertyHandle = DetailBuilder.GetProperty(GET_MEMBER_NAME_CHECKED(ALandscape, bCanHaveLayersContent));
+		DetailBuilder.HideProperty(PropertyHandle);
+		if (GetMutableDefault<UEditorExperimentalSettings>()->bLandscapeLayerSystem)
+		{
+			const FText DisplayAndFilterText(LOCTEXT("LandscapeToggleLayerName", "Enable Layer System"));
+			DetailBuilder.AddCustomRowToCategory(PropertyHandle, DisplayAndFilterText)
+			.NameContent()
+			[
+				PropertyHandle->CreatePropertyNameWidget(DisplayAndFilterText)
+			]
+			.ValueContent()
+			[
+				SNew(SCheckBox)
+				.ToolTipText(LOCTEXT("LandscapeToggleLayerToolTip", "Toggle whether or not to support layers on this Landscape and its streaming proxies. Toggling this will clear the undo stack."))
+				.Type(ESlateCheckBoxType::CheckBox)
+				.IsChecked_Lambda([=]()
+				{
+					return Landscape->CanHaveLayersContent() ? ECheckBoxState::Checked : ECheckBoxState::Unchecked;
+			    })
+			    .OnCheckStateChanged_Lambda([=](ECheckBoxState NewState)
+			    {
+			        bool bChecked = (NewState == ECheckBoxState::Checked);
+				    if (Landscape->CanHaveLayersContent() != bChecked)
+				    {
+					    ToggleCanHaveLayersContent(Landscape);
+				    }
+			    })
+			];
 		}
 	}
 }
@@ -120,11 +122,11 @@ void FLandscapeUIDetails::ToggleCanHaveLayersContent(ALandscape* Landscape)
 
 		if (bHasHiddenLayers)
 		{
-			Reason = NSLOCTEXT("UnrealEd", "LandscapeDisableLayers_HiddenLayers", "Are you sure you want to disable the layers system?\n\nDoing so, will result in losing the data stored for each layers, but the current visual output will be kept. Be aware that some layers are currently hidden, continuing will result in their data being lost. Undo/redo buffer will also be cleared.");
+			Reason = LOCTEXT("LandscapeDisableLayers_HiddenLayers", "Are you sure you want to disable the layers system?\n\nDoing so, will result in losing the data stored for each layers, but the current visual output will be kept. Be aware that some layers are currently hidden, continuing will result in their data being lost. Undo/redo buffer will also be cleared.");
 		}
 		else
 		{
-			Reason = NSLOCTEXT("UnrealEd", "LandscapeDisableLayers", "Are you sure you want to disable the layers system?\n\nDoing so, will result in losing the data stored for each layers, but the current visual output will be kept. Undo/redo buffer will also be cleared.");
+			Reason = LOCTEXT("LandscapeDisableLayers", "Are you sure you want to disable the layers system?\n\nDoing so, will result in losing the data stored for each layers, but the current visual output will be kept. Undo/redo buffer will also be cleared.");
 		}
 
 		bToggled = FMessageDialog::Open(EAppMsgType::YesNo, Reason) == EAppReturnType::Yes;
@@ -132,7 +134,7 @@ void FLandscapeUIDetails::ToggleCanHaveLayersContent(ALandscape* Landscape)
 	else
 	{
 		
-		bToggled = FMessageDialog::Open(EAppMsgType::YesNo, NSLOCTEXT("UnrealEd", "LandscapeEnableLayers", "Are you sure you want to enable the layer system on this landscape and streaming proxies? Doing so will clear the undo/redo buffer.")) == EAppReturnType::Yes;
+		bToggled = FMessageDialog::Open(EAppMsgType::YesNo, LOCTEXT("LandscapeEnableLayers", "Are you sure you want to enable the layer system on this landscape and streaming proxies? Doing so will clear the undo/redo buffer.")) == EAppReturnType::Yes;
 	}
 
 	if (bToggled)
@@ -140,7 +142,9 @@ void FLandscapeUIDetails::ToggleCanHaveLayersContent(ALandscape* Landscape)
 		Landscape->ToggleCanHaveLayersContent();
 		if (GEditor)
 		{
-			GEditor->ResetTransaction(NSLOCTEXT("UnrealEd", "ToggleLanscapeLayers", "Toggling Landscape Layers"));
+			GEditor->ResetTransaction(LOCTEXT("ToggleLanscapeLayers", "Toggling Landscape Layers"));
 		}
 	}
-}
\ No newline at end of file
+}
+
+#undef LOCTEXT_NAMESPACE
\ No newline at end of file
