Fix water body actor generated mesh components being left behind when duplicated
Fix water spline metadata not appearing when selecting a key

diff --git a/Engine/Source/Editor/ComponentVisualizers/Public/SplineComponentVisualizer.h b/Engine/Source/Editor/ComponentVisualizers/Public/SplineComponentVisualizer.h
--- a/Engine/Source/Editor/ComponentVisualizers/Public/SplineComponentVisualizer.h
+++ b/Engine/Source/Editor/ComponentVisualizers/Public/SplineComponentVisualizer.h
@@ -92,7 +92,7 @@ public:
 	/** Get the spline component we are currently editing */
 	USplineComponent* GetEditedSplineComponent() const;
 
-	TSet<int32> GetSelectedKeys() const { return SelectedKeys; }
+	const TSet<int32>& GetSelectedKeys() const { return SelectedKeys; }
 
 protected:
 
diff --git a/Engine/Source/Editor/DetailCustomizations/Private/SplineComponentDetails.cpp b/Engine/Source/Editor/DetailCustomizations/Private/SplineComponentDetails.cpp
--- a/Engine/Source/Editor/DetailCustomizations/Private/SplineComponentDetails.cpp
+++ b/Engine/Source/Editor/DetailCustomizations/Private/SplineComponentDetails.cpp
@@ -38,7 +38,7 @@ USplineMetadataDetailsFactoryBase::USplineMetadataDetailsFactoryBase(const FObje
 class FSplinePointDetails : public IDetailCustomNodeBuilder, public TSharedFromThis<FSplinePointDetails>
 {
 public:
-	FSplinePointDetails();
+	FSplinePointDetails(USplineComponent* InOwningSplineComponent);
 
 	//~ Begin IDetailCustomNodeBuilder interface
 	virtual void SetOnRebuildChildren(FSimpleDelegate InOnRegenerateChildren) override;
@@ -193,11 +193,13 @@ private:
 	UProperty* SplineCurvesProperty;
 	TArray<TSharedPtr<FString>> SplinePointTypes;
 	TSharedPtr<ISplineMetadataDetails> SplineMetaDataDetails;
+	FSimpleDelegate OnRegenerateChildren;
 };
 
-FSplinePointDetails::FSplinePointDetails()
+FSplinePointDetails::FSplinePointDetails(USplineComponent* InOwningSplineComponent)
 	: SplineComp(nullptr)
 {
+
 	TSharedPtr<FComponentVisualizer> Visualizer = GUnrealEd->FindComponentVisualizer(USplineComponent::StaticClass());
 	SplineVisualizer = (FSplineComponentVisualizer*)Visualizer.Get();
 	check(SplineVisualizer);
@@ -210,10 +212,13 @@ FSplinePointDetails::FSplinePointDetails()
 	{
 		SplinePointTypes.Add(MakeShareable(new FString(SplinePointTypeEnum->GetNameStringByIndex(EnumIndex))));
 	}
+
+	SplineComp = InOwningSplineComponent;
 }
 
 void FSplinePointDetails::SetOnRebuildChildren(FSimpleDelegate InOnRegenerateChildren)
 {
+	OnRegenerateChildren = InOnRegenerateChildren;
 }
 
 void FSplinePointDetails::GenerateHeaderRowContent(FDetailWidgetRow& NodeRow)
@@ -421,27 +426,19 @@ void FSplinePointDetails::GenerateChildContent(IDetailChildrenBuilder& ChildrenB
 		]
 	];
 
-	TArray<TWeakObjectPtr<UObject>> SelectedObjects;
-	ChildrenBuilder.GetParentCategory().GetParentLayout().GetObjectsBeingCustomized(SelectedObjects);
-	if (SelectedObjects.Num() == 1)
+	if (SplineComp && SplineVisualizer && SplineVisualizer->GetSelectedKeys().Num() > 0)
 	{
-		if (USplineComponent* SelectedSplineComp = Cast<USplineComponent>(SelectedObjects[0]))
+		for (TObjectIterator<UClass> ClassIterator; ClassIterator; ++ClassIterator)
 		{
-			if (SelectedSplineComp->GetSplinePointsMetadata())
+			if (ClassIterator->IsChildOf(USplineMetadataDetailsFactoryBase::StaticClass()) && !ClassIterator->HasAnyClassFlags(CLASS_Abstract | CLASS_Deprecated | CLASS_NewerVersionExists))
 			{
-				for (TObjectIterator<UClass> ClassIterator; ClassIterator; ++ClassIterator)
+				USplineMetadataDetailsFactoryBase* Factory = ClassIterator->GetDefaultObject<USplineMetadataDetailsFactoryBase>();
+				if (Factory->GetMetadataClass() == SplineComp->GetSplinePointsMetadata()->GetClass())
 				{
-					if (ClassIterator->IsChildOf(USplineMetadataDetailsFactoryBase::StaticClass()) && !ClassIterator->HasAnyClassFlags(CLASS_Abstract | CLASS_Deprecated | CLASS_NewerVersionExists))
-					{
-						USplineMetadataDetailsFactoryBase* Factory = ClassIterator->GetDefaultObject<USplineMetadataDetailsFactoryBase>();
-						if (Factory->GetMetadataClass() == SelectedSplineComp->GetSplinePointsMetadata()->GetClass())
-						{
-							SplineMetaDataDetails = Factory->Create();
-							IDetailGroup& Group = ChildrenBuilder.AddGroup(SplineMetaDataDetails->GetName(), SplineMetaDataDetails->GetDisplayName());
-							SplineMetaDataDetails->GenerateChildContent(Group);
-							break;
-						}
-					}
+					SplineMetaDataDetails = Factory->Create();
+					IDetailGroup& Group = ChildrenBuilder.AddGroup(SplineMetaDataDetails->GetName(), SplineMetaDataDetails->GetDisplayName());
+					SplineMetaDataDetails->GenerateChildContent(Group);
+					break;
 				}
 			}
 		}
@@ -456,8 +453,16 @@ void FSplinePointDetails::Tick(float DeltaTime)
 void FSplinePointDetails::UpdateValues()
 {
 	SplineComp = SplineVisualizer->GetEditedSplineComponent();
-	SelectedKeys = SplineVisualizer->GetSelectedKeys();
-		
+
+	bool bNeedsRebuild = false;
+	const TSet<int32>& NewSelectedKeys = SplineVisualizer->GetSelectedKeys();
+
+	if (NewSelectedKeys.Num() != SelectedKeys.Num())
+	{
+		bNeedsRebuild = true;
+	}
+	SelectedKeys = NewSelectedKeys;
+
 	// Cache values to be shown by the details customization.
 	// An unset optional value represents 'multiple values' (in the case where multiple points are selected).
 	InputKey.Reset();
@@ -486,6 +491,11 @@ void FSplinePointDetails::UpdateValues()
 	{
 		SplineMetaDataDetails->Update(SplineComp, SelectedKeys);
 	}
+
+	if (bNeedsRebuild)
+	{
+		OnRegenerateChildren.ExecuteIfBound();
+	}
 }
 
 FName FSplinePointDetails::GetName() const
@@ -740,10 +750,20 @@ void FSplineComponentDetails::CustomizeDetails(IDetailLayoutBuilder& DetailBuild
 	// Hide the SplineCurves property
 	TSharedPtr<IPropertyHandle> SplineCurvesProperty = DetailBuilder.GetProperty(GET_MEMBER_NAME_CHECKED(USplineComponent, SplineCurves));
 	SplineCurvesProperty->MarkHiddenByCustomization();
+	 
 
-	IDetailCategoryBuilder& Category = DetailBuilder.EditCategory("Selected Points");
-	TSharedRef<FSplinePointDetails> SplinePointDetails = MakeShareable(new FSplinePointDetails);
-	Category.AddCustomBuilder(SplinePointDetails);
+	TArray<TWeakObjectPtr<UObject>> ObjectsBeingCustomized;
+	DetailBuilder.GetObjectsBeingCustomized(ObjectsBeingCustomized);
+
+	if (ObjectsBeingCustomized.Num() == 1)
+	{
+		if (USplineComponent* SplineComp = Cast<USplineComponent>(ObjectsBeingCustomized[0]))
+		{
+			IDetailCategoryBuilder& Category = DetailBuilder.EditCategory("Selected Points");
+			TSharedRef<FSplinePointDetails> SplinePointDetails = MakeShareable(new FSplinePointDetails(SplineComp));
+			Category.AddCustomBuilder(SplinePointDetails);
+		}
+	}
 }
 
 #undef LOCTEXT_NAMESPACE
