Fix crash setting struct type in the defaults of a blueprint MovieSceneEventParameters property

diff --git a/Engine/Source/Editor/DetailCustomizations/Private/MovieSceneEventParametersCustomization.cpp b/Engine/Source/Editor/DetailCustomizations/Private/MovieSceneEventParametersCustomization.cpp
--- a/Engine/Source/Editor/DetailCustomizations/Private/MovieSceneEventParametersCustomization.cpp
+++ b/Engine/Source/Editor/DetailCustomizations/Private/MovieSceneEventParametersCustomization.cpp
@@ -107,7 +107,7 @@ void FMovieSceneEventParametersCustomization::OnStructChanged(const FAssetData&
 	FScopedTransaction Transaction(LOCTEXT("SetParameterStructText", "Set Parameter Structure"));
 
 	FEditPropertyChain PropertyChain;
-	PropertyChain.SetActivePropertyNode(PropertyHandle->GetProperty());
+	PropertyChain.AddHead(PropertyHandle->GetProperty());
 
 	// Fire off the pre-notify
 	FNotifyHook* Hook = PropertyUtilities->GetNotifyHook();
