Fix for crash when metadata is undefined in pre-existing spline components.

diff --git a/Engine/Source/Editor/DetailCustomizations/Private/SplineComponentDetails.cpp b/Engine/Source/Editor/DetailCustomizations/Private/SplineComponentDetails.cpp
--- a/Engine/Source/Editor/DetailCustomizations/Private/SplineComponentDetails.cpp
+++ b/Engine/Source/Editor/DetailCustomizations/Private/SplineComponentDetails.cpp
@@ -433,7 +433,8 @@ void FSplinePointDetails::GenerateChildContent(IDetailChildrenBuilder& ChildrenB
 			if (ClassIterator->IsChildOf(USplineMetadataDetailsFactoryBase::StaticClass()) && !ClassIterator->HasAnyClassFlags(CLASS_Abstract | CLASS_Deprecated | CLASS_NewerVersionExists))
 			{
 				USplineMetadataDetailsFactoryBase* Factory = ClassIterator->GetDefaultObject<USplineMetadataDetailsFactoryBase>();
-				if (Factory->GetMetadataClass() == SplineComp->GetSplinePointsMetadata()->GetClass())
+				USplineMetadata* SplinePointsMetadata = SplineComp->GetSplinePointsMetadata();
+				if (SplinePointsMetadata != nullptr && Factory->GetMetadataClass() == SplinePointsMetadata->GetClass())
 				{
 					SplineMetaDataDetails = Factory->Create();
 					IDetailGroup& Group = ChildrenBuilder.AddGroup(SplineMetaDataDetails->GetName(), SplineMetaDataDetails->GetDisplayName());
