Fix post process settings struct not showing their property name

diff --git a/Engine/Source/Editor/AdvancedPreviewScene/Public/AssetViewerSettings.h b/Engine/Source/Editor/AdvancedPreviewScene/Public/AssetViewerSettings.h
--- a/Engine/Source/Editor/AdvancedPreviewScene/Public/AssetViewerSettings.h
+++ b/Engine/Source/Editor/AdvancedPreviewScene/Public/AssetViewerSettings.h
@@ -98,7 +98,7 @@ struct FPreviewSceneProfile
 	bool bPostProcessingEnabled;
 
 	/** Manual set post processing settings */
-	UPROPERTY(EditAnywhere, config, Category = PostProcessing, AdvancedDisplay)
+	UPROPERTY(EditAnywhere, config, Category = PostProcessing, AdvancedDisplay, meta=(ShowOnlyInnerProperties))
 	FPostProcessSettings PostProcessingSettings;
 
 	/** Current rotation value of the sky in degrees (0 - 360) */
diff --git a/Engine/Source/Editor/DetailCustomizations/Private/PostProcessSettingsCustomization.cpp b/Engine/Source/Editor/DetailCustomizations/Private/PostProcessSettingsCustomization.cpp
--- a/Engine/Source/Editor/DetailCustomizations/Private/PostProcessSettingsCustomization.cpp
+++ b/Engine/Source/Editor/DetailCustomizations/Private/PostProcessSettingsCustomization.cpp
@@ -25,6 +25,9 @@
 
 #define LOCTEXT_NAMESPACE "PostProcessSettingsCustomization"
 
+const FName ShowPostProcessCategoriesName("ShowPostProcessCategories");
+const FName ShowOnlyInnerPropertiesName("ShowOnlyInnerProperties");
+
 struct FCategoryOrGroup
 {
 	IDetailCategoryBuilder* Category;
@@ -126,7 +129,7 @@ void FPostProcessSettingsCustomization::CustomizeChildren( TSharedRef<IPropertyH
 	const bool bExtendedLuminanceRange = VarDefaultAutoExposureExtendDefaultLuminanceRange->GetValueOnGameThread() == 1;
 	static const FName ExposureCategory("Lens|Exposure");
 
-	static const FName ShowPostProcessCategoriesName("ShowPostProcessCategories");
+
 
 	bool bShowPostProcessCategories = StructPropertyHandle->HasMetaData(ShowPostProcessCategoriesName);
 
@@ -278,7 +281,19 @@ PRAGMA_ENABLE_OPTIMIZATION
 
 void FPostProcessSettingsCustomization::CustomizeHeader( TSharedRef<IPropertyHandle> StructPropertyHandle, class FDetailWidgetRow& HeaderRow, IPropertyTypeCustomizationUtils& StructCustomizationUtils )
 {
-	// No header
+	bool bShowHeader = !StructPropertyHandle->HasMetaData(ShowPostProcessCategoriesName) && !StructPropertyHandle->HasMetaData(ShowOnlyInnerPropertiesName);
+	if(bShowHeader)
+	{
+		HeaderRow.NameContent()
+		[
+			StructPropertyHandle->CreatePropertyNameWidget()
+		];
+
+		HeaderRow.ValueContent()
+		[
+			StructPropertyHandle->CreatePropertyValueWidget()
+		];
+	}
 }
 
 void FWeightedBlendableCustomization::AddDirectAsset(TSharedRef<IPropertyHandle> StructPropertyHandle, TSharedPtr<IPropertyHandle> Weight, TSharedPtr<IPropertyHandle> Value, UClass* Class)
diff --git a/Engine/Source/Runtime/Engine/Classes/Camera/CameraComponent.h b/Engine/Source/Runtime/Engine/Classes/Camera/CameraComponent.h
--- a/Engine/Source/Runtime/Engine/Classes/Camera/CameraComponent.h
+++ b/Engine/Source/Runtime/Engine/Classes/Camera/CameraComponent.h
@@ -198,7 +198,7 @@ public:
 public:
 
 	/** Post process settings to use for this camera. Don't forget to check the properties you want to override */
-	UPROPERTY(Interp, BlueprintReadWrite, Category = PostProcess)
+	UPROPERTY(Interp, BlueprintReadWrite, Category = PostProcess, meta=(ShowOnlyInnerProperties))
 	struct FPostProcessSettings PostProcessSettings;
 
 #if WITH_EDITORONLY_DATA
