UE4.21 - Fix missing packaging names
#rnx

diff --git a/Engine/Source/Editor/DetailCustomizations/Private/WindowsTargetSettingsDetails.cpp b/Engine/Source/Editor/DetailCustomizations/Private/WindowsTargetSettingsDetails.cpp
--- a/Engine/Source/Editor/DetailCustomizations/Private/WindowsTargetSettingsDetails.cpp
+++ b/Engine/Source/Editor/DetailCustomizations/Private/WindowsTargetSettingsDetails.cpp
@@ -63,11 +63,11 @@ static FText GetFriendlyNameFromRHIName(const FString& InRHIName)
 	}
 	else if (InRHIName == TEXT("GLSL_150_ES2"))
 	{
-		FriendlyRHIName = LOCTEXT("OpenGL3ES2", "OpenGL 3 (ES2)");
+		FriendlyRHIName = LOCTEXT("OpenGL3ES2", "OpenGL 3 Mobile Emulation (ES2)");
 	}
 	else if (InRHIName == TEXT("GLSL_150_ES31"))
 	{
-		FriendlyRHIName = LOCTEXT("OpenGL3ES31", "OpenGL 3 (ES3.1, Experimental)");
+		FriendlyRHIName = LOCTEXT("OpenGL3ES31", "OpenGL 3 Mobile Emulation (ES3.1, Experimental)");
 	}
 	else if (InRHIName == TEXT("GLSL_430"))
 	{
@@ -75,7 +75,7 @@ static FText GetFriendlyNameFromRHIName(const FString& InRHIName)
 	}
 	else if (InRHIName == TEXT("SF_VULKAN_ES31_ANDROID") || InRHIName == TEXT("SF_VULKAN_ES31"))
 	{
-		FriendlyRHIName = LOCTEXT("Vulkan ES31", "Vulkan Mobile (ES3.1, Experimental)");
+		FriendlyRHIName = LOCTEXT("Vulkan ES31", "Vulkan Mobile (ES3.1)");
 	}
 	else if (InRHIName == TEXT("SF_VULKAN_SM4"))
 	{
@@ -93,6 +93,14 @@ static FText GetFriendlyNameFromRHIName(const FString& InRHIName)
 	{
 		FriendlyRHIName = LOCTEXT("SwitchForward", "Switch (Forward)");
 	}
+	else if (InRHIName == TEXT("PCD3D_ES2"))
+	{
+		FriendlyRHIName = LOCTEXT("DirectX11_ES2", "DirectX 11 Mobile Emulation (ES2)");
+	}
+	else if (InRHIName == TEXT("PCD3D_ES31"))
+	{
+		FriendlyRHIName = LOCTEXT("DirectX11_ES31", "DirectX 11 Mobile Emulation (ES3.1)");
+	}
 
 	return FriendlyRHIName;
 }
