UE4.21 - Remove invalid Windows target APIs
#rnx

diff --git a/Engine/Source/Editor/DetailCustomizations/Private/WindowsTargetSettingsDetails.cpp b/Engine/Source/Editor/DetailCustomizations/Private/WindowsTargetSettingsDetails.cpp
--- a/Engine/Source/Editor/DetailCustomizations/Private/WindowsTargetSettingsDetails.cpp
+++ b/Engine/Source/Editor/DetailCustomizations/Private/WindowsTargetSettingsDetails.cpp
@@ -27,6 +27,7 @@
 #include "Interfaces/ITargetPlatform.h"
 #include "Interfaces/ITargetPlatformModule.h"
 #include "SExternalImageReference.h"
+#include "UnrealEngine.h"
 
 #if WITH_ENGINE
 #include "AudioDevice.h"
@@ -48,7 +49,7 @@ namespace WindowsTargetSettingsDetailsConstants
 
 static FText GetFriendlyNameFromRHIName(const FString& InRHIName)
 {
-	FText FriendlyRHIName = LOCTEXT("UnknownRHI", "UnknownRHI");
+	FText FriendlyRHIName;
 	if (InRHIName == TEXT("PCD3D_SM5"))
 	{
 		FriendlyRHIName = LOCTEXT("DirectX11", "DirectX 11 (SM5)");
@@ -57,30 +58,10 @@ static FText GetFriendlyNameFromRHIName(const FString& InRHIName)
 	{
 		FriendlyRHIName = LOCTEXT("DirectX10", "DirectX 10 (SM4)");
 	}
-	else if (InRHIName == TEXT("GLSL_150"))
-	{
-		FriendlyRHIName = LOCTEXT("OpenGL3", "OpenGL 3 (SM4)");
-	}
-	else if (InRHIName == TEXT("GLSL_150_ES2"))
-	{
-		FriendlyRHIName = LOCTEXT("OpenGL3ES2", "OpenGL 3 Mobile Emulation (ES2)");
-	}
-	else if (InRHIName == TEXT("GLSL_150_ES31"))
-	{
-		FriendlyRHIName = LOCTEXT("OpenGL3ES31", "OpenGL 3 Mobile Emulation (ES3.1, Experimental)");
-	}
 	else if (InRHIName == TEXT("GLSL_430"))
 	{
 		FriendlyRHIName = LOCTEXT("OpenGL4", "OpenGL 4 (SM5, Experimental)");
 	}
-	else if (InRHIName == TEXT("SF_VULKAN_ES31_ANDROID") || InRHIName == TEXT("SF_VULKAN_ES31"))
-	{
-		FriendlyRHIName = LOCTEXT("Vulkan ES31", "Vulkan Mobile (ES3.1)");
-	}
-	else if (InRHIName == TEXT("SF_VULKAN_SM4"))
-	{
-		FriendlyRHIName = LOCTEXT("VulkanSM4", "Vulkan Desktop (SM4, Experimental)");
-	}
 	else if (InRHIName == TEXT("SF_VULKAN_SM5"))
 	{
 		FriendlyRHIName = LOCTEXT("VulkanSM5", "Vulkan Desktop (SM5, Experimental)");
@@ -93,13 +74,17 @@ static FText GetFriendlyNameFromRHIName(const FString& InRHIName)
 	{
 		FriendlyRHIName = LOCTEXT("SwitchForward", "Switch (Forward)");
 	}
-	else if (InRHIName == TEXT("PCD3D_ES2"))
+	else if (InRHIName == TEXT("GLSL_150_ES2") || InRHIName == TEXT("GLSL_150_ES31") || InRHIName == TEXT("GLSL_150")
+		|| InRHIName == TEXT("SF_VULKAN_ES31_ANDROID") || InRHIName == TEXT("SF_VULKAN_ES31")
+		|| InRHIName == TEXT("SF_VULKAN_SM4") || InRHIName == TEXT("PCD3D_ES2") || InRHIName == TEXT("PCD3D_ES31"))
 	{
-		FriendlyRHIName = LOCTEXT("DirectX11_ES2", "DirectX 11 Mobile Emulation (ES2)");
+		// Explicitly remove these formats as they are obsolete/not quite supported; users can still target them by adding them as +TargetedRHIs in the TargetPlatform ini.
+		FriendlyRHIName = FText::GetEmpty();
 	}
-	else if (InRHIName == TEXT("PCD3D_ES31"))
+	else
 	{
-		FriendlyRHIName = LOCTEXT("DirectX11_ES31", "DirectX 11 Mobile Emulation (ES3.1)");
+		UE_LOG(LogEngine, Warning, TEXT("Unknown Windows target RHI %s"), *InRHIName);
+		FriendlyRHIName = LOCTEXT("UnknownRHI", "UnknownRHI");
 	}
 
 	return FriendlyRHIName;
@@ -499,28 +484,31 @@ void FTargetShaderFormatsPropertyDetails::CreateTargetShaderFormatsPropertyView(
 	for (const FName& ShaderFormat : ShaderFormats)
 	{
 		FText FriendlyShaderFormatName = GetFriendlyNameFromRHIName(ShaderFormat.ToString());
+		// Skip explicitly removed entries
+		if (!FriendlyShaderFormatName.IsEmpty())
+		{
+			FDetailWidgetRow& TargetedRHIWidgetRow = TargetedRHICategoryBuilder.AddCustomRow(FriendlyShaderFormatName);
 
-		FDetailWidgetRow& TargetedRHIWidgetRow = TargetedRHICategoryBuilder.AddCustomRow(FriendlyShaderFormatName);
-
-		TargetedRHIWidgetRow
-		.NameContent()
-		[
-			SNew(SHorizontalBox)
-			+ SHorizontalBox::Slot()
-			.Padding(FMargin(0, 1, 0, 1))
-			.FillWidth(1.0f)
-			[
-				SNew(STextBlock)
-				.Text(FriendlyShaderFormatName)
+			TargetedRHIWidgetRow
+				.NameContent()
+				[
+					SNew(SHorizontalBox)
+					+ SHorizontalBox::Slot()
+				.Padding(FMargin(0, 1, 0, 1))
+				.FillWidth(1.0f)
+				[
+					SNew(STextBlock)
+					.Text(FriendlyShaderFormatName)
 				.Font(DetailBuilder->GetDetailFont())
-			]
-		]
-		.ValueContent()
-		[
-			SNew(SCheckBox)
-			.OnCheckStateChanged(this, &FTargetShaderFormatsPropertyDetails::OnTargetedRHIChanged, ShaderFormat)
-			.IsChecked(this, &FTargetShaderFormatsPropertyDetails::IsTargetedRHIChecked, ShaderFormat)
-		];
+				]
+				]
+			.ValueContent()
+				[
+					SNew(SCheckBox)
+					.OnCheckStateChanged(this, &FTargetShaderFormatsPropertyDetails::OnTargetedRHIChanged, ShaderFormat)
+				.IsChecked(this, &FTargetShaderFormatsPropertyDetails::IsTargetedRHIChecked, ShaderFormat)
+				];
+		}
 	}
 }
 
