Fix Generate and Invalidate Project ID

diff --git a/Engine/Source/Editor/DetailCustomizations/Private/GuidStructCustomization.cpp b/Engine/Source/Editor/DetailCustomizations/Private/GuidStructCustomization.cpp
--- a/Engine/Source/Editor/DetailCustomizations/Private/GuidStructCustomization.cpp
+++ b/Engine/Source/Editor/DetailCustomizations/Private/GuidStructCustomization.cpp
@@ -6,6 +6,7 @@
 #include "Engine/GameViewportClient.h"
 #include "Widgets/SBoxPanel.h"
 #include "Textures/SlateIcon.h"
+#include "Framework/Application/SlateApplication.h"
 #include "Framework/Commands/UIAction.h"
 #include "Framework/MultiBox/MultiBoxBuilder.h"
 #include "PropertyHandle.h"
@@ -107,6 +108,9 @@ void FGuidStructCustomization::SetGuidValue( const FGuid& Guid )
 
 void FGuidStructCustomization::HandleGuidActionClicked( EPropertyEditorGuidActions::Type Action )
 {
+	// Clear focus so text field can be updated
+	FSlateApplication::Get().ClearKeyboardFocus(EFocusCause::Cleared);
+
 	if (Action == EPropertyEditorGuidActions::Generate)
 	{
 		SetGuidValue(FGuid::NewGuid());
