Don't display the remove from array icon in the customization if the input action or axis mapping is not in an array.

diff --git a/Engine/Source/Editor/DetailCustomizations/Private/InputStructCustomization.cpp b/Engine/Source/Editor/DetailCustomizations/Private/InputStructCustomization.cpp
--- a/Engine/Source/Editor/DetailCustomizations/Private/InputStructCustomization.cpp
+++ b/Engine/Source/Editor/DetailCustomizations/Private/InputStructCustomization.cpp
@@ -64,8 +64,10 @@ void FInputActionMappingCustomization::CustomizeChildren( TSharedRef<class IProp
 	TSharedPtr<IPropertyHandle> AltHandle = InStructPropertyHandle->GetChildHandle(GET_MEMBER_NAME_CHECKED(FInputActionKeyMapping, bAlt));
 	TSharedPtr<IPropertyHandle> CmdHandle = InStructPropertyHandle->GetChildHandle(GET_MEMBER_NAME_CHECKED(FInputActionKeyMapping, bCmd));
 
-	TSharedRef<SWidget> RemoveButton = PropertyCustomizationHelpers::MakeDeleteButton(FSimpleDelegate::CreateSP(this, &FInputActionMappingCustomization::RemoveActionMappingButton_OnClick),
-		LOCTEXT("RemoveActionMappingToolTip", "Removes Action Mapping"));
+	const TSharedPtr<IPropertyHandleArray> ParentArrayHandle = InStructPropertyHandle->GetParentHandle()->AsArray();
+
+	TSharedRef<SWidget> RemoveButton = (ParentArrayHandle.IsValid() ? PropertyCustomizationHelpers::MakeDeleteButton(FSimpleDelegate::CreateSP(this, &FInputActionMappingCustomization::RemoveActionMappingButton_OnClick),
+		LOCTEXT("RemoveActionMappingToolTip", "Removes Action Mapping")) : SNullWidget::NullWidget);
 
 	StructBuilder.AddCustomRow( LOCTEXT("KeySearchStr", "Key") )
 	[
@@ -185,8 +187,10 @@ void FInputAxisMappingCustomization::CustomizeChildren( TSharedRef<class IProper
 	TSharedPtr<IPropertyHandle> KeyHandle = InStructPropertyHandle->GetChildHandle(GET_MEMBER_NAME_CHECKED(FInputAxisKeyMapping, Key));
 	TSharedPtr<IPropertyHandle> ScaleHandle = InStructPropertyHandle->GetChildHandle(GET_MEMBER_NAME_CHECKED(FInputAxisKeyMapping, Scale));
 
-	TSharedRef<SWidget> RemoveButton = PropertyCustomizationHelpers::MakeDeleteButton( FSimpleDelegate::CreateSP( this, &FInputAxisMappingCustomization::RemoveAxisMappingButton_OnClick), 
-		LOCTEXT("RemoveAxisMappingToolTip", "Removes Axis Mapping") );
+	const TSharedPtr<IPropertyHandleArray> ParentArrayHandle = InStructPropertyHandle->GetParentHandle()->AsArray();
+
+	TSharedRef<SWidget> RemoveButton = (ParentArrayHandle.IsValid() ? PropertyCustomizationHelpers::MakeDeleteButton(FSimpleDelegate::CreateSP(this, &FInputAxisMappingCustomization::RemoveAxisMappingButton_OnClick),
+		LOCTEXT("RemoveAxisMappingToolTip", "Removes Axis Mapping")) : SNullWidget::NullWidget);
 
 	StructBuilder.AddCustomRow( LOCTEXT("KeySearchStr", "Key") )
 	[
