Linux: Add WebM to the movie filter list

diff --git a/Engine/Source/Developer/SlateFileDialogs/Private/SlateFileDlgWindow.h b/Engine/Source/Developer/SlateFileDialogs/Private/SlateFileDlgWindow.h
--- a/Engine/Source/Developer/SlateFileDialogs/Private/SlateFileDlgWindow.h
+++ b/Engine/Source/Developer/SlateFileDialogs/Private/SlateFileDlgWindow.h
@@ -304,7 +304,8 @@ public:
 			{
 				Icon = StyleSet->GetBrush("SlateFileDialogs.Audio");
 			}
-			else if (Extension.Equals(TEXT("mp4"), ESearchCase::IgnoreCase))
+			else if (Extension.Equals(TEXT("mp4"), ESearchCase::IgnoreCase) ||
+					 Extension.Equals(TEXT("webm"), ESearchCase::IgnoreCase))
 			{
 				Icon = StyleSet->GetBrush("SlateFileDialogs.Video");
 			}
diff --git a/Engine/Source/Editor/DetailCustomizations/Private/MoviePlayerSettingsDetails.cpp b/Engine/Source/Editor/DetailCustomizations/Private/MoviePlayerSettingsDetails.cpp
--- a/Engine/Source/Editor/DetailCustomizations/Private/MoviePlayerSettingsDetails.cpp
+++ b/Engine/Source/Editor/DetailCustomizations/Private/MoviePlayerSettingsDetails.cpp
@@ -68,7 +68,7 @@ void FMoviePlayerSettingsDetails::GenerateArrayElementWidget(TSharedRef<IPropert
 					.BrowseDirectory(FEditorDirectories::Get().GetLastDirectory(ELastDirectory::GENERIC_OPEN))
 					.BrowseTitle(LOCTEXT("PropertyEditorTitle", "File picker..."))
 					.FilePath(this, &FMoviePlayerSettingsDetails::HandleFilePathPickerFilePath, PropertyHandle)
-					.FileTypeFilter(TEXT("MPEG-4 Movie (*.mp4)|*.mp4"))
+					.FileTypeFilter(TEXT("MPEG-4 Movie (*.mp4)|*.mp4|Web Movie (*.webm)|*.webm"))
 					.OnPathPicked(this, &FMoviePlayerSettingsDetails::HandleFilePathPickerPathPicked, PropertyHandle)
 			];
 	}
