Fix fbx import dialog inputs

diff --git a/Engine/Source/Editor/DetailCustomizations/Private/FbxImportUIDetails.cpp b/Engine/Source/Editor/DetailCustomizations/Private/FbxImportUIDetails.cpp
--- a/Engine/Source/Editor/DetailCustomizations/Private/FbxImportUIDetails.cpp
+++ b/Engine/Source/Editor/DetailCustomizations/Private/FbxImportUIDetails.cpp
@@ -25,6 +25,9 @@
 
 #define LOCTEXT_NAMESPACE "FbxImportUIDetails"
 
+#define MinimumLodNumberID 0
+#define LodNumberID 1
+
 //If the String is contain in the StringArray, it return the index. Otherwise return INDEX_NONE
 int FindString(const TArray<TSharedPtr<FString>> &StringArray, const FString &String) {
 	for (int i = 0; i < StringArray.Num(); i++)
@@ -125,7 +128,16 @@ void FFbxImportUIDetails::CustomizeDetails( IDetailLayoutBuilder& DetailBuilder
 	switch(ImportUI->MeshTypeToImport)
 	{
 		case FBXIT_StaticMesh:
-			CollectChildPropertiesRecursive(StaticMeshDataProp, ExtraProperties);
+			{
+				//Validate static mesh input
+				TSharedRef<IPropertyHandle> MinimumLodNumberProp = DetailBuilder.GetProperty(GET_MEMBER_NAME_CHECKED(UFbxImportUI, MinimumLodNumber));
+				MinimumLodNumberProp->SetOnPropertyValueChanged(FSimpleDelegate::CreateSP(this, &FFbxImportUIDetails::ValidateLodSettingsChanged, MinimumLodNumberID));
+				//Validate static mesh input
+				TSharedRef<IPropertyHandle> LodNumberProp = DetailBuilder.GetProperty(GET_MEMBER_NAME_CHECKED(UFbxImportUI, LodNumber));
+				LodNumberProp->SetOnPropertyValueChanged(FSimpleDelegate::CreateSP(this, &FFbxImportUIDetails::ValidateLodSettingsChanged, LodNumberID));
+
+				CollectChildPropertiesRecursive(StaticMeshDataProp, ExtraProperties);
+			}
 			break;
 		case FBXIT_SkeletalMesh:
 			if(ImportUI->bImportMesh)
@@ -143,10 +155,15 @@ void FFbxImportUIDetails::CustomizeDetails( IDetailLayoutBuilder& DetailBuilder
 	EFBXImportType ImportType = ImportUI->MeshTypeToImport;
 
 	//Hide LodDistance property if we do not need them
-	if (ImportType == FBXIT_StaticMesh && ImportUI->bAutoComputeLodDistances)
+	if (ImportType == FBXIT_StaticMesh)
 	{
-		for (int32 LodIndex = 0; LodIndex < 8; ++LodIndex)
+		int32 ShowMaxLodIndex = (ImportUI->bAutoComputeLodDistances ? 0 : ImportUI->LodNumber > 0 ? ImportUI->LodNumber : MAX_STATIC_MESH_LODS) - 1;
+		for (int32 LodIndex = 0; LodIndex < MAX_STATIC_MESH_LODS; ++LodIndex)
 		{
+			if (LodIndex <= ShowMaxLodIndex)
+			{
+				continue;
+			}
 			TArray<FStringFormatArg> Args;
 			Args.Add(TEXT("LodDistance"));
 			Args.Add(FString::FromInt(LodIndex));
@@ -747,6 +764,34 @@ void FFbxImportUIDetails::ImportAutoComputeLodDistancesChanged()
 	RefreshCustomDetail();
 }
 
+void FFbxImportUIDetails::ValidateLodSettingsChanged(int32 MemberID)
+{
+	//This feature is supported only for staticmesh
+	if (ImportUI->MeshTypeToImport != FBXIT_StaticMesh)
+	{
+		return;
+	}
+
+	if (ImportUI->MinimumLodNumber < 0 || ImportUI->MinimumLodNumber >= MAX_STATIC_MESH_LODS)
+	{
+		ImportUI->MinimumLodNumber = FMath::Clamp<int32>(ImportUI->MinimumLodNumber, 0, MAX_STATIC_MESH_LODS -1);
+	}
+	if (ImportUI->LodNumber < 0 || ImportUI->LodNumber >= MAX_STATIC_MESH_LODS)
+	{
+		ImportUI->LodNumber = FMath::Clamp<int32>(ImportUI->LodNumber, 0, MAX_STATIC_MESH_LODS);
+	}
+
+	if (ImportUI->LodNumber > 0 && ImportUI->MinimumLodNumber >= ImportUI->LodNumber)
+	{
+		ImportUI->MinimumLodNumber = FMath::Clamp<int32>(ImportUI->MinimumLodNumber, 0, ImportUI->LodNumber - 1);
+	}
+	
+	if (!ImportUI->bAutoComputeLodDistances && MemberID == LodNumberID)
+	{
+		RefreshCustomDetail();
+	}
+}
+
 void FFbxImportUIDetails::ImportMaterialsChanged()
 {
 	//We need to update the Base Material UI
diff --git a/Engine/Source/Editor/DetailCustomizations/Private/FbxImportUIDetails.h b/Engine/Source/Editor/DetailCustomizations/Private/FbxImportUIDetails.h
--- a/Engine/Source/Editor/DetailCustomizations/Private/FbxImportUIDetails.h
+++ b/Engine/Source/Editor/DetailCustomizations/Private/FbxImportUIDetails.h
@@ -47,6 +47,9 @@ public:
 	/** Called if the bAutoComputeLodDistances changes */
 	void ImportAutoComputeLodDistancesChanged();
 
+	/** Called when the LODSettings changes */
+	void ValidateLodSettingsChanged(int32 MemberID);
+
 	/** Called if the bImportMaterials changes */
 	void ImportMaterialsChanged();
 
diff --git a/Engine/Source/Editor/UnrealEd/Classes/Factories/FbxImportUI.h b/Engine/Source/Editor/UnrealEd/Classes/Factories/FbxImportUI.h
--- a/Engine/Source/Editor/UnrealEd/Classes/Factories/FbxImportUI.h
+++ b/Engine/Source/Editor/UnrealEd/Classes/Factories/FbxImportUI.h
@@ -100,11 +100,11 @@ public:
 	float LodDistance7;
 
 	/** Set the minimum LOD number. A value of 0 disable the option. */
-	UPROPERTY(EditAnywhere, BlueprintReadWrite, config, AdvancedDisplay, Category = LodSettings, meta = (ImportType = "StaticMesh", UIMin = "-1"))
+	UPROPERTY(EditAnywhere, BlueprintReadWrite, config, AdvancedDisplay, Category = LodSettings, meta = (ImportType = "StaticMesh", UIMin = "0"))
 	int32 MinimumLodNumber;
 
 	/** Set the number of LODs. A value of 0 disable the option. */
-	UPROPERTY(EditAnywhere, BlueprintReadWrite, config, AdvancedDisplay, Category = LodSettings, meta = (ImportType = "StaticMesh", UIMin = "-1"))
+	UPROPERTY(EditAnywhere, BlueprintReadWrite, config, AdvancedDisplay, Category = LodSettings, meta = (ImportType = "StaticMesh", UIMin = "0"))
 	int32 LodNumber;
 
 	/** True to import animations from the FBX File */
