Make sure the re-import dialog hide the lod settings, since these options are not use by the re-import.
Make sure the LODSettings are hide when the import LODs check box is unchecked

diff --git a/Engine/Source/Editor/DetailCustomizations/Private/FbxImportUIDetails.cpp b/Engine/Source/Editor/DetailCustomizations/Private/FbxImportUIDetails.cpp
--- a/Engine/Source/Editor/DetailCustomizations/Private/FbxImportUIDetails.cpp
+++ b/Engine/Source/Editor/DetailCustomizations/Private/FbxImportUIDetails.cpp
@@ -122,6 +122,9 @@ void FFbxImportUIDetails::CustomizeDetails( IDetailLayoutBuilder& DetailBuilder
 	TSharedRef<IPropertyHandle> ImportAutoComputeLodDistancesProp = DetailBuilder.GetProperty(GET_MEMBER_NAME_CHECKED(UFbxImportUI, bAutoComputeLodDistances));
 	ImportAutoComputeLodDistancesProp->SetOnPropertyValueChanged(FSimpleDelegate::CreateSP(this, &FFbxImportUIDetails::ImportAutoComputeLodDistancesChanged));
 
+	TSharedPtr<IPropertyHandle> ImportMeshLODsProp = StaticMeshDataProp->GetChildHandle(GET_MEMBER_NAME_CHECKED(UFbxStaticMeshImportData, bImportMeshLODs));
+	ImportMeshLODsProp->SetOnPropertyValueChanged(FSimpleDelegate::CreateSP(this, &FFbxImportUIDetails::RefreshCustomDetail));
+
 	MeshCategory.GetDefaultProperties(CategoryDefaultProperties);
 
 	
@@ -155,7 +158,11 @@ void FFbxImportUIDetails::CustomizeDetails( IDetailLayoutBuilder& DetailBuilder
 	EFBXImportType ImportType = ImportUI->MeshTypeToImport;
 
 	//Hide LodDistance property if we do not need them
-	if (ImportType == FBXIT_StaticMesh)
+	if (ImportUI->bIsReimport || ImportType != FBXIT_StaticMesh || !ImportUI->StaticMeshImportData->bImportMeshLODs)
+	{
+		DetailBuilder.HideCategory(FName(TEXT("LodSettings")));
+	}
+	else if (ImportType == FBXIT_StaticMesh)
 	{
 		int32 ShowMaxLodIndex = (ImportUI->bAutoComputeLodDistances ? 0 : ImportUI->LodNumber > 0 ? ImportUI->LodNumber : MAX_STATIC_MESH_LODS) - 1;
 		for (int32 LodIndex = 0; LodIndex < MAX_STATIC_MESH_LODS; ++LodIndex)
@@ -176,11 +183,6 @@ void FFbxImportUIDetails::CustomizeDetails( IDetailLayoutBuilder& DetailBuilder
 			}
 		}
 	}
-	else if (ImportType != FBXIT_StaticMesh)
-	{
-		DetailBuilder.HideCategory(FName(TEXT("LodSettings")));
-	}
-	
 
 	if(ImportType != FBXIT_Animation)
 	{
@@ -347,16 +349,10 @@ void FFbxImportUIDetails::CustomizeDetails( IDetailLayoutBuilder& DetailBuilder
 
 	// Material Category
 	IDetailCategoryBuilder& MaterialCategory = DetailBuilder.EditCategory("Material");
-	if(ImportType == FBXIT_Animation)
+	if(ImportUI->bIsReimport || ImportType == FBXIT_Animation)
 	{
-		// In animation-only mode, hide the material display
-		CategoryDefaultProperties.Empty();
-		MaterialCategory.GetDefaultProperties(CategoryDefaultProperties);
-
-		for(TSharedRef<IPropertyHandle> Handle : CategoryDefaultProperties)
-		{
-			DetailBuilder.HideProperty(Handle);
-		}
+		// hide the material category
+		DetailBuilder.HideCategory("Material");
 	}
 	else
 	{
