Fix CIS static analysis warning

diff --git a/Engine/Source/Editor/DetailCustomizations/Private/FbxImportUIDetails.cpp b/Engine/Source/Editor/DetailCustomizations/Private/FbxImportUIDetails.cpp
--- a/Engine/Source/Editor/DetailCustomizations/Private/FbxImportUIDetails.cpp
+++ b/Engine/Source/Editor/DetailCustomizations/Private/FbxImportUIDetails.cpp
@@ -162,7 +162,7 @@ void FFbxImportUIDetails::CustomizeDetails( IDetailLayoutBuilder& DetailBuilder
 	{
 		DetailBuilder.HideCategory(FName(TEXT("LodSettings")));
 	}
-	else if (ImportType == FBXIT_StaticMesh)
+	else
 	{
 		int32 ShowMaxLodIndex = (ImportUI->bAutoComputeLodDistances ? 0 : ImportUI->LodNumber > 0 ? ImportUI->LodNumber : MAX_STATIC_MESH_LODS) - 1;
 		for (int32 LodIndex = 0; LodIndex < MAX_STATIC_MESH_LODS; ++LodIndex)
