Disable animation property if the user do not import animation

diff --git a/Engine/Source/Editor/DetailCustomizations/Private/IntervalStructCustomization.cpp b/Engine/Source/Editor/DetailCustomizations/Private/IntervalStructCustomization.cpp
--- a/Engine/Source/Editor/DetailCustomizations/Private/IntervalStructCustomization.cpp
+++ b/Engine/Source/Editor/DetailCustomizations/Private/IntervalStructCustomization.cpp
@@ -111,6 +111,7 @@ void FIntervalStructCustomization<NumericType>::CustomizeHeader(TSharedRef<IProp
 				.UndeterminedString(LOCTEXT("MultipleValues", "Multiple Values"))
 				.Font(IDetailLayoutBuilder::GetDetailFont())
 				.AllowSpin(true)
+				.IsEnabled(this, &FIntervalStructCustomization::IsPropertyEnabled, EIntervalField::Min)
 				.LabelVAlign(VAlign_Center)
 				.Label()
 				[
@@ -137,6 +138,7 @@ void FIntervalStructCustomization<NumericType>::CustomizeHeader(TSharedRef<IProp
 				.UndeterminedString(LOCTEXT("MultipleValues", "Multiple Values"))
 				.Font(IDetailLayoutBuilder::GetDetailFont())
 				.AllowSpin(true)
+				.IsEnabled(this, &FIntervalStructCustomization::IsPropertyEnabled, EIntervalField::Max)
 				.LabelVAlign(VAlign_Center)
 				.Label()
 				[
@@ -149,7 +151,6 @@ void FIntervalStructCustomization<NumericType>::CustomizeHeader(TSharedRef<IProp
 	];
 }
 
-
 template <typename NumericType>
 void FIntervalStructCustomization<NumericType>::CustomizeChildren(TSharedRef<IPropertyHandle> StructPropertyHandle, IDetailChildrenBuilder& StructBuilder, IPropertyTypeCustomizationUtils& StructCustomizationUtils)
 {
@@ -296,6 +297,17 @@ bool FIntervalStructCustomization<NumericType>::ShouldAllowSpin() const
 	return true;
 }
 
+template <typename NumericType>
+bool FIntervalStructCustomization<NumericType>::IsPropertyEnabled(EIntervalField Field) const
+{
+	IPropertyHandle* Handle = (Field == EIntervalField::Min) ? MinValueHandle.Get() : MaxValueHandle.Get();
+	if (Handle == nullptr)
+	{
+		return false;
+	}
+	return !Handle->IsEditConst();
+}
+
 /* Only explicitly instantiate the types which are supported
  *****************************************************************************/
 
diff --git a/Engine/Source/Editor/DetailCustomizations/Private/IntervalStructCustomization.h b/Engine/Source/Editor/DetailCustomizations/Private/IntervalStructCustomization.h
--- a/Engine/Source/Editor/DetailCustomizations/Private/IntervalStructCustomization.h
+++ b/Engine/Source/Editor/DetailCustomizations/Private/IntervalStructCustomization.h
@@ -111,6 +111,9 @@ protected:
 	 */
 	bool ShouldAllowSpin() const;
 
+	/* If the property is disable we must disable the editing */
+	bool IsPropertyEnabled(EIntervalField Field) const;
+
 	// Cached shared pointers to properties that we are managing
 	TSharedPtr<IPropertyHandle> MinValueHandle;
 	TSharedPtr<IPropertyHandle> MaxValueHandle;
diff --git a/Engine/Source/Editor/UnrealEd/Private/Fbx/FbxFactory.cpp b/Engine/Source/Editor/UnrealEd/Private/Fbx/FbxFactory.cpp
--- a/Engine/Source/Editor/UnrealEd/Private/Fbx/FbxFactory.cpp
+++ b/Engine/Source/Editor/UnrealEd/Private/Fbx/FbxFactory.cpp
@@ -958,11 +958,11 @@ bool UFbxImportUI::CanEditChange( const UProperty* InProperty ) const
 	{
 		FName PropName = InProperty->GetFName();
 
-		if(PropName == TEXT("StartFrame") || PropName == TEXT("EndFrame"))
+		if(PropName == TEXT("FrameImportRange"))
 		{
 			bIsMutable = AnimSequenceImportData->AnimationLength == FBXALIT_SetRange && bImportAnimations;
 		}
-		else if(PropName == TEXT("bImportCustomAttribute") || PropName == TEXT("AnimationLength"))
+		else if(PropName == TEXT("bImportCustomAttribute") || PropName == TEXT("AnimationLength") || PropName == TEXT("CustomSampleRate") || PropName == TEXT("bUseDefaultSampleRate"))
 		{
 			bIsMutable = bImportAnimations;
 		}
