Make PostProcessSettings only show root categories if metadata is specified.  Otherwise it just shows nested categories
- Set post process volume to create these root categories to avoid redundant category expansions

diff --git a/Engine/Source/Editor/DetailCustomizations/Private/PostProcessSettingsCustomization.cpp b/Engine/Source/Editor/DetailCustomizations/Private/PostProcessSettingsCustomization.cpp
--- a/Engine/Source/Editor/DetailCustomizations/Private/PostProcessSettingsCustomization.cpp
+++ b/Engine/Source/Editor/DetailCustomizations/Private/PostProcessSettingsCustomization.cpp
@@ -25,11 +25,52 @@
 
 #define LOCTEXT_NAMESPACE "PostProcessSettingsCustomization"
 
+struct FCategoryOrGroup
+{
+	IDetailCategoryBuilder* Category;
+	IDetailGroup* Group;
+
+	FCategoryOrGroup(IDetailCategoryBuilder& NewCategory)
+		: Category(&NewCategory)
+		, Group(nullptr)
+	{}
+
+	FCategoryOrGroup(IDetailGroup& NewGroup)
+		: Category(nullptr)
+		, Group(&NewGroup)
+	{}
+
+	IDetailPropertyRow& AddProperty(TSharedRef<IPropertyHandle> PropertyHandle)
+	{
+		if (Category)
+		{
+			return Category->AddProperty(PropertyHandle);
+		}
+		else
+		{
+			return Group->AddPropertyRow(PropertyHandle);
+		}
+	}
+
+	IDetailGroup& AddGroup(FName GroupName, const FText& DisplayName)
+	{
+		if (Category)
+		{
+			return Category->AddGroup(GroupName, DisplayName);
+		}
+		else
+		{
+			return Group->AddGroup(GroupName, DisplayName);
+		}
+	}
+};
+
+
 struct FPostProcessGroup
 {
 	FString RawGroupName;
 	FString DisplayName;
-	IDetailCategoryBuilder* RootCategory;
+	FCategoryOrGroup* RootCategory;
 	TArray<TSharedPtr<IPropertyHandle>> SimplePropertyHandles;
 	TArray<TSharedPtr<IPropertyHandle>> AdvancedPropertyHandles;
 
@@ -43,7 +84,7 @@ struct FPostProcessGroup
 	{}
 };
 
-
+PRAGMA_DISABLE_OPTIMIZATION
 void FPostProcessSettingsCustomization::CustomizeChildren( TSharedRef<IPropertyHandle> StructPropertyHandle, class IDetailChildrenBuilder& StructBuilder, IPropertyTypeCustomizationUtils& StructCustomizationUtils )
 {
 	uint32 NumChildren = 0;
@@ -58,7 +99,7 @@ void FPostProcessSettingsCustomization::CustomizeChildren( TSharedRef<IPropertyH
 	// Create new categories in the parent layout rather than adding all post process settings to one category
 	IDetailLayoutBuilder& LayoutBuilder = StructBuilder.GetParentCategory().GetParentLayout();
 
-	TMap<FString, IDetailCategoryBuilder*> NameToCategoryBuilderMap;
+	TMap<FString, FCategoryOrGroup> NameToCategoryBuilderMap;
 	TMap<FString, FPostProcessGroup> NameToGroupMap;
 
 	static const auto VarTonemapperFilm = IConsoleManager::Get().FindTConsoleVariableDataInt(TEXT("r.TonemapperFilm"));
@@ -72,13 +113,16 @@ void FPostProcessSettingsCustomization::CustomizeChildren( TSharedRef<IPropertyH
 	const bool bUsingFilmTonemapper = bDesktopTonemapperFilm || bMobileTonemapperFilm;		// Are any platforms use film tonemapper
 	const bool bUsingLegacyTonemapper = !bDesktopTonemapperFilm || !bMobileTonemapperFilm;	// Are any platforms use legacy/ES2 tonemapper
 
+	static const FName ShowPostProcessCategoriesName("ShowPostProcessCategories");
+
+	bool bShowPostProcessCategories = StructPropertyHandle->HasMetaData(ShowPostProcessCategoriesName);
+
 	if(Result == FPropertyAccess::Success && NumChildren > 0)
 	{
 		for( uint32 ChildIndex = 0; ChildIndex < NumChildren; ++ChildIndex )
 		{
 			TSharedPtr<IPropertyHandle> ChildHandle = StructPropertyHandle->GetChildHandle( ChildIndex );
 
-			
 			if( ChildHandle.IsValid() && ChildHandle->GetProperty() )
 			{
 				UProperty* Property = ChildHandle->GetProperty();
@@ -114,13 +158,20 @@ void FPostProcessSettingsCustomization::CustomizeChildren( TSharedRef<IPropertyH
 
 				FString RootCategoryName = CategoryAndGroups.Num() > 0 ? CategoryAndGroups[0] : RawCategoryName;
 
-				IDetailCategoryBuilder* Category = NameToCategoryBuilderMap.FindRef(RootCategoryName);
-				if(!Category)
-				{
-					IDetailCategoryBuilder& NewCategory = LayoutBuilder.EditCategory(*RootCategoryName, FText::GetEmpty(), ECategoryPriority::TypeSpecific);
-					NameToCategoryBuilderMap.Add(RootCategoryName, &NewCategory);
 
-					Category = &NewCategory;
+				FCategoryOrGroup* Category = NameToCategoryBuilderMap.Find(RootCategoryName);
+				if (!Category)
+				{
+					if(bShowPostProcessCategories)
+					{
+						IDetailCategoryBuilder& NewCategory = LayoutBuilder.EditCategory(*RootCategoryName, FText::GetEmpty(), ECategoryPriority::TypeSpecific);
+						Category = &NameToCategoryBuilderMap.Emplace(RootCategoryName, NewCategory);
+					}
+					else
+					{
+						IDetailGroup& NewGroup = StructBuilder.AddGroup(*RootCategoryName, FText::FromString(RootCategoryName));
+						Category = &NameToCategoryBuilderMap.Emplace(RootCategoryName, NewGroup);
+					}
 				}
 
 				if(CategoryAndGroups.Num() > 1)
@@ -163,8 +214,10 @@ void FPostProcessSettingsCustomization::CustomizeChildren( TSharedRef<IPropertyH
 			{
 				IDetailGroup& SimpleGroup = PPGroup.RootCategory->AddGroup(*PPGroup.RawGroupName, FText::FromString(PPGroup.DisplayName));
 
+				static const FString ColorGradingName = TEXT("Color Grading");
+
 				// Only enable group reset on color grading category groups
-				if (PPGroup.RootCategory->GetDisplayName().IdenticalTo(FText::FromString(TEXT("Color Grading"))))
+				if (PPGroup.RawGroupName.Contains(ColorGradingName))
 				{
 					SimpleGroup.EnableReset(true);
 				}
@@ -187,7 +240,7 @@ void FPostProcessSettingsCustomization::CustomizeChildren( TSharedRef<IPropertyH
 		}
 	}
 }
-
+PRAGMA_ENABLE_OPTIMIZATION
 
 
 void FPostProcessSettingsCustomization::CustomizeHeader( TSharedRef<IPropertyHandle> StructPropertyHandle, class FDetailWidgetRow& HeaderRow, IPropertyTypeCustomizationUtils& StructCustomizationUtils )
diff --git a/Engine/Source/Runtime/Engine/Classes/Engine/PostProcessVolume.h b/Engine/Source/Runtime/Engine/Classes/Engine/PostProcessVolume.h
--- a/Engine/Source/Runtime/Engine/Classes/Engine/PostProcessVolume.h
+++ b/Engine/Source/Runtime/Engine/Classes/Engine/PostProcessVolume.h
@@ -24,7 +24,7 @@ class APostProcessVolume : public AVolume, public IInterface_PostProcessVolume
 	GENERATED_UCLASS_BODY()
 
 	/** Post process settings to use for this volume. */
-	UPROPERTY(interp, Category=PostProcessVolumeSettings)
+	UPROPERTY(interp, Category=PostProcessVolumeSettings, meta=(ShowPostProcessCategories))
 	struct FPostProcessSettings Settings;
 
 	/**
