Changed the asset viewer profile customziation to have proper sub-categories for easier reading

diff --git a/Engine/Source/Editor/DetailCustomizations/Private/AssetViewerSettingsCustomization.cpp b/Engine/Source/Editor/DetailCustomizations/Private/AssetViewerSettingsCustomization.cpp
--- a/Engine/Source/Editor/DetailCustomizations/Private/AssetViewerSettingsCustomization.cpp
+++ b/Engine/Source/Editor/DetailCustomizations/Private/AssetViewerSettingsCustomization.cpp
@@ -16,6 +16,7 @@
 
 #include "ScopedTransaction.h"
 #include "SSettingsEditorCheckoutNotice.h"
+#include "IDetailGroup.h"
 
 #define LOCTEXT_NAMESPACE "AssetViewerSettingsCustomizations"
 
@@ -53,6 +54,8 @@ void FAssetViewerSettingsCustomization::CustomizeDetails(IDetailLayoutBuilder& D
 	uint32 PropertyCount = 0;
 	ProfilePropertyHandle->GetNumChildren(PropertyCount);
 
+	TMap<FName, IDetailGroup*> SubcategoryToGroupMap;
+
 	FName NamePropertyName = GET_MEMBER_NAME_CHECKED(FPreviewSceneProfile, ProfileName);
 	FName SharedProfilePropertyName = GET_MEMBER_NAME_CHECKED(FPreviewSceneProfile, bSharedProfile);
 	for (uint32 PropertyIndex = 0; PropertyIndex < PropertyCount; ++PropertyIndex)
@@ -93,7 +96,16 @@ void FAssetViewerSettingsCustomization::CustomizeDetails(IDetailLayoutBuilder& D
 		}
 		else
 		{
-			IDetailPropertyRow& Row = CategoryBuilder.AddProperty(ProfileProperty);
+			FName DefaultCategoryName = ProfileProperty->GetDefaultCategoryName();
+
+			IDetailGroup*& Group = SubcategoryToGroupMap.FindOrAdd(DefaultCategoryName);
+
+			if (!Group)
+			{
+				Group = &CategoryBuilder.AddGroup(DefaultCategoryName, ProfileProperty->GetDefaultCategoryText());
+			}
+
+			Group->AddPropertyRow(ProfileProperty.ToSharedRef());
 		}
 	}
 }
