Fix automation tests

diff --git a/Engine/Source/Editor/DetailCustomizations/Private/PostProcessSettingsCustomization.cpp b/Engine/Source/Editor/DetailCustomizations/Private/PostProcessSettingsCustomization.cpp
--- a/Engine/Source/Editor/DetailCustomizations/Private/PostProcessSettingsCustomization.cpp
+++ b/Engine/Source/Editor/DetailCustomizations/Private/PostProcessSettingsCustomization.cpp
@@ -40,6 +40,11 @@ struct FCategoryOrGroup
 		, Group(&NewGroup)
 	{}
 
+	FCategoryOrGroup()
+		: Category(nullptr)
+		, Group(nullptr)
+	{}
+
 	IDetailPropertyRow& AddProperty(TSharedRef<IPropertyHandle> PropertyHandle)
 	{
 		if (Category)
@@ -63,6 +68,11 @@ struct FCategoryOrGroup
 			return Group->AddGroup(GroupName, DisplayName);
 		}
 	}
+
+	bool IsValid() const
+	{
+		return Group || Category;
+	}
 };
 
 
@@ -70,17 +80,17 @@ struct FPostProcessGroup
 {
 	FString RawGroupName;
 	FString DisplayName;
-	FCategoryOrGroup* RootCategory;
+	FCategoryOrGroup RootCategory;
 	TArray<TSharedPtr<IPropertyHandle>> SimplePropertyHandles;
 	TArray<TSharedPtr<IPropertyHandle>> AdvancedPropertyHandles;
 
 	bool IsValid() const
 	{
-		return !RawGroupName.IsEmpty() && !DisplayName.IsEmpty() && RootCategory;
+		return !RawGroupName.IsEmpty() && !DisplayName.IsEmpty() && RootCategory.IsValid();
 	}
 
 	FPostProcessGroup()
-		: RootCategory(nullptr)
+		: RootCategory()
 	{}
 };
 
@@ -183,7 +193,7 @@ void FPostProcessSettingsCustomization::CustomizeChildren( TSharedRef<IPropertyH
 					// Is this a new group? It wont be valid if it is
 					if(!PPGroup.IsValid())
 					{
-						PPGroup.RootCategory = Category;
+						PPGroup.RootCategory = *Category;
 						PPGroup.RawGroupName = RawCategoryName;
 						PPGroup.DisplayName = CategoryAndGroups[1].TrimStartAndEnd();
 					}
@@ -208,11 +218,11 @@ void FPostProcessSettingsCustomization::CustomizeChildren( TSharedRef<IPropertyH
 
 		for(auto& NameAndGroup : NameToGroupMap)
 		{
-			const FPostProcessGroup& PPGroup = NameAndGroup.Value;
+			FPostProcessGroup& PPGroup = NameAndGroup.Value;
 
 			if(PPGroup.SimplePropertyHandles.Num() > 0 || PPGroup.AdvancedPropertyHandles.Num() > 0 )
 			{
-				IDetailGroup& SimpleGroup = PPGroup.RootCategory->AddGroup(*PPGroup.RawGroupName, FText::FromString(PPGroup.DisplayName));
+				IDetailGroup& SimpleGroup = PPGroup.RootCategory.AddGroup(*PPGroup.RawGroupName, FText::FromString(PPGroup.DisplayName));
 
 				static const FString ColorGradingName = TEXT("Color Grading");
 
