PR #5168: Register the details customization for FIntVector (Contributed by [REDACTED_PERSON])

diff --git a/Engine/Source/Editor/DetailCustomizations/Private/DetailCustomizations.cpp b/Engine/Source/Editor/DetailCustomizations/Private/DetailCustomizations.cpp
--- a/Engine/Source/Editor/DetailCustomizations/Private/DetailCustomizations.cpp
+++ b/Engine/Source/Editor/DetailCustomizations/Private/DetailCustomizations.cpp
@@ -189,6 +189,7 @@ void FDetailCustomizationsModule::RegisterPropertyTypeCustomizations()
 	RegisterCustomPropertyTypeLayout("DataTableCategoryHandle", FOnGetPropertyTypeCustomizationInstance::CreateStatic(&FDataTableCategoryCustomizationLayout::MakeInstance));
 	RegisterCustomPropertyTypeLayout("CurveTableRowHandle", FOnGetPropertyTypeCustomizationInstance::CreateStatic(&FCurveTableCustomizationLayout::MakeInstance));
 	RegisterCustomPropertyTypeLayout(NAME_Vector, FOnGetPropertyTypeCustomizationInstance::CreateStatic(&FVectorStructCustomization::MakeInstance));
+	RegisterCustomPropertyTypeLayout("IntVector", FOnGetPropertyTypeCustomizationInstance::CreateStatic(&FVectorStructCustomization::MakeInstance));
 	RegisterCustomPropertyTypeLayout(NAME_Vector4, FOnGetPropertyTypeCustomizationInstance::CreateStatic(&FVector4StructCustomization::MakeInstance));
 	RegisterCustomPropertyTypeLayout(NAME_Vector2D, FOnGetPropertyTypeCustomizationInstance::CreateStatic(&FMathStructCustomization::MakeInstance));
 	RegisterCustomPropertyTypeLayout(NAME_IntPoint, FOnGetPropertyTypeCustomizationInstance::CreateStatic(&FMathStructCustomization::MakeInstance));
