PR #5246: The animations list in the Skeleton Notify details panel list slot will now fill the available width. (Contributed by [REDACTED_PERSON])

diff --git a/Engine/Source/Editor/DetailCustomizations/Private/SkeletonNotifyDetails.cpp b/Engine/Source/Editor/DetailCustomizations/Private/SkeletonNotifyDetails.cpp
--- a/Engine/Source/Editor/DetailCustomizations/Private/SkeletonNotifyDetails.cpp
+++ b/Engine/Source/Editor/DetailCustomizations/Private/SkeletonNotifyDetails.cpp
@@ -62,6 +62,7 @@ void FSkeletonNotifyDetails::CustomizeDetails( IDetailLayoutBuilder& DetailBuild
 			.Font( DetailFontInfo )
 		]
 		.ValueContent()
+		.HAlign( HAlign_Fill )
 		[
 			SNew(SHorizontalBox)
 			+SHorizontalBox::Slot()
