Refactor platform-duplicated shader format classes into FShderFormatsPropertyDetails.
Preserve Targeted RHI ordering when modifying checked state

diff --git a/Engine/Source/Developer/SharedSettingsWidgets/Private/ShaderFormatsPropertyDetails.cpp b/Engine/Source/Developer/SharedSettingsWidgets/Private/ShaderFormatsPropertyDetails.cpp
--- a/Engine/Source/Developer/SharedSettingsWidgets/Private/ShaderFormatsPropertyDetails.cpp
+++ b/Engine/Source/Developer/SharedSettingsWidgets/Private/ShaderFormatsPropertyDetails.cpp
@@ -21,7 +21,7 @@
 
 #define LOCTEXT_NAMESPACE "ShaderFormatsPropertyDetails"
 
-static FText GetFriendlyNameFromRHIName(const FString& InRHIName)
+FText FShaderFormatsPropertyDetails::GetFriendlyNameFromRHINameMac(const FString& InRHIName)
 {
 	FText FriendlyRHIName = FText::FromString(InRHIName);
 	
@@ -122,7 +122,7 @@ void FShaderFormatsPropertyDetails::SetOnUpdateShaderWarning(FSimpleDelegate con
 	ShaderFormatsPropertyHandle->SetOnPropertyValueChanged(Delegate);
 }
 
-void FShaderFormatsPropertyDetails::CreateTargetShaderFormatsPropertyView(ITargetPlatform* TargetPlatform)
+void FShaderFormatsPropertyDetails::CreateTargetShaderFormatsPropertyView(ITargetPlatform* TargetPlatform, GetFriendlyNameFromRHINameFnc FriendlyNameFnc)
 {
 	check(TargetPlatform);
 	DetailBuilder->HideProperty(ShaderFormatsPropertyHandle);
@@ -133,31 +133,35 @@ void FShaderFormatsPropertyDetails::CreateTargetShaderFormatsPropertyView(ITarge
 	
 	IDetailCategoryBuilder& TargetedRHICategoryBuilder = DetailBuilder->EditCategory(*Title);
 	
+	int32 ShaderCounter = 0;
 	for (const FName& ShaderFormat : ShaderFormats)
 	{
-		FText FriendlyShaderFormatName = GetFriendlyNameFromRHIName(ShaderFormat.ToString());
-		
-		FDetailWidgetRow& TargetedRHIWidgetRow = TargetedRHICategoryBuilder.AddCustomRow(FriendlyShaderFormatName);
-		
-		TargetedRHIWidgetRow
-		.NameContent()
-		[
-			SNew(SHorizontalBox)
-			+ SHorizontalBox::Slot()
-			.Padding(FMargin(0, 1, 0, 1))
-			.FillWidth(1.0f)
+		FText FriendlyShaderFormatName = FriendlyNameFnc(ShaderFormat.ToString());
+		if (!FriendlyShaderFormatName.IsEmpty())
+		{
+			ShaderFormatOrder.Add(ShaderFormat, ShaderCounter++);
+			FDetailWidgetRow& TargetedRHIWidgetRow = TargetedRHICategoryBuilder.AddCustomRow(FriendlyShaderFormatName);
+
+			TargetedRHIWidgetRow
+			.NameContent()
 			[
-				SNew(STextBlock)
-				.Text(FriendlyShaderFormatName)
-				.Font(DetailBuilder->GetDetailFont())
-			 ]
-		 ]
-		.ValueContent()
-		[
-			SNew(SCheckBox)
-			.OnCheckStateChanged(this, &FShaderFormatsPropertyDetails::OnTargetedRHIChanged, ShaderFormat)
-			.IsChecked(this, &FShaderFormatsPropertyDetails::IsTargetedRHIChecked, ShaderFormat)
-		 ];
+				SNew(SHorizontalBox)
+				+ SHorizontalBox::Slot()
+				.Padding(FMargin(0, 1, 0, 1))
+				.FillWidth(1.0f)
+				[
+					SNew(STextBlock)
+					.Text(FriendlyShaderFormatName)
+					.Font(DetailBuilder->GetDetailFont())
+				]
+			]
+			.ValueContent()
+			[
+				SNew(SCheckBox)
+				.OnCheckStateChanged(this, &FShaderFormatsPropertyDetails::OnTargetedRHIChanged, ShaderFormat)
+				.IsChecked(this, &FShaderFormatsPropertyDetails::IsTargetedRHIChecked, ShaderFormat)
+			];
+		}
 	}
 }
 
@@ -175,7 +179,17 @@ void FShaderFormatsPropertyDetails::OnTargetedRHIChanged(ECheckBoxState InNewVal
 			TArray<FString>& Array = *(TArray<FString>*)RawPtr;
 			if(InNewValue == ECheckBoxState::Checked)
 			{
-				Array.Add(InRHIName.ToString());
+				// Preserve order from GetAllPossibleShaderFormats
+				const int32 InIndex = ShaderFormatOrder[InRHIName];
+				int32 InsertIndex = 0;
+				for (; InsertIndex < Array.Num(); ++InsertIndex)
+				{
+					if (InIndex < ShaderFormatOrder[*Array[InsertIndex]])
+					{
+						break;
+					}
+				}
+				Array.Insert(InRHIName.ToString(), InsertIndex);
 			}
 			else
 			{
diff --git a/Engine/Source/Developer/SharedSettingsWidgets/Public/ShaderFormatsPropertyDetails.h b/Engine/Source/Developer/SharedSettingsWidgets/Public/ShaderFormatsPropertyDetails.h
--- a/Engine/Source/Developer/SharedSettingsWidgets/Public/ShaderFormatsPropertyDetails.h
+++ b/Engine/Source/Developer/SharedSettingsWidgets/Public/ShaderFormatsPropertyDetails.h
@@ -22,6 +22,8 @@ class SHAREDSETTINGSWIDGETS_API FShaderFormatsPropertyDetails
 {
 	
 public:
+	typedef FText (*GetFriendlyNameFromRHINameFnc)(const FString&);
+	static FText GetFriendlyNameFromRHINameMac(const FString& InRHIName);
 	
 	/**
 	 * Constructor for the parent property details view
@@ -30,7 +32,7 @@ public:
 	 * @param InProperty - The category name to override
 	 * @param InTitle - Title for display
 	 */
-	FShaderFormatsPropertyDetails(IDetailLayoutBuilder* InDetailBuilder, FString InProperty, FString InTitle);
+	FShaderFormatsPropertyDetails(IDetailLayoutBuilder* InDetailBuilder, FString InProperty = TEXT("TargetedRHIs"), FString InTitle = TEXT("Targeted RHIs"));
 	
 	/** Simple delegate for updating shader version warning */
 	void SetOnUpdateShaderWarning(FSimpleDelegate const& Delegate);
@@ -38,7 +40,7 @@ public:
 	/**
 	 * Create the UI to select which windows shader formats we are targetting
 	 */
-	void CreateTargetShaderFormatsPropertyView(ITargetPlatform* TargetPlatform);
+	void CreateTargetShaderFormatsPropertyView(ITargetPlatform* TargetPlatform, GetFriendlyNameFromRHINameFnc FriendlyNameFnc);
 	
 	
 	/**
@@ -59,10 +61,13 @@ private:
 	
 	/** Access to the Parent Property */
 	TSharedPtr<IPropertyHandle> ShaderFormatsPropertyHandle;
-	
+
 	/** The category name to override */
 	FString Property;
-	
+
 	/** Title for display */
 	FString Title;
+
+	/** Preserve shader format order when writing to property */
+	TMap<FName, int> ShaderFormatOrder;
 };
\ No newline at end of file
diff --git a/Engine/Source/Editor/DetailCustomizations/Private/LinuxTargetSettingsDetails.cpp b/Engine/Source/Editor/DetailCustomizations/Private/LinuxTargetSettingsDetails.cpp
--- a/Engine/Source/Editor/DetailCustomizations/Private/LinuxTargetSettingsDetails.cpp
+++ b/Engine/Source/Editor/DetailCustomizations/Private/LinuxTargetSettingsDetails.cpp
@@ -152,8 +152,9 @@ static FString GetIconFilename(EImageScope::Type Scope)
 void FLinuxTargetSettingsDetails::CustomizeDetails( IDetailLayoutBuilder& DetailBuilder )
 {
 	// Setup the supported/targeted RHI property view
-	TargetShaderFormatsDetails = MakeShareable(new FLinuxTargetShaderFormatsPropertyDetails(&DetailBuilder));
-	TargetShaderFormatsDetails->CreateTargetShaderFormatsPropertyView();
+	ITargetPlatform* TargetPlatform = FModuleManager::GetModuleChecked<ITargetPlatformModule>("LinuxTargetPlatform").GetTargetPlatforms()[0];
+	TargetShaderFormatsDetails = MakeShareable(new FShaderFormatsPropertyDetails(&DetailBuilder));
+	TargetShaderFormatsDetails->CreateTargetShaderFormatsPropertyView(TargetPlatform, GetFriendlyNameFromRHIName);
 
 	// Next add the splash image customization
 	const FText EditorSplashDesc(LOCTEXT("EditorSplashLabel", "Editor Splash"));
@@ -402,95 +403,4 @@ TSharedRef<SWidget> FLinuxTargetSettingsDetails::MakeAudioDeviceMenu(const TShar
 	return MenuBuilder.MakeWidget();
 }
 
-FLinuxTargetShaderFormatsPropertyDetails::FLinuxTargetShaderFormatsPropertyDetails(IDetailLayoutBuilder* InDetailBuilder)
-: DetailBuilder(InDetailBuilder)
-{
-	TargetShaderFormatsPropertyHandle = DetailBuilder->GetProperty("TargetedRHIs");
-	ensure(TargetShaderFormatsPropertyHandle.IsValid());
-}
-
-void FLinuxTargetShaderFormatsPropertyDetails::CreateTargetShaderFormatsPropertyView()
-{
-	DetailBuilder->HideProperty(TargetShaderFormatsPropertyHandle);
-
-	// List of supported RHI's and selected targets
-	ITargetPlatform* LinuxTargetPlatform = FModuleManager::GetModuleChecked<ITargetPlatformModule>("LinuxTargetPlatform").GetTargetPlatforms()[0];
-	TArray<FName> ShaderFormats;
-	LinuxTargetPlatform->GetAllPossibleShaderFormats(ShaderFormats);
-
-	IDetailCategoryBuilder& TargetedRHICategoryBuilder = DetailBuilder->EditCategory(TEXT("Targeted RHIs"));
-
-	for (const FName& ShaderFormat : ShaderFormats)
-	{
-		FText FriendlyShaderFormatName = GetFriendlyNameFromRHIName(ShaderFormat.ToString());
-
-		FDetailWidgetRow& TargetedRHIWidgetRow = TargetedRHICategoryBuilder.AddCustomRow(FriendlyShaderFormatName);
-
-		TargetedRHIWidgetRow
-		.NameContent()
-		[
-			SNew(SHorizontalBox)
-			+ SHorizontalBox::Slot()
-			.Padding(FMargin(0, 1, 0, 1))
-			.FillWidth(1.0f)
-			[
-				SNew(STextBlock)
-				.Text(FriendlyShaderFormatName)
-				.Font(DetailBuilder->GetDetailFont())
-			]
-		]
-		.ValueContent()
-		[
-			SNew(SCheckBox)
-			.OnCheckStateChanged(this, &FLinuxTargetShaderFormatsPropertyDetails::OnTargetedRHIChanged, ShaderFormat)
-			.IsChecked(this, &FLinuxTargetShaderFormatsPropertyDetails::IsTargetedRHIChecked, ShaderFormat)
-		];
-	}
-}
-
-
-void FLinuxTargetShaderFormatsPropertyDetails::OnTargetedRHIChanged(ECheckBoxState InNewValue, FName InRHIName)
-{
-	TArray<void*> RawPtrs;
-	TargetShaderFormatsPropertyHandle->AccessRawData(RawPtrs);
-
-	// Update the CVars with the selection
-	{
-		TargetShaderFormatsPropertyHandle->NotifyPreChange();
-		for (void* RawPtr : RawPtrs)
-		{
-			TArray<FString>& Array = *(TArray<FString>*)RawPtr;
-			if(InNewValue == ECheckBoxState::Checked)
-			{
-				Array.Add(InRHIName.ToString());
-			}
-			else
-			{
-				Array.Remove(InRHIName.ToString());
-			}
-		}
-		TargetShaderFormatsPropertyHandle->NotifyPostChange();
-	}
-}
-
-
-ECheckBoxState FLinuxTargetShaderFormatsPropertyDetails::IsTargetedRHIChecked(FName InRHIName) const
-{
-	ECheckBoxState CheckState = ECheckBoxState::Unchecked;
-
-	TArray<void*> RawPtrs;
-	TargetShaderFormatsPropertyHandle->AccessRawData(RawPtrs);
-	
-	for(void* RawPtr : RawPtrs)
-	{
-		TArray<FString>& Array = *(TArray<FString>*)RawPtr;
-		if(Array.Contains(InRHIName.ToString()))
-		{
-			CheckState = ECheckBoxState::Checked;
-		}
-	}
-	return CheckState;
-}
-
-
 #undef LOCTEXT_NAMESPACE
diff --git a/Engine/Source/Editor/DetailCustomizations/Private/LinuxTargetSettingsDetails.h b/Engine/Source/Editor/DetailCustomizations/Private/LinuxTargetSettingsDetails.h
--- a/Engine/Source/Editor/DetailCustomizations/Private/LinuxTargetSettingsDetails.h
+++ b/Engine/Source/Editor/DetailCustomizations/Private/LinuxTargetSettingsDetails.h
@@ -6,53 +6,13 @@
 #include "Styling/SlateColor.h"
 #include "Widgets/SWidget.h"
 #include "IDetailCustomization.h"
+#include "ShaderFormatsPropertyDetails.h"
 #include "TargetPlatformAudioCustomization.h"
 
 class IDetailLayoutBuilder;
 class IPropertyHandle;
 enum class ECheckBoxState : uint8;
 
-/* FLinuxTargetShaderFormatsPropertyDetails
-*****************************************************************************/
-
-/**
-* Helper which implements details panel customizations for a device profiles parent property
-*/
-class FLinuxTargetShaderFormatsPropertyDetails
-	: public TSharedFromThis<FLinuxTargetShaderFormatsPropertyDetails>
-{
-
-public:
-
-	/**
-	 * Constructor for the parent property details view
-	 *
-	 * @param InDetailsBuilder - Where we are adding our property view to
-	 */
-	FLinuxTargetShaderFormatsPropertyDetails(IDetailLayoutBuilder* InDetailBuilder);
-
-
-	/**
-	 * Create the UI to select which windows shader formats we are targetting
-	 */
-	void CreateTargetShaderFormatsPropertyView();
-
-private:
-
-	// Supported/Targeted RHI check boxes
-	void OnTargetedRHIChanged(ECheckBoxState InNewValue, FName InRHIName);
-	ECheckBoxState IsTargetedRHIChecked(FName InRHIName) const;
-
-private:
-
-	/** A handle to the detail view builder */
-	IDetailLayoutBuilder* DetailBuilder;
-
-	/** Access to the Parent Property */
-	TSharedPtr<IPropertyHandle> TargetShaderFormatsPropertyHandle;
-};
-
-
 /**
  * Manages the Transform section of a details view                    
  */
@@ -100,7 +60,7 @@ protected:
 
 private:
 	/** Reference to the target shader formats property view */
-	TSharedPtr<FLinuxTargetShaderFormatsPropertyDetails> TargetShaderFormatsDetails;
+	TSharedPtr<FShaderFormatsPropertyDetails> TargetShaderFormatsDetails;
 
 	/** Manager for Audio Plugin widget. */
 	FAudioPluginWidgetManager AudioPluginWidgetManager;
diff --git a/Engine/Source/Editor/DetailCustomizations/Private/MacTargetSettingsDetails.cpp b/Engine/Source/Editor/DetailCustomizations/Private/MacTargetSettingsDetails.cpp
--- a/Engine/Source/Editor/DetailCustomizations/Private/MacTargetSettingsDetails.cpp
+++ b/Engine/Source/Editor/DetailCustomizations/Private/MacTargetSettingsDetails.cpp
@@ -109,7 +109,7 @@ void FMacTargetSettingsDetails::CustomizeDetails( IDetailLayoutBuilder& DetailBu
 	// Setup the supported/targeted RHI property view
 	TargetShaderFormatsDetails = MakeShareable(new FShaderFormatsPropertyDetails(&DetailBuilder, TEXT("TargetedRHIs"), TEXT("Targeted RHIs")));
 	TargetShaderFormatsDetails->SetOnUpdateShaderWarning(OnUpdateShaderStandardWarning);
-	TargetShaderFormatsDetails->CreateTargetShaderFormatsPropertyView(TargetPlatform);
+	TargetShaderFormatsDetails->CreateTargetShaderFormatsPropertyView(TargetPlatform, FShaderFormatsPropertyDetails::GetFriendlyNameFromRHINameMac);
 	
 	// Setup the shader version property view
     // Handle max. shader version a little specially.
diff --git a/Engine/Source/Editor/DetailCustomizations/Private/WindowsTargetSettingsDetails.cpp b/Engine/Source/Editor/DetailCustomizations/Private/WindowsTargetSettingsDetails.cpp
--- a/Engine/Source/Editor/DetailCustomizations/Private/WindowsTargetSettingsDetails.cpp
+++ b/Engine/Source/Editor/DetailCustomizations/Private/WindowsTargetSettingsDetails.cpp
@@ -161,8 +161,9 @@ static FString GetIconFilename(EImageScope::Type Scope)
 void FWindowsTargetSettingsDetails::CustomizeDetails( IDetailLayoutBuilder& DetailBuilder )
 {
 	// Setup the supported/targeted RHI property view
-	TargetShaderFormatsDetails = MakeShareable(new FTargetShaderFormatsPropertyDetails(&DetailBuilder));
-	TargetShaderFormatsDetails->CreateTargetShaderFormatsPropertyView();
+	ITargetPlatform* TargetPlatform = FModuleManager::GetModuleChecked<ITargetPlatformModule>("WindowsTargetPlatform").GetTargetPlatforms()[0];
+	TargetShaderFormatsDetails = MakeShareable(new FShaderFormatsPropertyDetails(&DetailBuilder));
+	TargetShaderFormatsDetails->CreateTargetShaderFormatsPropertyView(TargetPlatform, GetFriendlyNameFromRHIName);
 
 	TSharedRef<IPropertyHandle> MinOSProperty = DetailBuilder.GetProperty("MinimumOSVersion");
 	IDetailCategoryBuilder& OSInfoCategory = DetailBuilder.EditCategory(TEXT("OS Info"));
@@ -462,99 +463,4 @@ TSharedRef<SWidget> FWindowsTargetSettingsDetails::MakeAudioDeviceMenu(const TSh
 	return MenuBuilder.MakeWidget();
 }
 
-FTargetShaderFormatsPropertyDetails::FTargetShaderFormatsPropertyDetails(IDetailLayoutBuilder* InDetailBuilder)
-: DetailBuilder(InDetailBuilder)
-{
-	TargetShaderFormatsPropertyHandle = DetailBuilder->GetProperty("TargetedRHIs");
-	ensure(TargetShaderFormatsPropertyHandle.IsValid());
-}
-
-void FTargetShaderFormatsPropertyDetails::CreateTargetShaderFormatsPropertyView()
-{
-	DetailBuilder->HideProperty(TargetShaderFormatsPropertyHandle);
-
-	// List of supported RHI's and selected targets
-	// @todo targetplatform: This [0] will probably always work, all of the windows TPs should return the same GetAllPossibleShaderFormats()
-	ITargetPlatform* WindowsTargetPlatform = FModuleManager::GetModuleChecked<ITargetPlatformModule>("WindowsTargetPlatform").GetTargetPlatforms()[0];
-	TArray<FName> ShaderFormats;
-	WindowsTargetPlatform->GetAllPossibleShaderFormats(ShaderFormats);
-
-	IDetailCategoryBuilder& TargetedRHICategoryBuilder = DetailBuilder->EditCategory(TEXT("Targeted RHIs"));
-
-	for (const FName& ShaderFormat : ShaderFormats)
-	{
-		FText FriendlyShaderFormatName = GetFriendlyNameFromRHIName(ShaderFormat.ToString());
-		// Skip explicitly removed entries
-		if (!FriendlyShaderFormatName.IsEmpty())
-		{
-			FDetailWidgetRow& TargetedRHIWidgetRow = TargetedRHICategoryBuilder.AddCustomRow(FriendlyShaderFormatName);
-
-			TargetedRHIWidgetRow
-				.NameContent()
-				[
-					SNew(SHorizontalBox)
-					+ SHorizontalBox::Slot()
-				.Padding(FMargin(0, 1, 0, 1))
-				.FillWidth(1.0f)
-				[
-					SNew(STextBlock)
-					.Text(FriendlyShaderFormatName)
-				.Font(DetailBuilder->GetDetailFont())
-				]
-				]
-			.ValueContent()
-				[
-					SNew(SCheckBox)
-					.OnCheckStateChanged(this, &FTargetShaderFormatsPropertyDetails::OnTargetedRHIChanged, ShaderFormat)
-				.IsChecked(this, &FTargetShaderFormatsPropertyDetails::IsTargetedRHIChecked, ShaderFormat)
-				];
-		}
-	}
-}
-
-
-void FTargetShaderFormatsPropertyDetails::OnTargetedRHIChanged(ECheckBoxState InNewValue, FName InRHIName)
-{
-	TArray<void*> RawPtrs;
-	TargetShaderFormatsPropertyHandle->AccessRawData(RawPtrs);
-
-	// Update the CVars with the selection
-	{
-		TargetShaderFormatsPropertyHandle->NotifyPreChange();
-		for (void* RawPtr : RawPtrs)
-		{
-			TArray<FString>& Array = *(TArray<FString>*)RawPtr;
-			if(InNewValue == ECheckBoxState::Checked)
-			{
-				Array.Add(InRHIName.ToString());
-			}
-			else
-			{
-				Array.Remove(InRHIName.ToString());
-			}
-		}
-		TargetShaderFormatsPropertyHandle->NotifyPostChange();
-	}
-}
-
-
-ECheckBoxState FTargetShaderFormatsPropertyDetails::IsTargetedRHIChecked(FName InRHIName) const
-{
-	ECheckBoxState CheckState = ECheckBoxState::Unchecked;
-
-	TArray<void*> RawPtrs;
-	TargetShaderFormatsPropertyHandle->AccessRawData(RawPtrs);
-	
-	for(void* RawPtr : RawPtrs)
-	{
-		TArray<FString>& Array = *(TArray<FString>*)RawPtr;
-		if(Array.Contains(InRHIName.ToString()))
-		{
-			CheckState = ECheckBoxState::Checked;
-		}
-	}
-	return CheckState;
-}
-
-
 #undef LOCTEXT_NAMESPACE
diff --git a/Engine/Source/Editor/DetailCustomizations/Private/WindowsTargetSettingsDetails.h b/Engine/Source/Editor/DetailCustomizations/Private/WindowsTargetSettingsDetails.h
--- a/Engine/Source/Editor/DetailCustomizations/Private/WindowsTargetSettingsDetails.h
+++ b/Engine/Source/Editor/DetailCustomizations/Private/WindowsTargetSettingsDetails.h
@@ -6,53 +6,13 @@
 #include "Styling/SlateColor.h"
 #include "Widgets/SWidget.h"
 #include "IDetailCustomization.h"
+#include "ShaderFormatsPropertyDetails.h"
 #include "TargetPlatformAudioCustomization.h"
 
 class IDetailLayoutBuilder;
 class IPropertyHandle;
 enum class ECheckBoxState : uint8;
 
-/* FTargetShaderFormatsPropertyDetails
-*****************************************************************************/
-
-/**
-* Helper which implements details panel customizations for a device profiles parent property
-*/
-class FTargetShaderFormatsPropertyDetails
-	: public TSharedFromThis<FTargetShaderFormatsPropertyDetails>
-{
-
-public:
-
-	/**
-	 * Constructor for the parent property details view
-	 *
-	 * @param InDetailsBuilder - Where we are adding our property view to
-	 */
-	FTargetShaderFormatsPropertyDetails(IDetailLayoutBuilder* InDetailBuilder);
-
-
-	/**
-	 * Create the UI to select which windows shader formats we are targetting
-	 */
-	void CreateTargetShaderFormatsPropertyView();
-
-private:
-
-	// Supported/Targeted RHI check boxes
-	void OnTargetedRHIChanged(ECheckBoxState InNewValue, FName InRHIName);
-	ECheckBoxState IsTargetedRHIChecked(FName InRHIName) const;
-
-private:
-
-	/** A handle to the detail view builder */
-	IDetailLayoutBuilder* DetailBuilder;
-
-	/** Access to the Parent Property */
-	TSharedPtr<IPropertyHandle> TargetShaderFormatsPropertyHandle;
-};
-
-
 /**
  * Manages the Transform section of a details view                    
  */
@@ -100,6 +60,6 @@ protected:
 
 private:
 	/** Reference to the target shader formats property view */
-	TSharedPtr<FTargetShaderFormatsPropertyDetails> TargetShaderFormatsDetails;
+	TSharedPtr<FShaderFormatsPropertyDetails> TargetShaderFormatsDetails;
 	FAudioPluginWidgetManager AudioPluginWidgetManager;
 };
