PR #4967: Add Hide Unconnected Pins button to Class Defaults (Contributed by [REDACTED_PERSON])

diff --git a/Engine/Source/Editor/DetailCustomizations/Private/SkeletalControlNodeDetails.cpp b/Engine/Source/Editor/DetailCustomizations/Private/SkeletalControlNodeDetails.cpp
--- a/Engine/Source/Editor/DetailCustomizations/Private/SkeletalControlNodeDetails.cpp
+++ b/Engine/Source/Editor/DetailCustomizations/Private/SkeletalControlNodeDetails.cpp
@@ -15,6 +15,7 @@
 #include "K2Node.h"
 #include "Widgets/Input/SButton.h"
 #include "K2Node_BreakStruct.h"
+#include "K2Node_GetClassDefaults.h"
 #include "ScopedTransaction.h"
 
 #define LOCTEXT_NAMESPACE "SkeletalControlNodeDetails"
@@ -36,16 +37,16 @@ void FSkeletalControlNodeDetails::CustomizeDetails(class IDetailLayoutBuilder& D
 	const TArray< TWeakObjectPtr<UObject> >& SelectedObjects = DetailBuilder.GetSelectedObjects();
 	for (const TWeakObjectPtr<UObject>& CurrentObject : SelectedObjects)
 	{
-		if (UK2Node_BreakStruct* CurrBreakStruct = Cast<UK2Node_BreakStruct>(CurrentObject.Get()))
+		if (Cast<UK2Node_BreakStruct>(CurrentObject.Get()) || Cast<UK2Node_GetClassDefaults>(CurrentObject.Get()))
 		{
-			if (BreakStructNode.IsValid())
+			if (HideUnconnectedPinsNode.IsValid())
 			{
 				// Have more than one break struct node, don't cache so we don't
 				// create the hide unconnected pins UI
-				BreakStructNode = nullptr;
+				HideUnconnectedPinsNode = nullptr;
 				break;
 			}
-			BreakStructNode = CurrBreakStruct;
+			HideUnconnectedPinsNode = Cast<UK2Node>(CurrentObject.Get());
 		}
 	}
 
@@ -90,7 +91,7 @@ void FSkeletalControlNodeDetails::CustomizeDetails(class IDetailLayoutBuilder& D
 	}
 
 	// Add the action buttons
-	if(BreakStructNode.IsValid())
+	if(HideUnconnectedPinsNode.IsValid())
 	{
 		FDetailWidgetRow& GroupActionsRow = DetailCategory->AddCustomRow(LOCTEXT("GroupActionsSearchText", "Split Sort"))
 		.ValueContent()
@@ -99,7 +100,7 @@ void FSkeletalControlNodeDetails::CustomizeDetails(class IDetailLayoutBuilder& D
 		[
 			SNew(SButton)
 			.OnClicked(this, &FSkeletalControlNodeDetails::HideAllUnconnectedPins)
-			.ToolTipText(LOCTEXT("HideAllUnconnectedPinsTooltip", "All unconnected pins of the structure get hidden (removed from the graph node)"))
+			.ToolTipText(LOCTEXT("HideAllUnconnectedPinsTooltip", "All unconnected pins get hidden (removed from the graph node)"))
 			[
 				SNew(STextBlock)
 				.Text(LOCTEXT("HideAllUnconnectedPins", "Hide Unconnected Pins"))
@@ -246,7 +247,7 @@ void FSkeletalControlNodeDetails::OnGenerateElementForPropertyPin(TSharedRef<IPr
 
 FReply FSkeletalControlNodeDetails::HideAllUnconnectedPins()
 {
-	if (ArrayProperty.IsValid() && BreakStructNode.IsValid())
+	if (ArrayProperty.IsValid() && HideUnconnectedPinsNode.IsValid())
 	{
 		uint32 NumChildren = 0;
 		ArrayProperty->GetNumElements(NumChildren);
@@ -261,7 +262,7 @@ FReply FSkeletalControlNodeDetails::HideAllUnconnectedPins()
 			FName ActualPropertyName;
 			if (PropertyNameHandle.IsValid() && PropertyNameHandle->GetValue(ActualPropertyName) == FPropertyAccess::Success)
 			{
-				const UEdGraphPin* Pin = BreakStructNode->FindPin(ActualPropertyName.ToString(), EGPD_Output);
+				const UEdGraphPin* Pin = HideUnconnectedPinsNode->FindPin(ActualPropertyName.ToString(), EGPD_Output);
 				if (Pin && Pin->LinkedTo.Num() <= 0)
 				{
 					OnShowPinChanged(ECheckBoxState::Unchecked, ElementHandle);
diff --git a/Engine/Source/Editor/DetailCustomizations/Private/SkeletalControlNodeDetails.h b/Engine/Source/Editor/DetailCustomizations/Private/SkeletalControlNodeDetails.h
--- a/Engine/Source/Editor/DetailCustomizations/Private/SkeletalControlNodeDetails.h
+++ b/Engine/Source/Editor/DetailCustomizations/Private/SkeletalControlNodeDetails.h
@@ -37,6 +37,6 @@ protected:
 
 private:
 	IDetailCategoryBuilder* DetailCategory;
-	TWeakObjectPtr<class UK2Node_BreakStruct> BreakStructNode;
+	TWeakObjectPtr<class UK2Node> HideUnconnectedPinsNode;
 	TSharedPtr<class IPropertyHandleArray> ArrayProperty;
 };
