Skeletal Reduction: Allow user to specify absolute vertex or triangle count.

diff --git a/Engine/Plugins/Experimental/SkeletalReduction/Source/Private/SkeletalMeshReductionPlugin.cpp b/Engine/Plugins/Experimental/SkeletalReduction/Source/Private/SkeletalMeshReductionPlugin.cpp
--- a/Engine/Plugins/Experimental/SkeletalReduction/Source/Private/SkeletalMeshReductionPlugin.cpp
+++ b/Engine/Plugins/Experimental/SkeletalReduction/Source/Private/SkeletalMeshReductionPlugin.cpp
@@ -826,21 +826,26 @@ float FQuadricSkeletalMeshReduction::SimplifyMesh( const FSkeletalMeshOptimizati
 
 	// Determine the stop criteria used
 
-	const bool bUseVertexCriterion   = Settings.TerminationCriterion == SkeletalMeshTerminationCriterion::SMTC_NumOfVerts     || Settings.TerminationCriterion == SkeletalMeshTerminationCriterion::SMTC_TriangleOrVert;
-	const bool bUseTriangleCriterion = Settings.TerminationCriterion == SkeletalMeshTerminationCriterion::SMTC_NumOfTriangles || Settings.TerminationCriterion == SkeletalMeshTerminationCriterion::SMTC_TriangleOrVert;
+	const bool bUseVertexPercentCriterion = Settings.TerminationCriterion == SkeletalMeshTerminationCriterion::SMTC_NumOfVerts || Settings.TerminationCriterion == SkeletalMeshTerminationCriterion::SMTC_TriangleOrVert;
+	const bool bUseTrianglePercentCriterion = Settings.TerminationCriterion == SkeletalMeshTerminationCriterion::SMTC_NumOfTriangles || Settings.TerminationCriterion == SkeletalMeshTerminationCriterion::SMTC_TriangleOrVert;
 
+	const bool bUseMaxVertNumCriterion = Settings.TerminationCriterion == SkeletalMeshTerminationCriterion::SMTC_AbsNumOfVerts || Settings.TerminationCriterion == SkeletalMeshTerminationCriterion::SMTC_AbsTriangleOrVert;
+	const bool bUseMaxTrisNumCriterion = Settings.TerminationCriterion == SkeletalMeshTerminationCriterion::SMTC_AbsNumOfTriangles || Settings.TerminationCriterion == SkeletalMeshTerminationCriterion::SMTC_AbsTriangleOrVert;
 
 	// We can support a stopping criteria based on the MaxDistance the new vertex is from the plans of the source triangles.
 	// but there seems to be no good use for this.  We are better off just using triangle count.
-	const float MaxDist              = FLT_MAX; // (Settings.ReductionMethod != SkeletalMeshOptimizationType::SMOT_NumOfTriangles) ? Settings.MaxDeviationPercentage * Bounds.SphereRadius : FLT_MAX;
-	const int32 MaxTriNumToRetain    = Mesh.NumIndices() / 3;
-	const float TriangleRetainRatio  = FMath::Clamp(Settings.NumOfTrianglesPercentage, 0.f, 1.f);
-	const int32 MinTriNumToRetain    = (bUseTriangleCriterion) ? FMath::Max(4, FMath::CeilToInt(TriangleRetainRatio * MaxTriNumToRetain)) : 4 ; 	
-	const float MaxCollapseCost      = FLT_MAX;
+	const float MaxDist = FLT_MAX; // (Settings.ReductionMethod != SkeletalMeshOptimizationType::SMOT_NumOfTriangles) ? Settings.MaxDeviationPercentage * Bounds.SphereRadius : FLT_MAX;
+	const int32 SrcTriNum = Mesh.NumIndices() / 3;
+	const float TriangleRetainRatio = FMath::Clamp(Settings.NumOfTrianglesPercentage, 0.f, 1.f);
+	const int32 TargetTriNum = (bUseTrianglePercentCriterion) ? FMath::CeilToInt(TriangleRetainRatio * SrcTriNum) : Settings.MaxNumOfTriangles;
 
-	const int32 MaxVerNumToRetain = Mesh.NumVertices();
-	const float VertRetainRatio   = FMath::Clamp(Settings.NumOfVertPercentage, 0.f, 1.f);
-	const int32 MinVerNumToRetain = (bUseVertexCriterion) ? FMath::Max(3, FMath::CeilToInt(VertRetainRatio * MaxVerNumToRetain)) : 0;
+	const int32 MinTriNumToRetain = (bUseTrianglePercentCriterion || bUseMaxTrisNumCriterion) ? FMath::Max(4, TargetTriNum) : 4;
+	const float MaxCollapseCost = FLT_MAX;
+
+	const int32 SrcVertNum = Mesh.NumVertices();
+	const float VertRetainRatio = FMath::Clamp(Settings.NumOfVertPercentage, 0.f, 1.f);
+	const int32 TargetVertNum = (bUseVertexPercentCriterion) ? FMath::CeilToInt(VertRetainRatio * SrcVertNum) : Settings.MaxNumOfVerts + 1;
+	const int32 MinVerNumToRetain = (bUseVertexPercentCriterion || bUseMaxVertNumCriterion) ? FMath::Max(6, TargetVertNum) : 6;
 
 	const float VolumeImportance      = FMath::Clamp(Settings.VolumeImportance, 0.f, 2.f);
 	const bool bLockEdges             = Settings.bLockEdges;
@@ -848,7 +853,7 @@ float FQuadricSkeletalMeshReduction::SimplifyMesh( const FSkeletalMeshOptimizati
 	const bool bEnforceBoneBoundaries = Settings.bEnforceBoneBoundaries;
 
 	// Terminator tells the simplifier when to stop
-	SkeletalSimplifier::FSimplifierTerminator Terminator(MinTriNumToRetain, MaxTriNumToRetain, MinVerNumToRetain, MaxVerNumToRetain, MaxCollapseCost, MaxDist);
+	SkeletalSimplifier::FSimplifierTerminator Terminator(MinTriNumToRetain, SrcTriNum, MinVerNumToRetain, SrcVertNum, MaxCollapseCost, MaxDist);
 
 	double NormalWeight    =16.00;
 	double TangentWeight   = 0.10;
diff --git a/Engine/Source/Editor/DetailCustomizations/Private/SkeletalMeshReductionSettingsDetails.cpp b/Engine/Source/Editor/DetailCustomizations/Private/SkeletalMeshReductionSettingsDetails.cpp
--- a/Engine/Source/Editor/DetailCustomizations/Private/SkeletalMeshReductionSettingsDetails.cpp
+++ b/Engine/Source/Editor/DetailCustomizations/Private/SkeletalMeshReductionSettingsDetails.cpp
@@ -107,6 +107,8 @@ void FSkeletalMeshReductionSettingsDetails::CustomizeChildren(TSharedRef<IProper
 		// Parameters not used by simplygon
 		const TArray<FName> CustomSimplifierOnlyProperties = {
 			GET_MEMBER_NAME_CHECKED(FSkeletalMeshOptimizationSettings, NumOfVertPercentage),
+			GET_MEMBER_NAME_CHECKED(FSkeletalMeshOptimizationSettings, MaxNumOfVerts),
+			GET_MEMBER_NAME_CHECKED(FSkeletalMeshOptimizationSettings, MaxNumOfTriangles),
 			GET_MEMBER_NAME_CHECKED(FSkeletalMeshOptimizationSettings, TerminationCriterion),
 			GET_MEMBER_NAME_CHECKED(FSkeletalMeshOptimizationSettings, bLockEdges),
 			GET_MEMBER_NAME_CHECKED(FSkeletalMeshOptimizationSettings, bEnforceBoneBoundaries),
@@ -176,6 +178,8 @@ void FSkeletalMeshReductionSettingsDetails::CustomizeChildren(TSharedRef<IProper
 		TSharedPtr<IPropertyHandle> VertPercentPropertyHandle = PropertyHandles.FindChecked(GET_MEMBER_NAME_CHECKED(FSkeletalMeshOptimizationSettings, NumOfVertPercentage));
 		TSharedPtr<IPropertyHandle> TriPercentPropertyHandle  = PropertyHandles.FindChecked(GET_MEMBER_NAME_CHECKED(FSkeletalMeshOptimizationSettings, NumOfTrianglesPercentage));
 
+		TSharedPtr<IPropertyHandle> MaxNumOfVertsPropertyHandle = PropertyHandles.FindChecked(GET_MEMBER_NAME_CHECKED(FSkeletalMeshOptimizationSettings, MaxNumOfVerts));
+		TSharedPtr<IPropertyHandle> MaxNumOfTrisPropertyHandle = PropertyHandles.FindChecked(GET_MEMBER_NAME_CHECKED(FSkeletalMeshOptimizationSettings, MaxNumOfTriangles));
 
 		for (auto Iter(PropertyHandles.CreateConstIterator()); Iter; ++Iter)
 		{
@@ -191,13 +195,43 @@ void FSkeletalMeshReductionSettingsDetails::CustomizeChildren(TSharedRef<IProper
 				// depending on the value of the pull down, optionally hide at most one of these.
 				if (Iter.Value() == VertPercentPropertyHandle)
 				{
+					const TArray< SkeletalMeshTerminationCriterion > VizList = {
+					SMTC_NumOfVerts,
+					SMTC_TriangleOrVert
+					};
+
 					// Hide property if using triangle percentage
-					SettingsRow.Visibility(TAttribute<EVisibility>::Create(TAttribute<EVisibility>::FGetter::CreateSP(this, &FSkeletalMeshReductionSettingsDetails::HideIfCurrentCriterionIs, SMTC_NumOfTriangles)));
+					SettingsRow.Visibility(TAttribute<EVisibility>::Create(TAttribute<EVisibility>::FGetter::CreateSP(this, &FSkeletalMeshReductionSettingsDetails::ShowIfCurrentCriterionIs, VizList)));
 				}
 				else if (Iter.Value() == TriPercentPropertyHandle)
 				{
+					const TArray< SkeletalMeshTerminationCriterion > VizList = 
+					{
+						SMTC_NumOfTriangles,
+						SMTC_TriangleOrVert
+					};
 					// Hide property if using vert percentage
-					SettingsRow.Visibility(TAttribute<EVisibility>::Create(TAttribute<EVisibility>::FGetter::CreateSP(this, &FSkeletalMeshReductionSettingsDetails::HideIfCurrentCriterionIs, SMTC_NumOfVerts)));
+					SettingsRow.Visibility(TAttribute<EVisibility>::Create(TAttribute<EVisibility>::FGetter::CreateSP(this, &FSkeletalMeshReductionSettingsDetails::ShowIfCurrentCriterionIs, VizList)));
+				}
+				else if (Iter.Value() == MaxNumOfVertsPropertyHandle)
+				{
+					const TArray< SkeletalMeshTerminationCriterion > VizList = 
+					{
+						SMTC_AbsNumOfVerts,
+						SMTC_AbsTriangleOrVert
+					};
+					// Hide property if using triangle percentage
+					SettingsRow.Visibility(TAttribute<EVisibility>::Create(TAttribute<EVisibility>::FGetter::CreateSP(this, &FSkeletalMeshReductionSettingsDetails::ShowIfCurrentCriterionIs, VizList)));
+				}
+				else if (Iter.Value() == MaxNumOfTrisPropertyHandle)
+				{
+					const TArray< SkeletalMeshTerminationCriterion > VizList = 
+					{
+						SMTC_AbsNumOfTriangles,
+						SMTC_AbsTriangleOrVert
+					};
+					// Hide property if using triangle percentage
+					SettingsRow.Visibility(TAttribute<EVisibility>::Create(TAttribute<EVisibility>::FGetter::CreateSP(this, &FSkeletalMeshReductionSettingsDetails::ShowIfCurrentCriterionIs, VizList)));
 				}
 
 			}
@@ -266,19 +300,19 @@ EVisibility FSkeletalMeshReductionSettingsDetails::GetVisibiltyIfCurrentReductio
 	return EVisibility::Hidden;
 }
 
-EVisibility FSkeletalMeshReductionSettingsDetails::HideIfCurrentCriterionIs(SkeletalMeshTerminationCriterion TerminationCriterion) const
+EVisibility FSkeletalMeshReductionSettingsDetails::ShowIfCurrentCriterionIs(TArray<SkeletalMeshTerminationCriterion> TerminationCriterionArray) const
 {
 	uint8 CurrentEnum;
 	if (TerminationCriterionPopertyHandle->GetValue(CurrentEnum) != FPropertyAccess::Fail)
 	{
 		enum SkeletalMeshTerminationCriterion CurrentReductionType = (SkeletalMeshTerminationCriterion)CurrentEnum;
-		if (CurrentReductionType == TerminationCriterion)
+		if (TerminationCriterionArray.Contains(CurrentReductionType))
 		{
-			return EVisibility::Hidden; 
+			return EVisibility::Visible;
 		}
 	}
 
-	return EVisibility::Visible;
+	return EVisibility::Hidden;
 }
 
 bool FSkeletalMeshReductionSettingsDetails::UseNativeLODTool() const
diff --git a/Engine/Source/Editor/DetailCustomizations/Private/SkeletalMeshReductionSettingsDetails.h b/Engine/Source/Editor/DetailCustomizations/Private/SkeletalMeshReductionSettingsDetails.h
--- a/Engine/Source/Editor/DetailCustomizations/Private/SkeletalMeshReductionSettingsDetails.h
+++ b/Engine/Source/Editor/DetailCustomizations/Private/SkeletalMeshReductionSettingsDetails.h
@@ -39,7 +39,7 @@ private:
 	EVisibility GetVisibiltyIfCurrentReductionMethodIsNot(SkeletalMeshOptimizationType ReductionType) const;
 
 	// Used by the native tool UI.
-	EVisibility HideIfCurrentCriterionIs(SkeletalMeshTerminationCriterion TerminationCriterion) const;
+	EVisibility ShowIfCurrentCriterionIs(TArray<SkeletalMeshTerminationCriterion> TerminationCriterionArray) const;
 
 	/** Detect usage of thirdparty vs native tool */
 	bool UseNativeLODTool() const;
diff --git a/Engine/Source/Runtime/Engine/Public/SkeletalMeshReductionSettings.h b/Engine/Source/Runtime/Engine/Public/SkeletalMeshReductionSettings.h
--- a/Engine/Source/Runtime/Engine/Public/SkeletalMeshReductionSettings.h
+++ b/Engine/Source/Runtime/Engine/Public/SkeletalMeshReductionSettings.h
@@ -18,7 +18,10 @@ enum SkeletalMeshTerminationCriterion
 {
 	SMTC_NumOfTriangles UMETA(DisplayName = "Triangles", ToolTip = "Triangle count criterion will be used for simplification."),
 	SMTC_NumOfVerts UMETA(DisplayName = "Vertices", ToolTip = "Vertex cont criterion will be used for simplification."),
-	SMTC_TriangleOrVert UMETA(DisplayName = "First Satisfied", ToolTip = "Simplification will continue until either Triangle or Vertex count criteria is met."),
+	SMTC_TriangleOrVert UMETA(DisplayName = "First Percent Satisfied", ToolTip = "Simplification will continue until either Triangle or Vertex count criteria is met."),
+	SMTC_AbsNumOfTriangles UMETA(DisplayName = "Max Triangles", ToolTip = "Triangle count criterion will be used for simplification."),
+	SMTC_AbsNumOfVerts UMETA(DisplayName = "Max Vertices", ToolTip = "Vertex cont criterion will be used for simplification."),
+	SMTC_AbsTriangleOrVert UMETA(DisplayName = "First Max Satisfied", ToolTip = "Simplification will continue until either Triangle or Vertex count criteria is met."),
 	SMTC_MAX UMETA(Hidden),
 };
 
@@ -64,6 +67,14 @@ struct FSkeletalMeshOptimizationSettings
 	/** The percentage of vertices to retain as a ratio, e.g. 0.1 indicates 10 percent */
 	UPROPERTY(EditAnywhere, Category = ReductionMethod, meta = (DisplayName = "Percent of Vertices"))
 	float NumOfVertPercentage;
+	
+	/** The maximum number of triangles to retain */
+	UPROPERTY(EditAnywhere, Category = ReductionMethod, meta = (DisplayName = "Max Triangle Count", ClampMin = 4))
+	uint32 MaxNumOfTriangles;
+
+	/** The maximum number of vertices to retain */
+	UPROPERTY(EditAnywhere, Category = ReductionMethod, meta = (DisplayName = "Max Vertex Count", ClampMin = 6))
+	uint32 MaxNumOfVerts;
 
 	/**If ReductionMethod equals MaxDeviation this value is the maximum deviation from the base mesh as a percentage of the bounding sphere. 
 	 * In code, it ranges from [0, 1]. In the editor UI, it ranges from [0, 100]
@@ -140,6 +151,8 @@ struct FSkeletalMeshOptimizationSettings
 		: TerminationCriterion(SMTC_NumOfTriangles)
 		, NumOfTrianglesPercentage(0.5f)
 		, NumOfVertPercentage(0.5f)
+		, MaxNumOfTriangles(4)
+		, MaxNumOfVerts(6)
 		, MaxDeviationPercentage(0.5f)
 		, ReductionMethod(SMOT_NumOfTriangles)
 		, SilhouetteImportance(SMOI_Normal)
