Crash regenerating an LOD with Base LOD set to the LOD that is being regenerated
#fix Added details customization to restrict the spinbox to < CurrentLODIndex

diff --git a/Engine/Source/Editor/DetailCustomizations/Private/SkeletalMeshReductionSettingsDetails.cpp b/Engine/Source/Editor/DetailCustomizations/Private/SkeletalMeshReductionSettingsDetails.cpp
--- a/Engine/Source/Editor/DetailCustomizations/Private/SkeletalMeshReductionSettingsDetails.cpp
+++ b/Engine/Source/Editor/DetailCustomizations/Private/SkeletalMeshReductionSettingsDetails.cpp
@@ -13,6 +13,7 @@
 #include "SkeletalMeshReductionSettings.h"
 #include "Widgets/Input/SSpinBox.h"
 #include "Widgets/Text/STextBlock.h"
+#include "SkeletalRenderPublic.h"
 
 #define LOCTEXT_NAMESPACE "SkeletalMeshReductionSettingsDetails"
 
@@ -59,10 +60,63 @@ void FSkeletalMeshReductionSettingsDetails::CustomizeChildren(TSharedRef<IProper
 
 	bool bUseThirdPartyUI = !UseNativeReductionTool();
 
+	const int32 LODIndex = [StructPropertyHandle]() -> int32
+	{
+		if (StructPropertyHandle->GetParentHandle().IsValid())
+		{
+			if (StructPropertyHandle->GetParentHandle()->GetProperty()->GetFName() == GET_MEMBER_NAME_CHECKED(FSkeletalMeshObject, LODInfo))
+			{
+				return StructPropertyHandle->GetParentHandle()->GetIndexInArray();
+			}
+		}
+
+		return INDEX_NONE;
+	}();
 
+	auto BaseLODCustomization = [LODIndex, &StructBuilder](TSharedPtr<IPropertyHandle>& BaseLODPropertyHandle)
+	{
+		// Only able to do this for LOD2 and above, so only show the property if this is the case
+		if (LODIndex > 1)
+		{
+			// Add and retrieve the default widgets
+			IDetailPropertyRow& Row = StructBuilder.AddProperty(BaseLODPropertyHandle->AsShared());
+
+			TSharedPtr<SWidget> NameWidget;
+			TSharedPtr<SWidget> ValueWidget;
+			FDetailWidgetRow DefaultWidgetRow;
+			Row.GetDefaultWidgets(NameWidget, ValueWidget, DefaultWidgetRow);
+
+			// Customize the value property to be a spinbox with Max value cap so it is always < CurrentLODIndex
+			Row.CustomWidget()
+				.NameContent()
+				[
+					NameWidget.ToSharedRef()
+				]
+				.ValueContent()
+				.MinDesiredWidth(DefaultWidgetRow.ValueWidget.MinWidth)
+				.MaxDesiredWidth(DefaultWidgetRow.ValueWidget.MaxWidth)
+				[
+					SNew(SSpinBox<int32>)
+					.Font(IDetailLayoutBuilder::GetDetailFont())
+					.MinValue(0)
+					.MaxValue(FMath::Max(LODIndex - 1, 0))
+					.Value_Lambda([BaseLODPropertyHandle]() -> int32
+					{
+						int32 Value = INDEX_NONE;
+						BaseLODPropertyHandle->GetValue(Value);
+						return Value;
+					})					
+					.OnValueChanged_Lambda([BaseLODPropertyHandle](int32 NewValue)
+					{
+						BaseLODPropertyHandle->SetValue(NewValue);
+					})
+				];
+		}
+	};
+	
+	
 	if (bUseThirdPartyUI)
 	{
-
 		StructBuilder.AddProperty(ReductionMethodPropertyHandle.ToSharedRef());
 
 		StructBuilder.AddCustomRow(LOCTEXT("PercentTriangles_Row", "Triangle Percentage"))
@@ -123,10 +177,18 @@ void FSkeletalMeshReductionSettingsDetails::CustomizeChildren(TSharedRef<IProper
 				// we don't want to add the things that we added first
 				// maybe we make array later if we have a lot. 
 				FName PropertyName = ChildHandle->GetProperty()->GetFName();
-				if (!CustomizedProperties.Contains(PropertyName) && !CustomSimplifierOnlyProperties.Contains(PropertyName))
+
+				if (PropertyName == GET_MEMBER_NAME_CHECKED(FSkeletalMeshOptimizationSettings, BaseLOD))
 				{
-					StructBuilder.AddProperty(ChildHandle.ToSharedRef());
+					BaseLODCustomization(ChildHandle);
 				}
+				else
+				{
+					if (!CustomizedProperties.Contains(PropertyName) && !CustomSimplifierOnlyProperties.Contains(PropertyName))
+					{
+						StructBuilder.AddProperty(ChildHandle.ToSharedRef());
+					}
+				}				
 			}
 		}
 	}
@@ -177,13 +239,17 @@ void FSkeletalMeshReductionSettingsDetails::CustomizeChildren(TSharedRef<IProper
 		TSharedPtr<IPropertyHandle> TriPercentPropertyHandle  = PropertyHandles.FindChecked(GET_MEMBER_NAME_CHECKED(FSkeletalMeshOptimizationSettings, NumOfTrianglesPercentage));
 
 
-		for (auto Iter(PropertyHandles.CreateConstIterator()); Iter; ++Iter)
+		for (auto Iter(PropertyHandles.CreateIterator()); Iter; ++Iter)
 		{
 			if (UnwantedPropertyHandles.Contains(Iter.Value())) 
 			{
 				IDetailPropertyRow& SettingsRow = StructBuilder.AddProperty(Iter.Value().ToSharedRef());
 				SettingsRow.Visibility(TAttribute<EVisibility>(this, &FSkeletalMeshReductionSettingsDetails::GetVisibilityForThirdPartyTool));
 			}
+			else if (Iter.Value()->GetProperty()->GetFName() == GET_MEMBER_NAME_CHECKED(FSkeletalMeshOptimizationSettings, BaseLOD))
+			{
+				BaseLODCustomization(Iter.Value());
+			}
 			else
 			{
 				IDetailPropertyRow& SettingsRow = StructBuilder.AddProperty(Iter.Value().ToSharedRef());
