Fix path to default icon for Windows in editor packaging settings panel.

diff --git a/Engine/Source/Editor/DetailCustomizations/Private/WindowsTargetSettingsDetails.cpp b/Engine/Source/Editor/DetailCustomizations/Private/WindowsTargetSettingsDetails.cpp
--- a/Engine/Source/Editor/DetailCustomizations/Private/WindowsTargetSettingsDetails.cpp
+++ b/Engine/Source/Editor/DetailCustomizations/Private/WindowsTargetSettingsDetails.cpp
@@ -140,7 +140,7 @@ static FString GetWindowsIconFilename(EWindowsImageScope::Type Scope)
 
 	if (Scope == EWindowsImageScope::Engine)
 	{
-		FString Filename = FPaths::EngineDir() / FString(TEXT("Source/Runtime/Launch/Resources")) / PlatformName / FString("UE4.ico");
+		FString Filename = FPaths::EngineDir() / FString(TEXT("Build/Windows/Resources/Default.ico"));
 		return FPaths::ConvertRelativePathToFull(Filename);
 	}
 	else
