Prevent showing empty drop down in the material base options of the fbx import dialog.

diff --git a/Engine/Source/Editor/DetailCustomizations/Private/FbxImportUIDetails.cpp b/Engine/Source/Editor/DetailCustomizations/Private/FbxImportUIDetails.cpp
--- a/Engine/Source/Editor/DetailCustomizations/Private/FbxImportUIDetails.cpp
+++ b/Engine/Source/Editor/DetailCustomizations/Private/FbxImportUIDetails.cpp
@@ -635,183 +635,201 @@ void FFbxImportUIDetails::ConstructBaseMaterialUI(TSharedPtr<IPropertyHandle> Ha
 	TSharedPtr<SWidget> ValueWidget;
 	FDetailWidgetRow Row;
 	MaterialPropertyRow.GetDefaultWidgets(NameWidget, ValueWidget, Row);
+	int InitialSelect = INDEX_NONE;
 
 	// base color properties, only used when there is no texture in the diffuse map
 	Material->GetAllVectorParameterInfo(OutParameterInfo, Guids);
-	for (FMaterialParameterInfo &ParameterInfo : OutParameterInfo)
-	{
-		BaseColorNames.Add(MakeShareable(new FString(ParameterInfo.Name.ToString())));
-	}
-	int InitialSelect = FindString(BaseColorNames, ImportUI->TextureImportData->BaseColorName);
-	InitialSelect = InitialSelect == INDEX_NONE ? 0 : InitialSelect; // default to the empty string located at index 0
-	MaterialCategory.AddCustomRow(LOCTEXT("BaseColorProperty", "Base Color Property"))
-	.NameContent()
-	[
-		SNew(STextBlock)
-		.Text(LOCTEXT("BaseColorProperty", "Base Color Property"))
-		.Font(IDetailLayoutBuilder::GetDetailFont())
-	]
-	.ValueContent()
-	.MaxDesiredWidth(Row.ValueWidget.MaxWidth)
-	[
-		SNew(SHorizontalBox)
-		+ SHorizontalBox::Slot()
-		.AutoWidth()
+	if (OutParameterInfo.Num() > 0)
+	{
+		for (FMaterialParameterInfo &ParameterInfo : OutParameterInfo)
+		{
+			BaseColorNames.Add(MakeShareable(new FString(ParameterInfo.Name.ToString())));
+		}
+	}
+	OutParameterInfo.Empty();
+	Material->GetAllTextureParameterInfo(OutParameterInfo, Guids);
+	if (OutParameterInfo.Num() > 0)
+	{
+		for (FMaterialParameterInfo &ParameterInfo : OutParameterInfo)
+		{
+			BaseTextureNames.Add(MakeShareable(new FString(ParameterInfo.Name.ToString())));
+		}
+	}
+	if (BaseColorNames.Num() > 1)
+	{
+		InitialSelect = FindString(BaseColorNames, ImportUI->TextureImportData->BaseColorName);
+		InitialSelect = InitialSelect == INDEX_NONE ? 0 : InitialSelect; // default to the empty string located at index 0
+		MaterialCategory.AddCustomRow(LOCTEXT("BaseColorProperty", "Base Color Property"))
+		.NameContent()
+		[
+			SNew(STextBlock)
+			.Text(LOCTEXT("BaseColorProperty", "Base Color Property"))
+			.Font(IDetailLayoutBuilder::GetDetailFont())
+		]
+		.ValueContent()
+		.MaxDesiredWidth(Row.ValueWidget.MaxWidth)
 		[
-			SNew(SBox)
-			.MinDesiredWidth(MinDesiredWidth)
+			SNew(SHorizontalBox)
+			+ SHorizontalBox::Slot()
+			.AutoWidth()
 			[
-				SNew(STextComboBox)
-				.OptionsSource(&BaseColorNames)
-				.ToolTip(SNew(SToolTip).Text(LOCTEXT("BaseColorFBXImportToolTip", "When there is no diffuse texture in the imported material this color property will be used to fill a contant color value instead.")))
-				.OnSelectionChanged(this, &FFbxImportUIDetails::OnBaseColor)
-				.InitiallySelectedItem(BaseColorNames[InitialSelect])
+				SNew(SBox)
+				.MinDesiredWidth(MinDesiredWidth)
+				[
+					SNew(STextComboBox)
+					.OptionsSource(&BaseColorNames)
+					.ToolTip(SNew(SToolTip).Text(LOCTEXT("BaseColorFBXImportToolTip", "When there is no diffuse texture in the imported material this color property will be used to fill a contant color value instead.")))
+					.OnSelectionChanged(this, &FFbxImportUIDetails::OnBaseColor)
+					.InitiallySelectedItem(BaseColorNames[InitialSelect])
+				]
 			]
-		]
-	];
-
+		];
+	}
 	// base texture properties
-	OutParameterInfo.Empty();
-	Material->GetAllTextureParameterInfo(OutParameterInfo, Guids);
-	for (FMaterialParameterInfo &ParameterInfo : OutParameterInfo)
-	{
-		BaseTextureNames.Add(MakeShareable(new FString(ParameterInfo.Name.ToString())));
-	}
-	InitialSelect = FindString(BaseTextureNames, ImportUI->TextureImportData->BaseDiffuseTextureName);
-	InitialSelect = InitialSelect == INDEX_NONE ? 0 : InitialSelect; // default to the empty string located at index 0
-	MaterialCategory.AddCustomRow(LOCTEXT("BaseTextureProperty", "Base Texture Property")).NameContent()
-	[
-		SNew(STextBlock)
-		.Text(LOCTEXT("BaseTextureProperty", "Base Texture Property"))
-		.Font(IDetailLayoutBuilder::GetDetailFont())
-	]
-	.ValueContent()
-	.MaxDesiredWidth(Row.ValueWidget.MaxWidth)
-	[
-		SNew(SHorizontalBox)
-		+ SHorizontalBox::Slot()
-		.AutoWidth()
+	if(BaseTextureNames.Num() > 1)
+	{
+		InitialSelect = FindString(BaseTextureNames, ImportUI->TextureImportData->BaseDiffuseTextureName);
+		InitialSelect = InitialSelect == INDEX_NONE ? 0 : InitialSelect; // default to the empty string located at index 0
+		MaterialCategory.AddCustomRow(LOCTEXT("BaseTextureProperty", "Base Texture Property")).NameContent()
+		[
+			SNew(STextBlock)
+			.Text(LOCTEXT("BaseTextureProperty", "Base Texture Property"))
+			.Font(IDetailLayoutBuilder::GetDetailFont())
+		]
+		.ValueContent()
+		.MaxDesiredWidth(Row.ValueWidget.MaxWidth)
 		[
-			SNew(SBox)
-			.MinDesiredWidth(MinDesiredWidth)
+			SNew(SHorizontalBox)
+			+ SHorizontalBox::Slot()
+			.AutoWidth()
 			[
-				SNew(STextComboBox)
-				.OptionsSource(&BaseTextureNames)
-				.OnSelectionChanged(this, &FFbxImportUIDetails::OnDiffuseTextureColor)
-				.InitiallySelectedItem(BaseTextureNames[InitialSelect])
+				SNew(SBox)
+				.MinDesiredWidth(MinDesiredWidth)
+				[
+					SNew(STextComboBox)
+					.OptionsSource(&BaseTextureNames)
+					.OnSelectionChanged(this, &FFbxImportUIDetails::OnDiffuseTextureColor)
+					.InitiallySelectedItem(BaseTextureNames[InitialSelect])
+				]
 			]
+		];
+
+		// base normal properties
+		InitialSelect = FindString(BaseTextureNames, ImportUI->TextureImportData->BaseNormalTextureName);
+		InitialSelect = InitialSelect == INDEX_NONE ? 0 : InitialSelect; // default to the empty string located at index 0
+		MaterialCategory.AddCustomRow(LOCTEXT("BaseNormalTextureProperty", "Base Normal Texture Property")).NameContent()
+		[
+			SNew(STextBlock)
+			.Text(LOCTEXT("BaseNormalTextureProperty", "Base Normal Texture Property"))
+			.Font(IDetailLayoutBuilder::GetDetailFont())
 		]
-	];
-
-	// base normal properties
-	InitialSelect = FindString(BaseTextureNames, ImportUI->TextureImportData->BaseNormalTextureName);
-	InitialSelect = InitialSelect == INDEX_NONE ? 0 : InitialSelect; // default to the empty string located at index 0
-	MaterialCategory.AddCustomRow(LOCTEXT("BaseNormalTextureProperty", "Base Normal Texture Property")).NameContent()
-	[
-		SNew(STextBlock)
-		.Text(LOCTEXT("BaseNormalTextureProperty", "Base Normal Texture Property"))
-		.Font(IDetailLayoutBuilder::GetDetailFont())
-	]
-	.ValueContent()
-	.MaxDesiredWidth(Row.ValueWidget.MaxWidth)
-	[
-		SNew(SHorizontalBox)
-		+ SHorizontalBox::Slot()
-		.AutoWidth()
+		.ValueContent()
+		.MaxDesiredWidth(Row.ValueWidget.MaxWidth)
 		[
-			SNew(SBox)
-			.MinDesiredWidth(MinDesiredWidth)
+			SNew(SHorizontalBox)
+			+ SHorizontalBox::Slot()
+			.AutoWidth()
 			[
-				SNew(STextComboBox)
-				.OptionsSource(&BaseTextureNames)
-				.OnSelectionChanged(this, &FFbxImportUIDetails::OnNormalTextureColor)
-				.InitiallySelectedItem(BaseTextureNames[InitialSelect])
+				SNew(SBox)
+				.MinDesiredWidth(MinDesiredWidth)
+				[
+					SNew(STextComboBox)
+					.OptionsSource(&BaseTextureNames)
+					.OnSelectionChanged(this, &FFbxImportUIDetails::OnNormalTextureColor)
+					.InitiallySelectedItem(BaseTextureNames[InitialSelect])
+				]
 			]
+		];
+	}
+
+	if(BaseColorNames.Num() > 1)
+	{
+		// base emissive color properties, only used when there is no texture in the emissive map
+		InitialSelect = FindString(BaseColorNames, ImportUI->TextureImportData->BaseEmissiveColorName);
+		InitialSelect = InitialSelect == INDEX_NONE ? 0 : InitialSelect; // default to the empty string located at index 0
+		MaterialCategory.AddCustomRow(LOCTEXT("BaseEmissiveColorProperty", "Base Emissive Color Property"))
+		.NameContent()
+		[
+			SNew(STextBlock)
+			.Text(LOCTEXT("BaseEmissiveColorProperty", "Base Emissive Color Property"))
+			.Font(IDetailLayoutBuilder::GetDetailFont())
 		]
-	];
-
-	// base emissive color properties, only used when there is no texture in the emissive map
-	InitialSelect = FindString(BaseColorNames, ImportUI->TextureImportData->BaseEmissiveColorName);
-	InitialSelect = InitialSelect == INDEX_NONE ? 0 : InitialSelect; // default to the empty string located at index 0
-	MaterialCategory.AddCustomRow(LOCTEXT("BaseEmissiveColorProperty", "Base Emissive Color Property"))
-	.NameContent()
-	[
-		SNew(STextBlock)
-		.Text(LOCTEXT("BaseEmissiveColorProperty", "Base Emissive Color Property"))
-		.Font(IDetailLayoutBuilder::GetDetailFont())
-	]
-	.ValueContent()
-	.MaxDesiredWidth(Row.ValueWidget.MaxWidth)
-	[
-		SNew(SHorizontalBox)
-		+ SHorizontalBox::Slot()
-		.AutoWidth()
+		.ValueContent()
+		.MaxDesiredWidth(Row.ValueWidget.MaxWidth)
 		[
-			SNew(SBox)
-			.MinDesiredWidth(MinDesiredWidth)
+			SNew(SHorizontalBox)
+			+ SHorizontalBox::Slot()
+			.AutoWidth()
 			[
-				SNew(STextComboBox)
-				.OptionsSource(&BaseColorNames)
-				.ToolTip(SNew(SToolTip).Text(LOCTEXT("BaseEmissiveColorFBXImportToolTip", "When there is no emissive texture in the imported material this emissive color property will be used to fill a contant color value instead.")))
-				.OnSelectionChanged(this, &FFbxImportUIDetails::OnEmissiveColor)
-				.InitiallySelectedItem(BaseColorNames[InitialSelect])
+				SNew(SBox)
+				.MinDesiredWidth(MinDesiredWidth)
+				[
+					SNew(STextComboBox)
+					.OptionsSource(&BaseColorNames)
+					.ToolTip(SNew(SToolTip).Text(LOCTEXT("BaseEmissiveColorFBXImportToolTip", "When there is no emissive texture in the imported material this emissive color property will be used to fill a contant color value instead.")))
+					.OnSelectionChanged(this, &FFbxImportUIDetails::OnEmissiveColor)
+					.InitiallySelectedItem(BaseColorNames[InitialSelect])
+				]
 			]
+		];
+	}
+
+	if (BaseTextureNames.Num() > 1)
+	{
+		// base emmisive properties
+		InitialSelect = FindString(BaseTextureNames, ImportUI->TextureImportData->BaseEmmisiveTextureName);
+		InitialSelect = InitialSelect == INDEX_NONE ? 0 : InitialSelect; // default to the empty string located at index 0
+		MaterialCategory.AddCustomRow(LOCTEXT("BaseEmmisiveTextureProperty", "Base Emmisive Texture Property")).NameContent()
+		[
+			SNew(STextBlock)
+			.Text(LOCTEXT("BaseEmmisiveTextureProperty", "Base Emmisive Texture Property"))
+			.Font(IDetailLayoutBuilder::GetDetailFont())
 		]
-	];
-
-	// base emmisive properties
-	InitialSelect = FindString(BaseTextureNames, ImportUI->TextureImportData->BaseEmmisiveTextureName);
-	InitialSelect = InitialSelect == INDEX_NONE ? 0 : InitialSelect; // default to the empty string located at index 0
-	MaterialCategory.AddCustomRow(LOCTEXT("BaseEmmisiveTextureProperty", "Base Emmisive Texture Property")).NameContent()
-	[
-		SNew(STextBlock)
-		.Text(LOCTEXT("BaseEmmisiveTextureProperty", "Base Emmisive Texture Property"))
-		.Font(IDetailLayoutBuilder::GetDetailFont())
-	]
-	.ValueContent()
-	.MaxDesiredWidth(Row.ValueWidget.MaxWidth)
-	[
-		SNew(SHorizontalBox)
-		+ SHorizontalBox::Slot()
-		.AutoWidth()
+		.ValueContent()
+		.MaxDesiredWidth(Row.ValueWidget.MaxWidth)
 		[
-			SNew(SBox)
-			.MinDesiredWidth(MinDesiredWidth)
+			SNew(SHorizontalBox)
+			+ SHorizontalBox::Slot()
+			.AutoWidth()
 			[
-				SNew(STextComboBox)
-				.OptionsSource(&BaseTextureNames)
-				.OnSelectionChanged(this, &FFbxImportUIDetails::OnEmmisiveTextureColor)
-				.InitiallySelectedItem(BaseTextureNames[InitialSelect])
+				SNew(SBox)
+				.MinDesiredWidth(MinDesiredWidth)
+				[
+					SNew(STextComboBox)
+					.OptionsSource(&BaseTextureNames)
+					.OnSelectionChanged(this, &FFbxImportUIDetails::OnEmmisiveTextureColor)
+					.InitiallySelectedItem(BaseTextureNames[InitialSelect])
+				]
 			]
+		];
+
+		// base specular properties
+		InitialSelect = FindString(BaseTextureNames, ImportUI->TextureImportData->BaseSpecularTextureName);
+		InitialSelect = InitialSelect == INDEX_NONE ? 0 : InitialSelect; // default to the empty string located at index 0
+		MaterialCategory.AddCustomRow(LOCTEXT("BaseSpecularTextureProperty", "Base Specular Texture Property")).NameContent()
+		[
+			SNew(STextBlock)
+			.Text(LOCTEXT("BaseSpecularTextureProperty", "Base Specular Texture Property"))
+			.Font(IDetailLayoutBuilder::GetDetailFont())
 		]
-	];
-
-	// base specular properties
-	InitialSelect = FindString(BaseTextureNames, ImportUI->TextureImportData->BaseSpecularTextureName);
-	InitialSelect = InitialSelect == INDEX_NONE ? 0 : InitialSelect; // default to the empty string located at index 0
-	MaterialCategory.AddCustomRow(LOCTEXT("BaseSpecularTextureProperty", "Base Specular Texture Property")).NameContent()
-	[
-		SNew(STextBlock)
-		.Text(LOCTEXT("BaseSpecularTextureProperty", "Base Specular Texture Property"))
-		.Font(IDetailLayoutBuilder::GetDetailFont())
-	]
-	.ValueContent()
-	.MaxDesiredWidth(Row.ValueWidget.MaxWidth)
-	[
-		SNew(SHorizontalBox)
-		+ SHorizontalBox::Slot()
-		.AutoWidth()
+		.ValueContent()
+		.MaxDesiredWidth(Row.ValueWidget.MaxWidth)
 		[
-			SNew(SBox)
-			.MinDesiredWidth(MinDesiredWidth)
+			SNew(SHorizontalBox)
+			+ SHorizontalBox::Slot()
+			.AutoWidth()
 			[
-				SNew(STextComboBox)
-				.OptionsSource(&BaseTextureNames)
-				.OnSelectionChanged(this, &FFbxImportUIDetails::OnSpecularTextureColor)
-				.InitiallySelectedItem(BaseTextureNames[InitialSelect])
+				SNew(SBox)
+				.MinDesiredWidth(MinDesiredWidth)
+				[
+					SNew(STextComboBox)
+					.OptionsSource(&BaseTextureNames)
+					.OnSelectionChanged(this, &FFbxImportUIDetails::OnSpecularTextureColor)
+					.InitiallySelectedItem(BaseTextureNames[InitialSelect])
+				]
 			]
-		]
-	];
+		];
+	}
 }
 
 void FFbxImportUIDetails::SetStaticMeshLODGroupWidget(IDetailPropertyRow& PropertyRow, const TSharedPtr<IPropertyHandle>& Handle)
