Fix a misspell word in the fbx importer options dialog.

diff --git a/Engine/Source/Editor/DetailCustomizations/Private/FbxImportUIDetails.cpp b/Engine/Source/Editor/DetailCustomizations/Private/FbxImportUIDetails.cpp
--- a/Engine/Source/Editor/DetailCustomizations/Private/FbxImportUIDetails.cpp
+++ b/Engine/Source/Editor/DetailCustomizations/Private/FbxImportUIDetails.cpp
@@ -779,10 +779,10 @@ void FFbxImportUIDetails::ConstructBaseMaterialUI(TSharedPtr<IPropertyHandle> Ha
 		// base emmisive properties
 		InitialSelect = FindString(BaseTextureNames, ImportUI->TextureImportData->BaseEmmisiveTextureName);
 		InitialSelect = InitialSelect == INDEX_NONE ? 0 : InitialSelect; // default to the empty string located at index 0
-		MaterialCategory.AddCustomRow(LOCTEXT("BaseEmmisiveTextureProperty", "Base Emmisive Texture Property")).NameContent()
+		MaterialCategory.AddCustomRow(LOCTEXT("BaseEmissiveTextureProperty", "Base Emissive Texture Property")).NameContent()
 		[
 			SNew(STextBlock)
-			.Text(LOCTEXT("BaseEmmisiveTextureProperty", "Base Emmisive Texture Property"))
+			.Text(LOCTEXT("BaseEmissiveTextureProperty", "Base Emissive Texture Property"))
 			.Font(IDetailLayoutBuilder::GetDetailFont())
 		]
 		.ValueContent()
