PR #5415: Bug Fix: Use TWeakObjectPtr to prevent dangling UObject pointers in FLightComponentDetails (Contributed by [REDACTED_PERSON])

diff --git a/Engine/Source/Editor/DetailCustomizations/Private/LightComponentDetails.cpp b/Engine/Source/Editor/DetailCustomizations/Private/LightComponentDetails.cpp
--- a/Engine/Source/Editor/DetailCustomizations/Private/LightComponentDetails.cpp
+++ b/Engine/Source/Editor/DetailCustomizations/Private/LightComponentDetails.cpp
@@ -22,7 +22,7 @@ void FLightComponentDetails::CustomizeDetails( IDetailLayoutBuilder& DetailBuild
 {
 	TArray<TWeakObjectPtr<UObject>> Objects;
 	DetailBuilder.GetObjectsBeingCustomized(Objects);
-	ULightComponent* Component = Cast<ULightComponent>(Objects[0].Get());
+	TWeakObjectPtr<ULightComponent> Component(Cast<ULightComponent>(Objects[0].Get()));
 
 	// Mobility property is on the scene component base class not the light component and that is why we have to use USceneComponent::StaticClass
 	TSharedRef<IPropertyHandle> MobilityHandle = DetailBuilder.GetProperty(GET_MEMBER_NAME_CHECKED(ULightComponent, Mobility), USceneComponent::StaticClass());
@@ -112,12 +112,12 @@ void FLightComponentDetails::SetComponentIntensity(ULightComponent* Component, f
 	}
 }
 
-void FLightComponentDetails::ResetIntensityToDefault(TSharedPtr<IPropertyHandle> PropertyHandle, ULightComponent* Component)
+void FLightComponentDetails::ResetIntensityToDefault(TSharedPtr<IPropertyHandle> PropertyHandle, TWeakObjectPtr<ULightComponent> Component)
 {
-	ULightComponent* ArchetypeComponent = Component ? Cast<ULocalLightComponent>(Component->GetArchetype()) : nullptr;
+	ULightComponent* ArchetypeComponent = Component.IsValid() ? Cast<ULocalLightComponent>(Component->GetArchetype()) : nullptr;
 	if (ArchetypeComponent)
 	{
-		SetComponentIntensity(Component, ArchetypeComponent->ComputeLightBrightness());
+		SetComponentIntensity(Component.Get(), ArchetypeComponent->ComputeLightBrightness());
 	}
 	else
 	{
@@ -126,9 +126,9 @@ void FLightComponentDetails::ResetIntensityToDefault(TSharedPtr<IPropertyHandle>
 	}
 }
 
-bool FLightComponentDetails::IsIntensityResetToDefaultVisible(TSharedPtr<IPropertyHandle> PropertyHandle, ULightComponent* Component) const
+bool FLightComponentDetails::IsIntensityResetToDefaultVisible(TSharedPtr<IPropertyHandle> PropertyHandle, TWeakObjectPtr<ULightComponent> Component) const
 {
-	ULightComponent* ArchetypeComponent = Component ? Cast<ULocalLightComponent>(Component->GetArchetype()) : nullptr;
+	ULightComponent* ArchetypeComponent = Component.IsValid() ? Cast<ULocalLightComponent>(Component->GetArchetype()) : nullptr;
 	if (ArchetypeComponent)
 	{
 		return !FMath::IsNearlyEqual(Component->ComputeLightBrightness(), ArchetypeComponent->ComputeLightBrightness());
diff --git a/Engine/Source/Editor/DetailCustomizations/Private/LightComponentDetails.h b/Engine/Source/Editor/DetailCustomizations/Private/LightComponentDetails.h
--- a/Engine/Source/Editor/DetailCustomizations/Private/LightComponentDetails.h
+++ b/Engine/Source/Editor/DetailCustomizations/Private/LightComponentDetails.h
@@ -27,8 +27,8 @@ private:
 	bool IsUseIESBrightnessEnabled() const;
 	bool IsIESBrightnessScaleEnabled() const;
 
-	void ResetIntensityToDefault(TSharedPtr<IPropertyHandle> PropertyHandle, ULightComponent* Component);
-	bool IsIntensityResetToDefaultVisible(TSharedPtr<IPropertyHandle> PropertyHandle, ULightComponent* Component) const;
+	void ResetIntensityToDefault(TSharedPtr<IPropertyHandle> PropertyHandle, TWeakObjectPtr<ULightComponent> Component);
+	bool IsIntensityResetToDefaultVisible(TSharedPtr<IPropertyHandle> PropertyHandle, TWeakObjectPtr<ULightComponent> Component) const;
 
 private:
 	TSharedPtr<IPropertyHandle> IESBrightnessTextureProperty;
