In fbx import option, add a "Do Not Override" string when base material is override.

diff --git a/Engine/Source/Editor/DetailCustomizations/Private/FbxImportUIDetails.cpp b/Engine/Source/Editor/DetailCustomizations/Private/FbxImportUIDetails.cpp
--- a/Engine/Source/Editor/DetailCustomizations/Private/FbxImportUIDetails.cpp
+++ b/Engine/Source/Editor/DetailCustomizations/Private/FbxImportUIDetails.cpp
@@ -31,6 +31,8 @@
 #define MinimumLodNumberID 0
 #define LodNumberID 1
 
+static FString DoNotOverrideString = FText(LOCTEXT("BaseColorPropertyDoNotOverride", "Do Not Override")).ToString();
+
 //If the String is contain in the StringArray, it return the index. Otherwise return INDEX_NONE
 int FindString(const TArray<TSharedPtr<FString>> &StringArray, const FString &String) {
 	for (int i = 0; i < StringArray.Num(); i++)
@@ -626,8 +628,8 @@ void FFbxImportUIDetails::ConstructBaseMaterialUI(TSharedPtr<IPropertyHandle> Ha
 
 	BaseColorNames.Empty();
 	BaseTextureNames.Empty();
-	BaseColorNames.Add(MakeShareable(new FString()));
-	BaseTextureNames.Add(MakeShareable(new FString()));
+	BaseColorNames.Add(MakeShareable(new FString(DoNotOverrideString)));
+	BaseTextureNames.Add(MakeShareable(new FString(DoNotOverrideString)));
 	TArray<FMaterialParameterInfo> OutParameterInfo;
 	TArray<FGuid> Guids;
 	float MinDesiredWidth = 150.0f;
@@ -830,6 +832,21 @@ void FFbxImportUIDetails::ConstructBaseMaterialUI(TSharedPtr<IPropertyHandle> Ha
 			]
 		];
 	}
+	if (BaseTextureNames.Num() > 1 || BaseColorNames.Num() > 1)
+	{
+		MaterialCategory.AddCustomRow(LOCTEXT("BaseParamPropertyClearAll", "Clear All Properties"))
+		.ValueContent()
+		[
+			SNew(SButton)
+			.OnClicked(this, &FFbxImportUIDetails::MaterialBaseParamClearAllProperties)
+			.Content()
+			[
+				SNew(STextBlock)
+				.Text(LOCTEXT("BaseParamPropertyClearAll", "Clear All Properties"))
+				.Font(IDetailLayoutBuilder::GetDetailFont())
+			]
+		];
+	}
 }
 
 void FFbxImportUIDetails::SetStaticMeshLODGroupWidget(IDetailPropertyRow& PropertyRow, const TSharedPtr<IPropertyHandle>& Handle)
@@ -999,23 +1016,56 @@ void FFbxImportUIDetails::BaseMaterialChanged()
 	RefreshCustomDetail();
 }
 
-void FFbxImportUIDetails::OnBaseColor(TSharedPtr<FString> Selection, ESelectInfo::Type SelectInfo) {
-	ImportUI->TextureImportData->BaseColorName = *Selection.Get();
+void GetSelectionParameterString(TSharedPtr<FString> Selection, FString &OutParameterName)
+{
+	OutParameterName = *Selection.Get();
+	if (OutParameterName.Equals(DoNotOverrideString))
+	{
+		OutParameterName.Empty();
+	}
+}
+
+void FFbxImportUIDetails::OnBaseColor(TSharedPtr<FString> Selection, ESelectInfo::Type SelectInfo)
+{
+	GetSelectionParameterString(Selection, ImportUI->TextureImportData->BaseColorName);
+}
+
+void FFbxImportUIDetails::OnDiffuseTextureColor(TSharedPtr<FString> Selection, ESelectInfo::Type SelectInfo)
+{
+	GetSelectionParameterString(Selection, ImportUI->TextureImportData->BaseDiffuseTextureName);
 }
-void FFbxImportUIDetails::OnDiffuseTextureColor(TSharedPtr<FString> Selection, ESelectInfo::Type SelectInfo) {
-	ImportUI->TextureImportData->BaseDiffuseTextureName = *Selection.Get();
+
+void FFbxImportUIDetails::OnNormalTextureColor(TSharedPtr<FString> Selection, ESelectInfo::Type SelectInfo)
+{
+	GetSelectionParameterString(Selection, ImportUI->TextureImportData->BaseNormalTextureName);
 }
-void FFbxImportUIDetails::OnNormalTextureColor(TSharedPtr<FString> Selection, ESelectInfo::Type SelectInfo) {
-	ImportUI->TextureImportData->BaseNormalTextureName = *Selection.Get();
+
+void FFbxImportUIDetails::OnEmmisiveTextureColor(TSharedPtr<FString> Selection, ESelectInfo::Type SelectInfo)
+{
+	GetSelectionParameterString(Selection, ImportUI->TextureImportData->BaseEmmisiveTextureName);
 }
-void FFbxImportUIDetails::OnEmmisiveTextureColor(TSharedPtr<FString> Selection, ESelectInfo::Type SelectInfo) {
-	ImportUI->TextureImportData->BaseEmmisiveTextureName = *Selection.Get();
+
+void FFbxImportUIDetails::OnEmissiveColor(TSharedPtr<FString> Selection, ESelectInfo::Type SelectInfo)
+{
+	GetSelectionParameterString(Selection, ImportUI->TextureImportData->BaseEmissiveColorName);
 }
-void FFbxImportUIDetails::OnEmissiveColor(TSharedPtr<FString> Selection, ESelectInfo::Type SelectInfo) {
-	ImportUI->TextureImportData->BaseEmissiveColorName = *Selection.Get();
+
+void FFbxImportUIDetails::OnSpecularTextureColor(TSharedPtr<FString> Selection, ESelectInfo::Type SelectInfo)
+{
+	GetSelectionParameterString(Selection, ImportUI->TextureImportData->BaseSpecularTextureName);
 }
-void FFbxImportUIDetails::OnSpecularTextureColor(TSharedPtr<FString> Selection, ESelectInfo::Type SelectInfo) {
-	ImportUI->TextureImportData->BaseSpecularTextureName = *Selection.Get();
+
+FReply FFbxImportUIDetails::MaterialBaseParamClearAllProperties()
+{
+	ImportUI->TextureImportData->BaseColorName.Empty();
+	ImportUI->TextureImportData->BaseDiffuseTextureName.Empty();
+	ImportUI->TextureImportData->BaseNormalTextureName.Empty();
+	ImportUI->TextureImportData->BaseEmmisiveTextureName.Empty();
+	ImportUI->TextureImportData->BaseEmissiveColorName.Empty();
+	ImportUI->TextureImportData->BaseSpecularTextureName.Empty();
+	//Need to refresh the custom detail since we do not have any pointer on the combo box
+	RefreshCustomDetail();
+	return FReply::Handled();
 }
 
 FReply FFbxImportUIDetails::ShowConflictDialog(ConflictDialogType DialogType)
diff --git a/Engine/Source/Editor/DetailCustomizations/Private/FbxImportUIDetails.h b/Engine/Source/Editor/DetailCustomizations/Private/FbxImportUIDetails.h
--- a/Engine/Source/Editor/DetailCustomizations/Private/FbxImportUIDetails.h
+++ b/Engine/Source/Editor/DetailCustomizations/Private/FbxImportUIDetails.h
@@ -76,6 +76,7 @@ public:
 	void OnEmmisiveTextureColor(TSharedPtr<FString> Selection, ESelectInfo::Type SelectInfo);
 	void OnEmissiveColor(TSharedPtr<FString> Selection, ESelectInfo::Type SelectInfo);
 	void OnSpecularTextureColor(TSharedPtr<FString> Selection, ESelectInfo::Type SelectInfo);
+	FReply MaterialBaseParamClearAllProperties();
 	
 	TWeakObjectPtr<UFbxImportUI> ImportUI;		// The UI data object being customised
 	IDetailLayoutBuilder* CachedDetailBuilder;	// The detail builder for this cusomtomisation
