Polish fbx material import options

diff --git a/Engine/Source/Editor/DetailCustomizations/Private/FbxImportUIDetails.cpp b/Engine/Source/Editor/DetailCustomizations/Private/FbxImportUIDetails.cpp
--- a/Engine/Source/Editor/DetailCustomizations/Private/FbxImportUIDetails.cpp
+++ b/Engine/Source/Editor/DetailCustomizations/Private/FbxImportUIDetails.cpp
@@ -33,6 +33,10 @@
 
 static FString DoNotOverrideString = FText(LOCTEXT("BaseColorPropertyDoNotOverride", "Do Not Override")).ToString();
 
+static FString CreateNewMaterialsString = FText(LOCTEXT("MaterialImportMethodCreateNewMaterials", "Create New Materials")).ToString();
+static FString CreateNewInstancedMaterialsString = FText(LOCTEXT("MaterialImportMethodCreateNewInstancedMaterials", "Create New Instanced Materials")).ToString();
+static FString DoNotCreateMaterialString = FText(LOCTEXT("MaterialImportMethodDoNotCreateMaterial", "Do Not Create Material")).ToString();
+
 //If the String is contain in the StringArray, it return the index. Otherwise return INDEX_NONE
 int FindString(const TArray<TSharedPtr<FString>> &StringArray, const FString &String) {
 	for (int i = 0; i < StringArray.Num(); i++)
@@ -511,25 +515,52 @@ void FFbxImportUIDetails::CustomizeDetails( IDetailLayoutBuilder& DetailBuilder
 	}
 	else
 	{
+		TSharedRef<IPropertyHandle> ImportMaterialPropHandle = DetailBuilder.GetProperty(GET_MEMBER_NAME_CHECKED(UFbxImportUI, bImportMaterials));
+
 		TSharedRef<IPropertyHandle> TextureDataProp = DetailBuilder.GetProperty(GET_MEMBER_NAME_CHECKED(UFbxImportUI, TextureImportData));
 		DetailBuilder.HideProperty(TextureDataProp);
 
 		ExtraProperties.Empty();
 		CollectChildPropertiesRecursive(TextureDataProp, ExtraProperties);
+		
+
+		TSharedPtr<IPropertyHandle> MaterialLocationPropHandle;
+		for (TSharedPtr<IPropertyHandle> Handle : ExtraProperties)
+		{
+			// We ignore base import data for this window.
+			if (Handle->GetProperty()->GetOuter() == UFbxTextureImportData::StaticClass())
+			{
+				if (Handle->GetProperty()->GetFName() == GET_MEMBER_NAME_CHECKED(UFbxTextureImportData, MaterialSearchLocation))
+				{
+					MaterialLocationPropHandle = Handle;
+				}
+			}
+		}
+
+		//The order is
+		// Search Location
+		// Import Materials
+		// [Base Material Name]
+		// [All Base Material Parameter]
+		// 
+		DetailBuilder.HideProperty(MaterialLocationPropHandle);
+		MaterialCategory.AddProperty(MaterialLocationPropHandle);
+		DetailBuilder.HideProperty(ImportMaterialPropHandle);
+		ConstructMaterialImportMethod(ImportMaterialPropHandle, MaterialCategory);
 
 		for(TSharedPtr<IPropertyHandle> Handle : ExtraProperties)
 		{
 			// We ignore base import data for this window.
 			if(Handle->GetProperty()->GetOuter() == UFbxTextureImportData::StaticClass())
 			{
-				if (Handle->GetPropertyDisplayName().ToString() == FString(TEXT("Base Material Name")))
+				if (Handle->GetProperty()->GetFName() == GET_MEMBER_NAME_CHECKED(UFbxTextureImportData, BaseMaterialName))
 				{
-					if (ImportUI->bImportMaterials)
+					if (ImportUI->bImportMaterials && ImportUI->TextureImportData->bUseBaseMaterial)
 					{
 						ConstructBaseMaterialUI(Handle, MaterialCategory);
 					}
 				}
-				else
+				else if (Handle != MaterialLocationPropHandle)
 				{
 					MaterialCategory.AddProperty(Handle);
 				}
@@ -610,6 +641,49 @@ void FFbxImportUIDetails::AddSubCategory(IDetailLayoutBuilder& DetailBuilder, FN
 	}
 }
 
+void FFbxImportUIDetails::ConstructMaterialImportMethod(TSharedPtr<IPropertyHandle> ImportMaterialPropHandle, IDetailCategoryBuilder& MaterialCategory)
+{
+	//The import material is represent by a combobox with 3 choices
+	//1. Create New Materials
+	//2. Create New Instanced Materials (Using an existing base material)
+	//3. Do not Create Materials
+	ImportMethodNames.Reset();
+	ImportMethodNames.Add(MakeShareable(new FString(CreateNewMaterialsString)));
+	ImportMethodNames.Add(MakeShareable(new FString(CreateNewInstancedMaterialsString)));
+	ImportMethodNames.Add(MakeShareable(new FString(DoNotCreateMaterialString)));
+
+	if(ImportUI->TextureImportData->BaseMaterialName.IsValid())
+	{
+		//When we load the UI the first time we set this boolean to true in case the BaseMaterialName is valid.
+		ImportUI->TextureImportData->bUseBaseMaterial = true;
+	}
+
+	int32 InitialSelect = ImportUI->bImportMaterials ? (ImportUI->TextureImportData->bUseBaseMaterial ? 1 : 0) : 2;
+
+	MaterialCategory.AddCustomRow(LOCTEXT("MaterialImportMethod", "Material Import Method"))
+	.NameContent()
+	[
+		SNew(STextBlock)
+		.Text(LOCTEXT("MaterialImportMethod", "Material Import Method"))
+		.ToolTipText(LOCTEXT("MaterialImportMethodToolTip", "How materials are created when the importer cannot found it using the search location."))
+		.Font(IDetailLayoutBuilder::GetDetailFont())
+	]
+	.ValueContent()
+	[
+		SNew(SHorizontalBox)
+		+ SHorizontalBox::Slot()
+		.AutoWidth()
+		[
+			SNew(SBox)
+			[
+				SNew(STextComboBox)
+				.OptionsSource(&ImportMethodNames)
+				.OnSelectionChanged(this, &FFbxImportUIDetails::OnMaterialImportMethodChanged)
+				.InitiallySelectedItem(ImportMethodNames[InitialSelect])
+			]
+		]
+	];
+}
 
 void FFbxImportUIDetails::ConstructBaseMaterialUI(TSharedPtr<IPropertyHandle> Handle, IDetailCategoryBuilder& MaterialCategory)
 {
@@ -1068,6 +1142,31 @@ FReply FFbxImportUIDetails::MaterialBaseParamClearAllProperties()
 	return FReply::Handled();
 }
 
+void FFbxImportUIDetails::OnMaterialImportMethodChanged(TSharedPtr<FString> Selection, ESelectInfo::Type SelectInfo)
+{
+	FString SelectName = *Selection.Get();
+	if (SelectName.Equals(CreateNewMaterialsString))
+	{
+		ImportUI->bImportMaterials = true;
+		//Reset the base material and the UseBaseMaterial flag to hide the base material name property
+		ImportUI->TextureImportData->bUseBaseMaterial = false;
+		ImportUI->TextureImportData->BaseMaterialName.Reset();
+	}
+	else if (SelectName.Equals(CreateNewInstancedMaterialsString))
+	{
+		ImportUI->bImportMaterials = true;
+		ImportUI->TextureImportData->bUseBaseMaterial = true;
+	}
+	else
+	{
+		ImportUI->bImportMaterials = false;
+		//Reset the base material and the UseBaseMaterial flag to hide the base material name property
+		ImportUI->TextureImportData->bUseBaseMaterial = false;
+		ImportUI->TextureImportData->BaseMaterialName.Reset();
+	}
+	RefreshCustomDetail();
+}
+
 FReply FFbxImportUIDetails::ShowConflictDialog(ConflictDialogType DialogType)
 {
 	if (DialogType == Conflict_Material)
diff --git a/Engine/Source/Editor/DetailCustomizations/Private/FbxImportUIDetails.h b/Engine/Source/Editor/DetailCustomizations/Private/FbxImportUIDetails.h
--- a/Engine/Source/Editor/DetailCustomizations/Private/FbxImportUIDetails.h
+++ b/Engine/Source/Editor/DetailCustomizations/Private/FbxImportUIDetails.h
@@ -43,6 +43,8 @@ public:
 
 	void CollectChildPropertiesRecursive(TSharedPtr<IPropertyHandle> Node, TArray<TSharedPtr<IPropertyHandle>>& OutProperties);
 	
+	void ConstructMaterialImportMethod(TSharedPtr<IPropertyHandle> ImportMaterialPropHandle, class IDetailCategoryBuilder& MaterialCategory);
+
 	void ConstructBaseMaterialUI(TSharedPtr<IPropertyHandle> Handle, class IDetailCategoryBuilder& MaterialCategory);
 
 	/** Checks whether a metadata string is valid for a given import type 
@@ -77,6 +79,8 @@ public:
 	void OnEmissiveColor(TSharedPtr<FString> Selection, ESelectInfo::Type SelectInfo);
 	void OnSpecularTextureColor(TSharedPtr<FString> Selection, ESelectInfo::Type SelectInfo);
 	FReply MaterialBaseParamClearAllProperties();
+
+	void OnMaterialImportMethodChanged(TSharedPtr<FString> Selection, ESelectInfo::Type SelectInfo);
 	
 	TWeakObjectPtr<UFbxImportUI> ImportUI;		// The UI data object being customised
 	IDetailLayoutBuilder* CachedDetailBuilder;	// The detail builder for this cusomtomisation
@@ -110,6 +114,10 @@ private:
 	TSharedPtr<IPropertyHandle> VertexColorImportOptionHandle;
 	TSharedPtr<IPropertyHandle> SkeletalMeshVertexColorImportOptionHandle;
 
+	TArray< TSharedPtr< FString > > ImportMethodNames;
+	TSharedPtr<class STextComboBox> MaterialImportMethodComboBox;
+	bool bShowBaseMaterialUI;
+
 	TArray< TSharedPtr< FString > > BaseColorNames;
 	TArray< TSharedPtr< FString > > BaseTextureNames;
 };
diff --git a/Engine/Source/Editor/UnrealEd/Classes/Factories/FbxTextureImportData.h b/Engine/Source/Editor/UnrealEd/Classes/Factories/FbxTextureImportData.h
--- a/Engine/Source/Editor/UnrealEd/Classes/Factories/FbxTextureImportData.h
+++ b/Engine/Source/Editor/UnrealEd/Classes/Factories/FbxTextureImportData.h
@@ -29,6 +29,9 @@ class UNREALED_API UFbxTextureImportData : public UFbxAssetImportData
 	UPROPERTY(EditAnywhere, config, Category = Material, meta = (ImportType = "Mesh", AllowedClasses = "MaterialInterface"))
 	FSoftObjectPath BaseMaterialName;
 
+	/** transient, ImportUI customize helper to store if we must show or not the BaseMaterialName property. */
+	bool bUseBaseMaterial;
+
 	UPROPERTY(config, meta = (ImportType = "Mesh"))
 	FString BaseColorName;
 
