Make sure the font use in the import material method combobox is the good one

diff --git a/Engine/Source/Editor/DetailCustomizations/Private/FbxImportUIDetails.cpp b/Engine/Source/Editor/DetailCustomizations/Private/FbxImportUIDetails.cpp
--- a/Engine/Source/Editor/DetailCustomizations/Private/FbxImportUIDetails.cpp
+++ b/Engine/Source/Editor/DetailCustomizations/Private/FbxImportUIDetails.cpp
@@ -680,6 +680,7 @@ void FFbxImportUIDetails::ConstructMaterialImportMethod(TSharedPtr<IPropertyHand
 				.OptionsSource(&ImportMethodNames)
 				.OnSelectionChanged(this, &FFbxImportUIDetails::OnMaterialImportMethodChanged)
 				.InitiallySelectedItem(ImportMethodNames[InitialSelect])
+				.Font(IDetailLayoutBuilder::GetDetailFont())
 			]
 		]
 	];
