Fixed Object Type (collision presets) not updating immediately after toggling 'Simulate Physics'

diff --git a/Engine/Source/Editor/DetailCustomizations/Private/BodyInstanceCustomization.cpp b/Engine/Source/Editor/DetailCustomizations/Private/BodyInstanceCustomization.cpp
--- a/Engine/Source/Editor/DetailCustomizations/Private/BodyInstanceCustomization.cpp
+++ b/Engine/Source/Editor/DetailCustomizations/Private/BodyInstanceCustomization.cpp
@@ -237,6 +237,12 @@ void FBodyInstanceCustomization::CustomizeChildren( TSharedRef<class IPropertyHa
 		}
 	}
 	
+	TSharedPtr<IPropertyHandle> SimulatePhysicsPropertyHandle = BodyInstanceHandle->GetChildHandle(GET_MEMBER_NAME_CHECKED(FBodyInstance, bSimulatePhysics));
+	if (SimulatePhysicsPropertyHandle.IsValid())
+	{
+		FSimpleDelegate OnSimulatePhysicsChangedDelegate = FSimpleDelegate::CreateSP(this, &FBodyInstanceCustomization::OnSimulatePhysicsChanged);
+		SimulatePhysicsPropertyHandle->SetOnPropertyValueChanged(OnSimulatePhysicsChangedDelegate);
+	}
 
 	AddCollisionCategory(StructPropertyHandle, StructBuilder, StructCustomizationUtils);
 }
@@ -783,6 +789,10 @@ TSharedRef<SWidget> FBodyInstanceCustomization::MakeCollisionProfileComboWidget(
 		.Font(IDetailLayoutBuilder::GetDetailFont());
 }
 
+void FBodyInstanceCustomization::OnSimulatePhysicsChanged()
+{
+	UpdateCollisionProfile();
+}
 
 /////////////////////////////////////////////////////////////////////////////////////////////////////////////
 // NOTE! I have a lot of ensure to make sure it's set correctly
diff --git a/Engine/Source/Editor/DetailCustomizations/Private/BodyInstanceCustomization.h b/Engine/Source/Editor/DetailCustomizations/Private/BodyInstanceCustomization.h
--- a/Engine/Source/Editor/DetailCustomizations/Private/BodyInstanceCustomization.h
+++ b/Engine/Source/Editor/DetailCustomizations/Private/BodyInstanceCustomization.h
@@ -43,6 +43,10 @@ public:
 	virtual void CustomizeChildren( TSharedRef<class IPropertyHandle> StructPropertyHandle, class IDetailChildrenBuilder& StructBuilder, IPropertyTypeCustomizationUtils& StructCustomizationUtils ) override;
 
 private:
+
+	// Simulate physics toggle
+	void OnSimulatePhysicsChanged();
+
 	// Profile combo related
 	TSharedRef<SWidget> MakeCollisionProfileComboWidget( TSharedPtr<FString> InItem );
 	void OnCollisionProfileChanged( TSharedPtr<FString> NewSelection, ESelectInfo::Type SelectInfo, IDetailGroup* CollisionGroup );
