FFilePath: Adding metadata to support relative path to project directory

diff --git a/Engine/Source/Editor/DetailCustomizations/Private/FilePathStructCustomization.cpp b/Engine/Source/Editor/DetailCustomizations/Private/FilePathStructCustomization.cpp
--- a/Engine/Source/Editor/DetailCustomizations/Private/FilePathStructCustomization.cpp
+++ b/Engine/Source/Editor/DetailCustomizations/Private/FilePathStructCustomization.cpp
@@ -34,6 +34,7 @@ void FFilePathStructCustomization::CustomizeHeader( TSharedRef<IPropertyHandle>
 
 
 	bLongPackageName = StructPropertyHandle->HasMetaData(TEXT("LongPackageName"));
+	bRelativeToGameDir = StructPropertyHandle->HasMetaData(TEXT("RelativeToGameDir"));
 
 	if (MetaData.IsEmpty())
 	{
@@ -92,7 +93,20 @@ void FFilePathStructCustomization::HandleFilePathPickerPathPicked( const FString
 		}
 		FinalPath = LongPackageName;
 	}
-
+	else if (bRelativeToGameDir)
+	{
+		//If the path is part of the project directory, make it relative to it or keep the absolute path.
+		const FString AbsolutePickedPath = FPaths::ConvertRelativePathToFull(PickedPath);
+		const FString AbsoluteProjectDir = FPaths::ConvertRelativePathToFull(FPaths::ProjectDir());
+		if (AbsolutePickedPath.StartsWith(AbsoluteProjectDir))
+		{
+			FinalPath = AbsolutePickedPath.RightChop(AbsoluteProjectDir.Len());
+		}
+		else
+		{
+			FinalPath = AbsolutePickedPath;
+		}
+	}
 
 	PathStringProperty->SetValue(FinalPath);
 	FEditorDirectories::Get().SetLastDirectory(ELastDirectory::GENERIC_OPEN, FPaths::GetPath(PickedPath));
diff --git a/Engine/Source/Editor/DetailCustomizations/Private/FilePathStructCustomization.h b/Engine/Source/Editor/DetailCustomizations/Private/FilePathStructCustomization.h
--- a/Engine/Source/Editor/DetailCustomizations/Private/FilePathStructCustomization.h
+++ b/Engine/Source/Editor/DetailCustomizations/Private/FilePathStructCustomization.h
@@ -44,4 +44,7 @@ private:
 	/** Pointer to the string that will be seet when changing the path */
 	TSharedPtr<IPropertyHandle> PathStringProperty;
 	bool bLongPackageName;
+	
+	/** Whether the output file path tries to be relative to the project directory */
+	bool bRelativeToGameDir;
 };
diff --git a/Engine/Source/Runtime/CoreUObject/Public/UObject/ObjectMacros.h b/Engine/Source/Runtime/CoreUObject/Public/UObject/ObjectMacros.h
--- a/Engine/Source/Runtime/CoreUObject/Public/UObject/ObjectMacros.h
+++ b/Engine/Source/Runtime/CoreUObject/Public/UObject/ObjectMacros.h
@@ -1133,6 +1133,9 @@ namespace UM
 		/// [PropertyMetadata] Used by FFilePath properties. Indicates the path filter to display in the file picker.
 		FilePathFilter,
 
+		/// [PropertyMetadata] Used by FFilePath properties. Indicates that the FilePicker dialog will output a path relative to the game directory when setting the property. An absolute path will be used when outside the game directory.
+		RelativeToGameDir,
+
 		/// [PropertyMetadata] Deprecated.
 		FixedIncrement,
 
