Added a new type of combobox which supports search. This emulates the exact same search behaviour found it the asset type browsers - but for any combo box. We could of course lift this to a feature which would be for UE - it seems though that needs more discussion.

Due to simplicity and stability of the widgets instead of adding / removing the widgets from the view we instead collapse them. That way multi selections can be maintained and there's no lookup / mapping cost inside of the TableViewBase.

diff --git a/Engine/Source/Editor/DetailCustomizations/Private/PoseAssetDetails.cpp b/Engine/Source/Editor/DetailCustomizations/Private/PoseAssetDetails.cpp
--- a/Engine/Source/Editor/DetailCustomizations/Private/PoseAssetDetails.cpp
+++ b/Engine/Source/Editor/DetailCustomizations/Private/PoseAssetDetails.cpp
@@ -13,6 +13,7 @@
 #include "ScopedTransaction.h"
 #include "Framework/Notifications/NotificationManager.h"
 #include "Widgets/Notifications/SNotificationList.h"
+#include "Widgets/Input/SEditableTextBox.h"
 
 #define LOCTEXT_NAMESPACE	"PoseAssetDetails"
 
@@ -92,7 +93,7 @@ void FPoseAssetDetails::CustomizeDetails(IDetailLayoutBuilder& DetailBuilder)
 	]
 	.ValueContent()
 	[
-		SAssignNew(RetargetSourceComboBox, SComboBox< TSharedPtr<FString> >)
+		SAssignNew(RetargetSourceComboBox, SSearchableComboBox)
 		.OptionsSource(&RetargetSourceComboList)
 		.OnGenerateWidget(this, &FPoseAssetDetails::MakeRetargetSourceComboWidget)
 		.OnSelectionChanged(this, &FPoseAssetDetails::OnRetargetSourceChanged)
@@ -188,7 +189,7 @@ void FPoseAssetDetails::CustomizeDetails(IDetailLayoutBuilder& DetailBuilder)
 		.HAlign(HAlign_Fill)
 		.Padding(3)
 		[
-			SAssignNew(BasePoseComboBox, SComboBox< TSharedPtr<FString> >)
+			SAssignNew(BasePoseComboBox, SSearchableComboBox)
 			.OptionsSource(&BasePoseComboList)
 			.OnGenerateWidget(this, &FPoseAssetDetails::MakeBasePoseComboWidget)
 			.OnSelectionChanged(this, &FPoseAssetDetails::OnBasePoseChanged)
diff --git a/Engine/Source/Editor/DetailCustomizations/Private/PoseAssetDetails.h b/Engine/Source/Editor/DetailCustomizations/Private/PoseAssetDetails.h
--- a/Engine/Source/Editor/DetailCustomizations/Private/PoseAssetDetails.h
+++ b/Engine/Source/Editor/DetailCustomizations/Private/PoseAssetDetails.h
@@ -5,7 +5,7 @@
 #include "CoreMinimal.h"
 #include "Input/Reply.h"
 #include "Widgets/SWidget.h"
-#include "Widgets/Input/SComboBox.h"
+#include "SSearchableComboBox.h"
 #include "Animation/Skeleton.h"
 #include "IDetailCustomization.h"
 #include "Animation/PoseAsset.h"
@@ -33,7 +33,7 @@ private:
 	TSharedPtr<IPropertyHandle> RetargetSourceNameHandler;
 
 	// retarget source related
-	TSharedPtr<class SComboBox< TSharedPtr<FString> > > RetargetSourceComboBox;
+	TSharedPtr< SSearchableComboBox > RetargetSourceComboBox;
 	TArray< TSharedPtr< FString > >						RetargetSourceComboList;
 
 	TSharedRef<SWidget> MakeRetargetSourceComboWidget( TSharedPtr<FString> InItem );
@@ -53,7 +53,7 @@ private:
 	ECheckBoxState IsAdditiveChecked() const;
 
 	// base pose
-	TSharedPtr<class SComboBox< TSharedPtr<FString> > > BasePoseComboBox;
+	TSharedPtr< SSearchableComboBox > BasePoseComboBox;
 	TArray< TSharedPtr< FString > >						BasePoseComboList;
 
 	TSharedRef<SWidget> MakeBasePoseComboWidget(TSharedPtr<FString> InItem);
diff --git a/Engine/Source/Editor/DetailCustomizations/Private/SSearchableComboBox.cpp b/Engine/Source/Editor/DetailCustomizations/Private/SSearchableComboBox.cpp
new file mode 100644
--- /dev/null
+++ b/Engine/Source/Editor/DetailCustomizations/Private/SSearchableComboBox.cpp
@@ -0,0 +1,240 @@
+// Copyright 1998-2018 Epic Games, Inc. All Rights Reservekd.
+
+#include "SSearchableComboBox.h"
+#include "DetailLayoutBuilder.h"
+
+#define LOCTEXT_NAMESPACE "SearchableComboBox"
+
+void SSearchableComboBox::Construct(const FArguments& InArgs)
+{
+	check(InArgs._ComboBoxStyle);
+
+	ItemStyle = InArgs._ItemStyle;
+
+	// Work out which values we should use based on whether we were given an override, or should use the style's version
+	const FComboButtonStyle& OurComboButtonStyle = InArgs._ComboBoxStyle->ComboButtonStyle;
+	const FButtonStyle* const OurButtonStyle = InArgs._ButtonStyle ? InArgs._ButtonStyle : &OurComboButtonStyle.ButtonStyle;
+
+	this->OnComboBoxOpening = InArgs._OnComboBoxOpening;
+	this->OnSelectionChanged = InArgs._OnSelectionChanged;
+	this->OnGenerateWidget = InArgs._OnGenerateWidget;
+
+	OptionsSource = InArgs._OptionsSource;
+	CustomScrollbar = InArgs._CustomScrollbar;
+
+	TSharedRef<SWidget> ComboBoxMenuContent =
+		SNew(SBox)
+		.MaxDesiredHeight(InArgs._MaxListHeight)
+
+		[
+			SNew(SVerticalBox)
+
+			+ SVerticalBox::Slot()
+			.AutoHeight()
+			[
+				SAssignNew(this->SearchField, SEditableTextBox)
+				.Font(IDetailLayoutBuilder::GetDetailFont())
+				.HintText(LOCTEXT("Search", "Search"))
+				.OnTextChanged(this, &SSearchableComboBox::OnSearchTextChanged)
+			]
+
+			+ SVerticalBox::Slot()
+			[
+				SAssignNew(this->ComboListView, SComboListType)
+				.ListItemsSource(OptionsSource)
+				.OnGenerateRow(this, &SSearchableComboBox::GenerateMenuItemRow)
+				.OnSelectionChanged(this, &SSearchableComboBox::OnSelectionChanged_Internal)
+				.SelectionMode(ESelectionMode::Single)
+				.ExternalScrollbar(InArgs._CustomScrollbar)
+			]
+		];
+
+	// Set up content
+	TSharedPtr<SWidget> ButtonContent = InArgs._Content.Widget;
+	if (InArgs._Content.Widget == SNullWidget::NullWidget)
+	{
+		SAssignNew(ButtonContent, STextBlock)
+			.Text(NSLOCTEXT("SSearchableComboBox", "ContentWarning", "No Content Provided"))
+			.ColorAndOpacity(FLinearColor::Red);
+	}
+
+
+	SComboButton::Construct(SComboButton::FArguments()
+		.ComboButtonStyle(&OurComboButtonStyle)
+		.ButtonStyle(OurButtonStyle)
+		.Method(InArgs._Method)
+		.ButtonContent()
+		[
+			ButtonContent.ToSharedRef()
+		]
+		.MenuContent()
+		[
+			ComboBoxMenuContent
+		]
+		.HasDownArrow(InArgs._HasDownArrow)
+		.ContentPadding(InArgs._ContentPadding)
+		.ForegroundColor(InArgs._ForegroundColor)
+		.OnMenuOpenChanged(this, &SSearchableComboBox::OnMenuOpenChanged)
+		.IsFocusable(true)
+		);
+	SetMenuContentWidgetToFocus(ComboListView);
+
+	// Need to establish the selected item at point of construction so its available for querying
+	// NB: If you need a selection to fire use SetItemSelection rather than setting an IntiallySelectedItem
+	SelectedItem = InArgs._InitiallySelectedItem;
+	if (TListTypeTraits<TSharedPtr<FString>>::IsPtrValid(SelectedItem))
+	{
+		ComboListView->Private_SetItemSelection(SelectedItem, true);
+	}
+
+}
+
+void SSearchableComboBox::ClearSelection()
+{
+	ComboListView->ClearSelection();
+}
+
+void SSearchableComboBox::SetSelectedItem(TSharedPtr<FString> InSelectedItem)
+{
+	if (TListTypeTraits<TSharedPtr<FString>>::IsPtrValid(InSelectedItem))
+	{
+		ComboListView->SetSelection(InSelectedItem);
+	}
+	else
+	{
+		ComboListView->ClearSelection();
+	}
+}
+
+TSharedPtr<FString> SSearchableComboBox::GetSelectedItem()
+{
+	return SelectedItem;
+}
+
+void SSearchableComboBox::RefreshOptions()
+{
+	if (!ComboListView->IsPendingRefresh())
+	{
+		ComboListView->RequestListRefresh();
+	}
+}
+
+TSharedRef<ITableRow> SSearchableComboBox::GenerateMenuItemRow(TSharedPtr<FString> InItem, const TSharedRef<STableViewBase>& OwnerTable)
+{
+	if (OnGenerateWidget.IsBound())
+	{
+		FString SearchToken = SearchField->GetText().ToString().ToLower();
+		EVisibility WidgetVisibility = EVisibility::Visible;
+		if (!SearchToken.IsEmpty())
+		{
+			if (InItem->ToLower().Find(SearchToken) < 0)
+			{
+				WidgetVisibility = EVisibility::Collapsed;
+			}
+		}
+		return SNew(SComboRow<TSharedPtr<FString>>, OwnerTable)
+			.Style(ItemStyle)
+			.Visibility(WidgetVisibility)
+			[
+				OnGenerateWidget.Execute(InItem)
+			];
+	}
+	else
+	{
+		return SNew(SComboRow<TSharedPtr<FString>>, OwnerTable)
+			[
+				SNew(STextBlock).Text(NSLOCTEXT("SlateCore", "ComboBoxMissingOnGenerateWidgetMethod", "Please provide a .OnGenerateWidget() handler."))
+			];
+
+	}
+}
+
+void SSearchableComboBox::OnMenuOpenChanged(bool bOpen)
+{
+	if (bOpen == false)
+	{
+		if (TListTypeTraits<TSharedPtr<FString>>::IsPtrValid(SelectedItem))
+		{
+			// Ensure the ListView selection is set back to the last committed selection
+			ComboListView->SetSelection(SelectedItem, ESelectInfo::OnNavigation);
+			ComboListView->RequestScrollIntoView(SelectedItem, 0);
+		}
+
+		// Set focus back to ComboBox for users focusing the ListView that just closed
+		FSlateApplication::Get().ForEachUser([&](FSlateUser* User) {
+			if (FSlateApplication::Get().HasUserFocusedDescendants(AsShared(), User->GetUserIndex()))
+			{
+				FSlateApplication::Get().SetUserFocus(User->GetUserIndex(), AsShared(), EFocusCause::SetDirectly);
+			}
+		});
+
+	}
+}
+
+
+void SSearchableComboBox::OnSelectionChanged_Internal(TSharedPtr<FString> ProposedSelection, ESelectInfo::Type SelectInfo)
+{
+	// Ensure that the proposed selection is different
+	if (SelectInfo != ESelectInfo::OnNavigation)
+	{
+		// Ensure that the proposed selection is different from selected
+		if (ProposedSelection != SelectedItem)
+		{
+			SelectedItem = ProposedSelection;
+			OnSelectionChanged.ExecuteIfBound(ProposedSelection, SelectInfo);
+		}
+		// close combo even if user reselected item
+		this->SetIsOpen(false);
+	}
+}
+
+void SSearchableComboBox::OnSearchTextChanged(const FText& ChangedText)
+{
+	FString SearchToken = ChangedText.ToString().ToLower();
+	for (int32 i = 0; i < OptionsSource->Num(); i++)
+	{
+		TSharedPtr<ITableRow> Row = ComboListView->WidgetFromItem((*OptionsSource)[i]);
+		if (Row)
+		{
+			if (SearchToken.IsEmpty())
+			{
+				Row->AsWidget()->SetVisibility(EVisibility::Visible);
+			}
+			else if ((*OptionsSource)[i]->ToLower().Find(SearchToken) >= 0)
+			{
+				Row->AsWidget()->SetVisibility(EVisibility::Visible);
+			}
+			else
+			{
+				Row->AsWidget()->SetVisibility(EVisibility::Collapsed);
+			}
+		}
+	}
+
+	ComboListView->RequestListRefresh();
+
+	SelectedItem = TSharedPtr< FString >();
+}
+
+FReply SSearchableComboBox::OnButtonClicked()
+{
+	// if user clicked to close the combo menu
+	if (this->IsOpen())
+	{
+		// Re-select first selected item, just in case it was selected by navigation previously
+		TArray<TSharedPtr<FString>> SelectedItems = ComboListView->GetSelectedItems();
+		if (SelectedItems.Num() > 0)
+		{
+			OnSelectionChanged_Internal(SelectedItems[0], ESelectInfo::Direct);
+		}
+	}
+	else
+	{
+		SearchField->SetText(FText::GetEmpty());
+		OnComboBoxOpening.ExecuteIfBound();
+	}
+
+	return SComboButton::OnButtonClicked();
+}
+
+#undef LOCTEXT_NAMESPACE
diff --git a/Engine/Source/Editor/DetailCustomizations/Private/SSearchableComboBox.h b/Engine/Source/Editor/DetailCustomizations/Private/SSearchableComboBox.h
new file mode 100644
--- /dev/null
+++ b/Engine/Source/Editor/DetailCustomizations/Private/SSearchableComboBox.h
@@ -0,0 +1,140 @@
+// Copyright 1998-2018 Epic Games, Inc. All Rights Reservekd.
+
+#pragma once
+
+#include "CoreMinimal.h"
+#include "Input/Reply.h"
+#include "Widgets/SWidget.h"
+#include "Widgets/Input/SComboButton.h"
+#include "Widgets/Input/SEditableTextBox.h"
+#include "Widgets/Views/SListView.h"
+#include "Widgets/Input/SComboBox.h"
+
+/**
+ * A searchable text combo box
+ */
+class SSearchableComboBox : public SComboButton
+{
+public:
+	/** Type of list used for showing menu options. */
+	typedef SListView< TSharedPtr<FString> > SComboListType;
+	/** Delegate type used to generate widgets that represent Options */
+	typedef typename TSlateDelegates< TSharedPtr<FString> >::FOnGenerateWidget FOnGenerateWidget;
+	typedef typename TSlateDelegates< TSharedPtr<FString> >::FOnSelectionChanged FOnSelectionChanged;
+
+	SLATE_BEGIN_ARGS(SSearchableComboBox)
+		: _Content()
+		, _ComboBoxStyle(&FCoreStyle::Get().GetWidgetStyle< FComboBoxStyle >("ComboBox"))
+		, _ButtonStyle(nullptr)
+		, _ItemStyle(&FCoreStyle::Get().GetWidgetStyle< FTableRowStyle >("TableView.Row"))
+		, _ContentPadding(FMargin(4.0, 2.0))
+		, _ForegroundColor(FCoreStyle::Get().GetSlateColor("InvertedForeground"))
+		, _OptionsSource()
+		, _OnSelectionChanged()
+		, _OnGenerateWidget()
+		, _InitiallySelectedItem(nullptr)
+		, _Method()
+		, _MaxListHeight(450.0f)
+		, _HasDownArrow(true)
+	{}
+
+	/** Slot for this button's content (optional) */
+	SLATE_DEFAULT_SLOT(FArguments, Content)
+
+		SLATE_STYLE_ARGUMENT(FComboBoxStyle, ComboBoxStyle)
+
+		/** The visual style of the button part of the combo box (overrides ComboBoxStyle) */
+		SLATE_STYLE_ARGUMENT(FButtonStyle, ButtonStyle)
+
+		SLATE_STYLE_ARGUMENT(FTableRowStyle, ItemStyle)
+
+		SLATE_ATTRIBUTE(FMargin, ContentPadding)
+		SLATE_ATTRIBUTE(FSlateColor, ForegroundColor)
+
+		SLATE_ARGUMENT(const TArray< TSharedPtr<FString> >*, OptionsSource)
+		SLATE_EVENT(FOnSelectionChanged, OnSelectionChanged)
+		SLATE_EVENT(FOnGenerateWidget, OnGenerateWidget)
+
+		/** Called when combo box is opened, before list is actually created */
+		SLATE_EVENT(FOnComboBoxOpening, OnComboBoxOpening)
+
+		/** The custom scrollbar to use in the ListView */
+		SLATE_ARGUMENT(TSharedPtr<SScrollBar>, CustomScrollbar)
+
+		/** The option that should be selected when the combo box is first created */
+		SLATE_ARGUMENT(TSharedPtr<FString>, InitiallySelectedItem)
+
+		SLATE_ARGUMENT(TOptional<EPopupMethod>, Method)
+
+		/** The max height of the combo box menu */
+		SLATE_ARGUMENT(float, MaxListHeight)
+
+		/**
+		 * When false, the down arrow is not generated and it is up to the API consumer
+		 * to make their own visual hint that this is a drop down.
+		 */
+		SLATE_ARGUMENT(bool, HasDownArrow)
+
+	SLATE_END_ARGS()
+
+	/**
+	 * Construct the widget from a declaration
+	 *
+	 * @param InArgs   Declaration from which to construct the combo box
+	 */
+	void Construct(const FArguments& InArgs);
+
+	void ClearSelection();
+
+	void SetSelectedItem(TSharedPtr<FString> InSelectedItem);
+
+	/** @return the item currently selected by the combo box. */
+	TSharedPtr<FString> GetSelectedItem();
+
+	/**
+	 * Requests a list refresh after updating options
+	 * Call SetSelectedItem to update the selected item if required
+	 * @see SetSelectedItem
+	 */
+	void RefreshOptions();
+
+private:
+
+	/** Generate a row for the InItem in the combo box's list (passed in as OwnerTable). Do this by calling the user-specified OnGenerateWidget */
+	TSharedRef<ITableRow> GenerateMenuItemRow(TSharedPtr<FString> InItem, const TSharedRef<STableViewBase>& OwnerTable);
+
+	//** Called if the menu is closed
+	void OnMenuOpenChanged(bool bOpen);
+
+	/** Invoked when the selection in the list changes */
+	void OnSelectionChanged_Internal(TSharedPtr<FString> ProposedSelection, ESelectInfo::Type SelectInfo);
+
+	/** Invoked when the search text changes */
+	void OnSearchTextChanged(const FText& ChangedText);
+
+	/** Handle clicking on the content menu */
+	virtual FReply OnButtonClicked() override;
+
+	/** The item style to use. */
+	const FTableRowStyle* ItemStyle;
+
+private:
+	/** Delegate that is invoked when the selected item in the combo box changes */
+	FOnSelectionChanged OnSelectionChanged;
+	/** The item currently selected in the combo box */
+	TSharedPtr<FString> SelectedItem;
+	/** The search field used for the combox box's contents */
+	TSharedPtr< SEditableTextBox > SearchField;
+	/** The ListView that we pop up; visualized the available options. */
+	TSharedPtr< SComboListType > ComboListView;
+	/** The Scrollbar used in the ListView. */
+	TSharedPtr< SScrollBar > CustomScrollbar;
+	/** Delegate to invoke before the combo box is opening. */
+	FOnComboBoxOpening OnComboBoxOpening;
+	/** Delegate to invoke when we need to visualize an option as a widget. */
+	FOnGenerateWidget OnGenerateWidget;
+
+	const TArray< TSharedPtr<FString> >* OptionsSource;
+};
+
+
diff --git a/Engine/Source/Runtime/Slate/Private/Widgets/Views/SListPanel.cpp b/Engine/Source/Runtime/Slate/Private/Widgets/Views/SListPanel.cpp
--- a/Engine/Source/Runtime/Slate/Private/Widgets/Views/SListPanel.cpp
+++ b/Engine/Source/Runtime/Slate/Private/Widgets/Views/SListPanel.cpp
@@ -77,14 +77,15 @@ void SListPanel::OnArrangeChildren( const FGeometry& AllottedGeometry, FArranged
 				if ( ListItemAlignment == EListItemAlignment::RightAligned || ListItemAlignment == EListItemAlignment::CenterAligned )
 				{
 					const float LinePadding = GetLinePadding(AllottedGeometry, ItemIndex);
+					const bool IsVisible = Children[ItemIndex].GetWidget()->GetVisibility().IsVisible();
 					if ( ListItemAlignment == EListItemAlignment::RightAligned )
 					{
-						WidthSoFar += LinePadding;
+						WidthSoFar += IsVisible ? LinePadding : 0;
 					}
 					else
 					{
 						const float HalfLinePadding = LinePadding * 0.5;
-						WidthSoFar += HalfLinePadding;
+						WidthSoFar += IsVisible ? HalfLinePadding : 0;
 					}
 				}
 
@@ -114,10 +115,9 @@ void SListPanel::OnArrangeChildren( const FGeometry& AllottedGeometry, FArranged
 			for( int32 ItemIndex = 0; ItemIndex < Children.Num(); ++ItemIndex )
 			{
 				const FVector2D ItemDesiredSize = Children[ItemIndex].GetWidget()->GetDesiredSize();
-				const float LocalItemHeight = ItemDesiredSize.Y;
+				const bool IsVisible = Children[ItemIndex].GetWidget()->GetVisibility().IsVisible();
+				const float LocalItemHeight = IsVisible ? ItemDesiredSize.Y : 0.f;
 
-				// Note that ListPanel does not respect child Visibility.
-				// It is simply not useful for ListPanels.
 				ArrangedChildren.AddWidget(
 					AllottedGeometry.MakeChild( Children[ItemIndex].GetWidget(), FVector2D(0, HeightSoFar), FVector2D(AllottedGeometry.GetLocalSize().X, LocalItemHeight) )
 					);
diff --git a/Engine/Source/Runtime/Slate/Public/Widgets/Views/SListView.h b/Engine/Source/Runtime/Slate/Public/Widgets/Views/SListView.h
--- a/Engine/Source/Runtime/Slate/Public/Widgets/Views/SListView.h
+++ b/Engine/Source/Runtime/Slate/Public/Widgets/Views/SListView.h
@@ -984,7 +984,7 @@ public:
 
 					const float ItemHeight = GenerateWidgetForItem(CurItem, ItemIndex, StartIndex, LayoutScaleMultiplier);
 
-					if (HeightGeneratedSoFar + ItemHeight > MyGeometry.GetLocalSize().Y)
+					if (HeightGeneratedSoFar + ItemHeight > MyGeometry.GetLocalSize().Y && ItemHeight > 0.f)
 					{
 						// Generated the item that puts us over the top.
 						// Count the fraction of this item that will stick out above the list
@@ -1028,7 +1028,8 @@ public:
 		NewlyGeneratedWidget->SlatePrepass(LayoutScaleMultiplier);
 
 		const bool IsFirstWidgetOnScreen = (ItemIndex == StartIndex);
-		const float ItemHeight = NewlyGeneratedWidget->GetDesiredSize().Y;
+		const bool IsVisible = NewlyGeneratedWidget->GetVisibility().IsVisible();
+		const float ItemHeight = IsVisible ? NewlyGeneratedWidget->GetDesiredSize().Y : 0;
 
 		// We have a widget for this item; add it to the panel so that it is part of the UI.
 		if (ItemIndex >= StartIndex)
