UE4.22 - Fix windows target RHI name
#rnx

diff --git a/Engine/Source/Editor/DetailCustomizations/Private/WindowsTargetSettingsDetails.cpp b/Engine/Source/Editor/DetailCustomizations/Private/WindowsTargetSettingsDetails.cpp
--- a/Engine/Source/Editor/DetailCustomizations/Private/WindowsTargetSettingsDetails.cpp
+++ b/Engine/Source/Editor/DetailCustomizations/Private/WindowsTargetSettingsDetails.cpp
@@ -52,7 +52,7 @@ static FText GetFriendlyNameFromWindowsRHIName(const FString& InRHIName)
 	FText FriendlyRHIName;
 	if (InRHIName == TEXT("PCD3D_SM5"))
 	{
-		FriendlyRHIName = LOCTEXT("DirectX11", "DirectX 11 (SM5)");
+		FriendlyRHIName = LOCTEXT("DirectX11", "DirectX 11 & 12 (SM5)");
 	}
 	else if (InRHIName == TEXT("PCD3D_SM4"))
 	{
@@ -64,7 +64,7 @@ static FText GetFriendlyNameFromWindowsRHIName(const FString& InRHIName)
 	}
 	else if (InRHIName == TEXT("SF_VULKAN_SM5"))
 	{
-		FriendlyRHIName = LOCTEXT("VulkanSM5", "Vulkan Desktop (SM5, Experimental)");
+		FriendlyRHIName = LOCTEXT("VulkanSM5", "Vulkan (SM5, Experimental)");
 	}
 	else if (InRHIName == TEXT("GLSL_SWITCH"))
 	{
