Changed the directional light intensity slider to a max of 150 instead of 120000.
It can still be set to higher values manually.

diff --git a/Engine/Source/Editor/DetailCustomizations/Private/DirectionalLightComponentDetails.cpp b/Engine/Source/Editor/DetailCustomizations/Private/DirectionalLightComponentDetails.cpp
--- a/Engine/Source/Editor/DetailCustomizations/Private/DirectionalLightComponentDetails.cpp
+++ b/Engine/Source/Editor/DetailCustomizations/Private/DirectionalLightComponentDetails.cpp
@@ -37,8 +37,8 @@ void FDirectionalLightComponentDetails::CustomizeDetails( IDetailLayoutBuilder&
 	TSharedPtr<IPropertyHandle> LightIntensityProperty = DetailBuilder.GetProperty("Intensity", ULightComponentBase::StaticClass());
 	// Point lights need to override the ui min and max for units of lumens, so we have to undo that
 	LightIntensityProperty->SetInstanceMetaData("UIMin",TEXT("0.0f"));
-	LightIntensityProperty->SetInstanceMetaData("UIMax",TEXT("120000.0f"));
-	LightIntensityProperty->SetInstanceMetaData("SliderExponent", TEXT("10.0f"));
+	LightIntensityProperty->SetInstanceMetaData("UIMax",TEXT("150.0f"));
+	LightIntensityProperty->SetInstanceMetaData("SliderExponent", TEXT("2.0f"));
 	LightIntensityProperty->SetInstanceMetaData("Units", TEXT("lux"));
 	LightIntensityProperty->SetToolTipText(LOCTEXT("DirectionalLightIntensityToolTipText", "Maximum illumination from the light in lux"));
 
