- Landscape Layer System: Move setting to enable/disable layers to the LandscapeActor itself

#rnx

diff --git a/Engine/Source/Editor/DetailCustomizations/Private/LandscapeUIDetails.cpp b/Engine/Source/Editor/DetailCustomizations/Private/LandscapeUIDetails.cpp
--- a/Engine/Source/Editor/DetailCustomizations/Private/LandscapeUIDetails.cpp
+++ b/Engine/Source/Editor/DetailCustomizations/Private/LandscapeUIDetails.cpp
@@ -1,13 +1,16 @@
 // Copyright 1998-2019 Epic Games, Inc. All Rights Reserved.
 
 #include "LandscapeUIDetails.h"
-//#include "Misc/Attribute.h"
-//#include "Misc/Guid.h"
 #include "Widgets/DeclarativeSyntaxSupport.h"
 #include "PropertyHandle.h"
 #include "DetailLayoutBuilder.h"
 #include "Runtime/Landscape/Classes/Landscape.h"
-//#include "Editor.h"
+#include "Settings/EditorExperimentalSettings.h"
+#include "DetailCategoryBuilder.h"
+#include "DetailWidgetRow.h"
+#include "Widgets/Input/SButton.h"
+#include "Widgets/Text/STextBlock.h"
+#include "Misc/MessageDialog.h"
 
 FLandscapeUIDetails::FLandscapeUIDetails()
 {
@@ -15,6 +18,10 @@ FLandscapeUIDetails::FLandscapeUIDetails()
 
 FLandscapeUIDetails::~FLandscapeUIDetails()
 {
+	if (DetailLayoutBuilder)
+	{
+		GetMutableDefault<UEditorExperimentalSettings>()->OnSettingChanged().RemoveAll(this);
+	}
 }
 
 TSharedRef<IDetailCustomization> FLandscapeUIDetails::MakeInstance()
@@ -36,5 +43,98 @@ void FLandscapeUIDetails::CustomizeDetails( IDetailLayoutBuilder& DetailBuilder
 			TSharedRef<IPropertyHandle> ComponentScreenSizeToUseSubSectionsProp = DetailBuilder.GetProperty(GET_MEMBER_NAME_CHECKED(ALandscapeProxy, ComponentScreenSizeToUseSubSections));
 			DetailBuilder.HideProperty(ComponentScreenSizeToUseSubSectionsProp);
 		}
+
+		if (Landscape != nullptr)
+		{
+			if (!DetailLayoutBuilder)
+			{
+				DetailLayoutBuilder = &DetailBuilder;
+				GetMutableDefault<UEditorExperimentalSettings>()->OnSettingChanged().AddSP(this, &FLandscapeUIDetails::OnEditorExperimentalSettingsChanged);
+			}
+
+			TSharedRef<IPropertyHandle> PropertyHandle = DetailBuilder.GetProperty(GET_MEMBER_NAME_CHECKED(ALandscape, bCanHaveLayersContent));
+			DetailBuilder.HideProperty(PropertyHandle);
+			if(GetMutableDefault<UEditorExperimentalSettings>()->bLandscapeLayerSystem)
+			{
+				IDetailCategoryBuilder& Category = DetailBuilder.EditCategory(PropertyHandle->GetDefaultCategoryName());
+				Category.AddCustomRow(FText(), true)
+				.NameContent()
+				[
+					PropertyHandle->CreatePropertyNameWidget(FText::FromString("Layer System"))
+				]
+				.ValueContent()
+				[
+					SNew(SButton)
+					.OnClicked_Lambda([this, Landscape]() -> FReply
+					{
+						return ToggleCanHaveLayersContent(Landscape);
+					})
+					[
+						SNew(STextBlock)
+						.Text_Lambda([Landscape]()
+						{
+							return Landscape->bCanHaveLayersContent ? NSLOCTEXT("UnrealEd", "LandscapeDisableLayers", "Disable") : NSLOCTEXT("UnrealEd", "LandscapeEnableLayers", "Enable");	
+						})
+					]
+				];
+			}
+			
+		}
 	}
+}
+
+void FLandscapeUIDetails::OnEditorExperimentalSettingsChanged(FName PropertyName)
+{
+	if (DetailLayoutBuilder && PropertyName == GET_MEMBER_NAME_CHECKED(UEditorExperimentalSettings, bLandscapeLayerSystem))
+	{
+		DetailLayoutBuilder->ForceRefreshDetails();
+	}
+}
+
+FReply FLandscapeUIDetails::ToggleCanHaveLayersContent(ALandscape* Landscape)
+{
+	bool bToggled = false;
+
+	if (Landscape->bCanHaveLayersContent)
+	{
+		bool bHasHiddenLayers = false;
+		for (int32 i = 0; i < Landscape->GetLayerCount(); ++i)
+		{
+			const FLandscapeLayer* Layer = Landscape->GetLayer(i);
+			check(Layer != nullptr);
+
+			if (!Layer->bVisible)
+			{
+				bHasHiddenLayers = true;
+				break;
+			}
+		}
+
+		
+		FText Reason;
+
+		if (bHasHiddenLayers)
+		{
+			Reason = NSLOCTEXT("UnrealEd", "LandscapeDisableLayers_HiddenLayers", "Are you sure you want to disable the layers system?\n\nDoing so, will result in losing the data stored for each layers, but the current visual output will be kept. Be aware that some layers are currently hidden, continuing will result in their data being lost.");
+		}
+		else
+		{
+			Reason = NSLOCTEXT("UnrealEd", "LandscapeDisableLayers", "Are you sure you want to disable the layers system?\n\nDoing so, will result in losing the data stored for each layers, but the current visual output will be kept.");
+		}
+
+		bToggled = FMessageDialog::Open(EAppMsgType::YesNo, Reason) == EAppReturnType::Yes;
+	}
+	else
+	{
+		FMessageDialog::Open(EAppMsgType::Ok, NSLOCTEXT("UnrealEd", "LandscapeEnableLayers", "Layer system as been enabled on this landscape and streaming proxies."));
+		bToggled = true;
+	}
+
+	if (DetailLayoutBuilder && bToggled)
+	{
+		Landscape->ToggleCanHaveLayersContent();
+		DetailLayoutBuilder->ForceRefreshDetails();
+	}
+
+	return FReply::Handled();
 }
\ No newline at end of file
diff --git a/Engine/Source/Editor/DetailCustomizations/Private/LandscapeUIDetails.h b/Engine/Source/Editor/DetailCustomizations/Private/LandscapeUIDetails.h
--- a/Engine/Source/Editor/DetailCustomizations/Private/LandscapeUIDetails.h
+++ b/Engine/Source/Editor/DetailCustomizations/Private/LandscapeUIDetails.h
@@ -8,6 +8,7 @@
 #include "IDetailCustomization.h"
 
 class IDetailLayoutBuilder;
+class ALandscape;
 
 class FLandscapeUIDetails : public IDetailCustomization
 {
@@ -22,7 +23,11 @@ public:
 	/** End IDetailCustomization interface */
 
 private:
+	void OnEditorExperimentalSettingsChanged(FName PropertyName);
+	FReply ToggleCanHaveLayersContent(ALandscape* Landscape);
 
 	/** Use MakeInstance to create an instance of this class */
 	FLandscapeUIDetails();
+
+	IDetailLayoutBuilder* DetailLayoutBuilder = nullptr;
 };
diff --git a/Engine/Source/Editor/DetailCustomizations/Private/WorldSettingsDetails.cpp b/Engine/Source/Editor/DetailCustomizations/Private/WorldSettingsDetails.cpp
--- a/Engine/Source/Editor/DetailCustomizations/Private/WorldSettingsDetails.cpp
+++ b/Engine/Source/Editor/DetailCustomizations/Private/WorldSettingsDetails.cpp
@@ -52,11 +52,6 @@ void FWorldSettingsDetails::CustomizeDetails( IDetailLayoutBuilder& DetailBuilde
 	AddLightmapCustomization(DetailBuilder);
 
 	DetailBuilder.HideProperty(GET_MEMBER_NAME_CHECKED(AActor, bHidden), AActor::StaticClass());
-
-	if (!GetMutableDefault<UEditorExperimentalSettings>()->bLandscapeLayerSystem)
-	{
-		DetailBuilder.HideProperty(GET_MEMBER_NAME_CHECKED(AWorldSettings, bEnableLandscapeLayers), AWorldSettings::StaticClass());
-	}
 }
 
 
diff --git a/Engine/Source/Editor/LandscapeEditor/Private/LandscapeEdMode.cpp b/Engine/Source/Editor/LandscapeEditor/Private/LandscapeEdMode.cpp
--- a/Engine/Source/Editor/LandscapeEditor/Private/LandscapeEdMode.cpp
+++ b/Engine/Source/Editor/LandscapeEditor/Private/LandscapeEdMode.cpp
@@ -466,6 +466,11 @@ TSharedRef<FUICommandList> FEdModeLandscape::GetUICommandList() const
 	return Toolkit->GetToolkitCommands();
 }
 
+void FEdModeLandscape::OnCanHaveLayersContentChanged()
+{
+	RefreshDetailPanel();
+}
+
 ELandscapeToolTargetType::Type FEdModeLandscape::GetLandscapeToolTargetType() const
 {
 	if (CurrentToolMode)
@@ -525,31 +530,7 @@ void FEdModeLandscape::Enter()
 
 	OnLevelActorDeletedDelegateHandle = GEngine->OnLevelActorDeleted().AddSP(this, &FEdModeLandscape::OnLevelActorRemoved);
 	OnLevelActorAddedDelegateHandle = GEngine->OnLevelActorAdded().AddSP(this, &FEdModeLandscape::OnLevelActorAdded);
-
-	AWorldSettings* WorldSettings = GWorld != nullptr ? GWorld->GetWorldSettings() : nullptr;
-	
-	if (WorldSettings)
-	{
-		OnLandscapeLayerSystemFlagChangedDelegateHandle = WorldSettings->OnSettingChanged().AddLambda([this](const FName& InPropertyName)
-		{
-			if (InPropertyName == GET_MEMBER_NAME_CHECKED(AWorldSettings, bEnableLandscapeLayers))
-			{
-				UpdateToolModes();
-
-				if (CurrentToolMode->ToolModeName == TEXT("ToolMode_Paint"))
-				{
-					SetCurrentTool("Paint");
-				}
-				else if (CurrentToolMode->ToolModeName == TEXT("ToolMode_Sculpt"))
-				{
-					SetCurrentTool("Sculpt");
-				}
-
-				RefreshDetailPanel();
-			}
-		});
-	}
-
+		
 	UpdateToolModes();
 
 	ALandscapeProxy* SelectedLandscape = GEditor->GetSelectedActors()->GetTop<ALandscapeProxy>();
@@ -763,14 +744,6 @@ void FEdModeLandscape::Exit()
 
 	GEngine->OnLevelActorDeleted().Remove(OnLevelActorDeletedDelegateHandle);
 	GEngine->OnLevelActorAdded().Remove(OnLevelActorAddedDelegateHandle);
-	GetMutableDefault<UEditorExperimentalSettings>()->OnSettingChanged().Remove(OnLandscapeLayerSystemFlagChangedDelegateHandle);
-
-	AWorldSettings* WorldSettings = GWorld != nullptr ? GWorld->GetWorldSettings() : nullptr;
-
-	if (WorldSettings)
-	{
-		WorldSettings->OnSettingChanged().Remove(OnLandscapeLayerSystemFlagChangedDelegateHandle);
-	}
 	
 	FEditorSupportDelegates::WorldChange.Remove(OnWorldChangeDelegateHandle);
 	GetWorld()->OnLevelsChanged().Remove(OnLevelsChangedDelegateHandle);
@@ -2253,7 +2226,12 @@ void FEdModeLandscape::SetCurrentTool(int32 ToolIndex)
 				break;
 			}
 		}
-		check(bFoundValidMode);
+		
+		// default to first valid tool of current mode
+		if (!bFoundValidMode)
+		{
+			SetCurrentTool(CurrentToolMode->ValidTools[0]);
+		}
 	}
 
 	// Set target type appropriate for tool
@@ -4078,7 +4056,7 @@ ALandscape* FEdModeLandscape::ChangeComponentSetting(int32 NumComponentsX, int32
 			FActorSpawnParameters SpawnParams;
 			SpawnParams.OverrideLevel = OldLandscape->GetLevel();
 			NewLandscape = OldLandscape->GetWorld()->SpawnActor<ALandscape>(Location, OldLandscape->GetActorRotation(), SpawnParams);
-
+			NewLandscape->bCanHaveLayersContent = OldLandscape->bCanHaveLayersContent;
 			const FVector OldScale = OldLandscape->GetActorScale();
 			NewLandscape->SetActorRelativeScale3D(FVector(OldScale.X * LandscapeScaleFactor, OldScale.Y * LandscapeScaleFactor, OldScale.Z));
 
diff --git a/Engine/Source/Editor/LandscapeEditor/Private/LandscapeEdMode.h b/Engine/Source/Editor/LandscapeEditor/Private/LandscapeEdMode.h
--- a/Engine/Source/Editor/LandscapeEditor/Private/LandscapeEdMode.h
+++ b/Engine/Source/Editor/LandscapeEditor/Private/LandscapeEdMode.h
@@ -333,6 +333,7 @@ public:
 	virtual ELandscapeToolTargetType::Type GetLandscapeToolTargetType() const override;
 	virtual const FLandscapeLayer* GetLandscapeSelectedLayer() const override;
 	virtual ULandscapeLayerInfoObject* GetSelectedLandscapeLayerInfo() const override;
+	virtual void OnCanHaveLayersContentChanged() override;
 
 	/** FGCObject interface */
 	virtual void AddReferencedObjects(FReferenceCollector& Collector) override;
@@ -599,8 +600,7 @@ private:
 
 	FDelegateHandle OnLevelActorDeletedDelegateHandle;
 	FDelegateHandle OnLevelActorAddedDelegateHandle;
-	FDelegateHandle OnLandscapeLayerSystemFlagChangedDelegateHandle;
-
+	
 	/** Check if we are painting using the VREditor */
 	bool bIsPaintingInVR;
 
diff --git a/Engine/Source/Editor/LandscapeEditor/Private/LandscapeEditorDetailCustomization_NewLandscape.cpp b/Engine/Source/Editor/LandscapeEditor/Private/LandscapeEditorDetailCustomization_NewLandscape.cpp
--- a/Engine/Source/Editor/LandscapeEditor/Private/LandscapeEditorDetailCustomization_NewLandscape.cpp
+++ b/Engine/Source/Editor/LandscapeEditor/Private/LandscapeEditorDetailCustomization_NewLandscape.cpp
@@ -40,6 +40,7 @@
 #include "Framework/Application/SlateApplication.h"
 #include "Widgets/Input/SNumericEntryBox.h"
 #include "LandscapeDataAccess.h"
+#include "Settings/EditorExperimentalSettings.h"
 
 #define LOCTEXT_NAMESPACE "LandscapeEditor.NewLandscape"
 
@@ -87,6 +88,16 @@ void FLandscapeEditorDetailCustomization_NewLandscape::CustomizeDetails(IDetailL
 		]
 	];
 
+	TSharedRef<IPropertyHandle> PropertyHandle_CanHaveLayersContent = DetailBuilder.GetProperty(GET_MEMBER_NAME_CHECKED(ULandscapeEditorObject, bCanHaveLayersContent));
+	if (GetMutableDefault<UEditorExperimentalSettings>()->bLandscapeLayerSystem)
+	{
+		NewLandscapeCategory.AddProperty(PropertyHandle_CanHaveLayersContent);
+	}
+	else
+	{
+		DetailBuilder.HideProperty(PropertyHandle_CanHaveLayersContent);
+	}
+
 	TSharedRef<IPropertyHandle> PropertyHandle_HeightmapFilename = DetailBuilder.GetProperty(GET_MEMBER_NAME_CHECKED(ULandscapeEditorObject, ImportLandscape_HeightmapFilename));
 	TSharedRef<IPropertyHandle> PropertyHandle_HeightmapImportResult = DetailBuilder.GetProperty(GET_MEMBER_NAME_CHECKED(ULandscapeEditorObject, ImportLandscape_HeightmapImportResult));
 	TSharedRef<IPropertyHandle> PropertyHandle_HeightmapErrorMessage = DetailBuilder.GetProperty(GET_MEMBER_NAME_CHECKED(ULandscapeEditorObject, ImportLandscape_HeightmapErrorMessage));
@@ -809,6 +820,7 @@ FReply FLandscapeEditorDetailCustomization_NewLandscape::OnCreateButtonClicked()
 		const FVector Offset = FTransform(UISettings->NewLandscape_Rotation, FVector::ZeroVector, UISettings->NewLandscape_Scale).TransformVector(FVector(-ComponentCountX * QuadsPerComponent / 2, -ComponentCountY * QuadsPerComponent / 2, 0));
 		
 		ALandscape* Landscape = LandscapeEdMode->GetWorld()->SpawnActor<ALandscape>(UISettings->NewLandscape_Location + Offset, UISettings->NewLandscape_Rotation);
+		Landscape->bCanHaveLayersContent = UISettings->bCanHaveLayersContent;
 		Landscape->LandscapeMaterial = UISettings->NewLandscape_Material.Get();
 		Landscape->SetActorRelativeScale3D(UISettings->NewLandscape_Scale);
 
diff --git a/Engine/Source/Editor/LandscapeEditor/Private/LandscapeEditorModule.cpp b/Engine/Source/Editor/LandscapeEditor/Private/LandscapeEditorModule.cpp
--- a/Engine/Source/Editor/LandscapeEditor/Private/LandscapeEditorModule.cpp
+++ b/Engine/Source/Editor/LandscapeEditor/Private/LandscapeEditorModule.cpp
@@ -71,8 +71,7 @@ public:
 			true,
 			300
 			);
-		ALandscape::RegisterChangeLandscapeLayersStateDelegate();
-
+	
 		// register customizations
 		FPropertyEditorModule& PropertyModule = FModuleManager::GetModuleChecked<FPropertyEditorModule>("PropertyEditor");
 		PropertyModule.RegisterCustomClassLayout("LandscapeEditorObject", FOnGetDetailCustomizationInstance::CreateStatic(&FLandscapeEditorDetails::MakeInstance));
@@ -119,7 +118,6 @@ public:
 	 */
 	virtual void ShutdownModule() override
 	{
-		ALandscape::UnregisterChangeLandscapeLayersStateDelegate();
 		FLandscapeEditorCommands::Unregister();
 
 		FEditorDelegates::PreSaveWorld.Remove(PreSaveWorldHandle);
diff --git a/Engine/Source/Editor/LandscapeEditor/Public/LandscapeEditorObject.h b/Engine/Source/Editor/LandscapeEditor/Public/LandscapeEditorObject.h
--- a/Engine/Source/Editor/LandscapeEditor/Public/LandscapeEditorObject.h
+++ b/Engine/Source/Editor/LandscapeEditor/Public/LandscapeEditorObject.h
@@ -523,6 +523,8 @@ private:
 	UPROPERTY(NonTransactional)
 	TArray<uint16> ImportLandscape_Data;
 public:
+	UPROPERTY(Category = "New Landscape", EditAnywhere, NonTransactional, meta = (DisplayName= "Enable Layer Support", ShowForTools= "NewLandscape"))
+	bool bCanHaveLayersContent = false;
 
 	// Whether the imported alpha maps are to be interpreted as "layered" or "additive" (UE4 uses additive internally)
 	UPROPERTY(Category="New Landscape", EditAnywhere, NonTransactional, meta=(DisplayName="Layer Alphamap Type", ShowForTools="NewLandscape"))
diff --git a/Engine/Source/Editor/WorldBrowser/Private/Tiles/WorldTileCollectionModel.cpp b/Engine/Source/Editor/WorldBrowser/Private/Tiles/WorldTileCollectionModel.cpp
--- a/Engine/Source/Editor/WorldBrowser/Private/Tiles/WorldTileCollectionModel.cpp
+++ b/Engine/Source/Editor/WorldBrowser/Private/Tiles/WorldTileCollectionModel.cpp
@@ -1629,6 +1629,7 @@ void FWorldTileCollectionModel::ImportTiledLandscape_Executed()
 			SetupLandscapeImportLayers(ImportSettings, GetWorld()->GetOutermost()->GetName(), INDEX_NONE, ImportLayers);
 
 			// Set landscape configuration
+			Landscape->bCanHaveLayersContent = false; 
 			Landscape->LandscapeMaterial	= ImportSettings.LandscapeMaterial.Get();
 			Landscape->ComponentSizeQuads	= ImportSettings.QuadsPerSection*ImportSettings.SectionsPerComponent;
 			Landscape->NumSubsections		= ImportSettings.SectionsPerComponent;
diff --git a/Engine/Source/Runtime/Engine/Classes/GameFramework/WorldSettings.h b/Engine/Source/Runtime/Engine/Classes/GameFramework/WorldSettings.h
--- a/Engine/Source/Runtime/Engine/Classes/GameFramework/WorldSettings.h
+++ b/Engine/Source/Runtime/Engine/Classes/GameFramework/WorldSettings.h
@@ -463,11 +463,7 @@ public:
 	 */
 	UPROPERTY(EditAnywhere, BlueprintReadOnly, Category=World)
 	uint8 bEnableWorldComposition:1;
-
-	/** Enables landscape layers system for this world */
-	UPROPERTY(EditAnywhere, BlueprintReadOnly, Category = World)
-	uint8 bEnableLandscapeLayers : 1;
-
+		
 	/**
 	 * Enables client-side streaming volumes instead of server-side.
 	 * Expected usage scenario: server has all streaming levels always loaded, clients independently stream levels in/out based on streaming volumes.
@@ -718,22 +714,6 @@ public:
 	UPROPERTY()
 	TArray<UAssetUserData*> AssetUserData;
 
-public:
-	/**
-	 * Returns an event delegate that is executed when a setting has changed.
-	 *
-	 * @return The delegate.
-	 */
-	DECLARE_EVENT_OneParam(AWorldSettings, FSettingChangedEvent, FName /*PropertyName*/);
-	FSettingChangedEvent& OnSettingChanged()
-	{
-		return SettingChangedEvent;
-	}
-
-private:
-	// Holds an event delegate that is executed when a setting has changed.
-	FSettingChangedEvent SettingChangedEvent;
-
 	// If paused, PlayerState of person pausing the game.
 	UPROPERTY(transient, replicated)
 	class APlayerState* PauserPlayerState;
diff --git a/Engine/Source/Runtime/Engine/Private/WorldSettings.cpp b/Engine/Source/Runtime/Engine/Private/WorldSettings.cpp
--- a/Engine/Source/Runtime/Engine/Private/WorldSettings.cpp
+++ b/Engine/Source/Runtime/Engine/Private/WorldSettings.cpp
@@ -73,7 +73,6 @@ PRAGMA_ENABLE_DEPRECATION_WARNINGS
 	NavigationSystemConfig = nullptr;
 	bEnableAISystem = true;
 	bEnableWorldComposition = false;
-	bEnableLandscapeLayers = false;
 	bEnableWorldOriginRebasing = false;
 #if WITH_EDITORONLY_DATA	
 	bEnableHierarchicalLODSystem = false;
@@ -608,22 +607,6 @@ void AWorldSettings::PostEditChangeProperty(FPropertyChangedEvent& PropertyChang
 				bEnableWorldComposition = false;
 			}
 		}
-		else if (PropertyName == GET_MEMBER_NAME_CHECKED(AWorldSettings, bEnableLandscapeLayers))
-		{
-			if (ALandscape::ChangeLandscapeLayersStateEvent.IsBound())
-			{
-				bool WasChanged = ALandscape::ChangeLandscapeLayersStateEvent.Execute(GetWorld(), bEnableLandscapeLayers);
-
-				if (!WasChanged)
-				{
-					bEnableLandscapeLayers = !bEnableLandscapeLayers;
-				}
-			}
-			else
-			{
-				bEnableLandscapeLayers = false;
-			}
-		}
 		else if (PropertyName == GET_MEMBER_NAME_CHECKED(AWorldSettings, NavigationSystemConfig))
 		{
 			UWorld* World = GetWorld();
@@ -704,11 +687,6 @@ void AWorldSettings::PostEditChangeProperty(FPropertyChangedEvent& PropertyChang
 	}
 
 	Super::PostEditChangeProperty(PropertyChangedEvent);
-
-	if (PropertyThatChanged)
-	{
-		SettingChangedEvent.Broadcast(PropertyThatChanged->GetFName());
-	}
 }
 
 void UHierarchicalLODSetup::PostEditChangeProperty(struct FPropertyChangedEvent& PropertyChangedEvent)
diff --git a/Engine/Source/Runtime/Landscape/Classes/Landscape.h b/Engine/Source/Runtime/Landscape/Classes/Landscape.h
--- a/Engine/Source/Runtime/Landscape/Classes/Landscape.h
+++ b/Engine/Source/Runtime/Landscape/Classes/Landscape.h
@@ -188,8 +188,6 @@ struct FLandscapeLayer
 	TMap<ULandscapeLayerInfoObject*, bool> WeightmapLayerAllocationBlend; // True -> Substractive, False -> Additive
 };
 
-DECLARE_DELEGATE_RetVal_TwoParams(bool, FChangeLandscapeLayersState, UWorld*, bool);
-
 UCLASS(MinimalAPI, showcategories=(Display, Movement, Collision, Lighting, LOD, Input), hidecategories=(Mobility))
 class ALandscape : public ALandscapeProxy
 {
@@ -202,6 +200,7 @@ public:
 
 	//~ Begin ALandscapeProxy Interface
 	LANDSCAPE_API virtual ALandscape* GetLandscapeActor() override;
+	LANDSCAPE_API virtual const ALandscape* GetLandscapeActor() const override;
 #if WITH_EDITOR
 	//~ End ALandscapeProxy Interface
 
@@ -293,9 +292,8 @@ public:
 	LANDSCAPE_API void OnPreSave();
 
 	void ReleaseLayersRenderingResource();
-
-	LANDSCAPE_API static void RegisterChangeLandscapeLayersStateDelegate();
-	LANDSCAPE_API static void UnregisterChangeLandscapeLayersStateDelegate();
+	
+	LANDSCAPE_API void ToggleCanHaveLayersContent();
 
 private:
 	void TickLayers(float DeltaTime, ELevelTick TickType, FActorTickFunction& ThisTickFunction);
@@ -355,13 +353,14 @@ private:
 	void PrintLayersDebugTextureResource(const FString& InContext, FTextureResource* InTextureResource, uint8 InMipRender = 0, bool InOutputHeight = true, bool InOutputNormals = false) const;
 	void PrintLayersDebugHeightData(const FString& InContext, const TArray<FColor>& InHeightmapData, const FIntPoint& InDataSize, uint8 InMipRender, bool InOutputNormals = false) const;
 	void PrintLayersDebugWeightData(const FString& InContext, const TArray<FColor>& InWeightmapData, const FIntPoint& InDataSize, uint8 InMipRender) const;
-
-	LANDSCAPE_API static bool ChangeLandscapeLayersState(UWorld* InWorld, bool bInEnable);
 #endif
 
 public:
 
 #if WITH_EDITORONLY_DATA
+	UPROPERTY(EditAnywhere, Category=Experimental, AdvancedDisplay)
+	bool bCanHaveLayersContent = false;
+
 	/** Target Landscape Layer for Landscape Splines */
 	UPROPERTY()
 	FGuid LandscapeSplinesTargetLayerGuid;
@@ -378,9 +377,6 @@ public:
 	UPROPERTY(Transient)
 	TArray<UTextureRenderTarget2D*> WeightmapRTList;
 
-	/** Event that will be triggered if the Landscape Layers settings from the WorldSettings is changed*/
-	LANDSCAPE_API static FChangeLandscapeLayersState ChangeLandscapeLayersStateEvent;
-
 private:
 	/** Provides information from LandscapeEdMode */
 	ILandscapeEdModeInterface* LandscapeEdMode;
diff --git a/Engine/Source/Runtime/Landscape/Classes/LandscapeProxy.h b/Engine/Source/Runtime/Landscape/Classes/LandscapeProxy.h
--- a/Engine/Source/Runtime/Landscape/Classes/LandscapeProxy.h
+++ b/Engine/Source/Runtime/Landscape/Classes/LandscapeProxy.h
@@ -731,6 +731,7 @@ public:
 	FGuid GetLandscapeGuid() const { return LandscapeGuid; }
 	void SetLandscapeGuid(const FGuid& Guid) { LandscapeGuid = Guid; }
 	virtual ALandscape* GetLandscapeActor() PURE_VIRTUAL(GetLandscapeActor, return nullptr;)
+	virtual const ALandscape* GetLandscapeActor() const PURE_VIRTUAL(GetLandscapeActor, return nullptr;)
 
 	/* Per-frame call to update dynamic grass placement and render grassmaps */
 	void TickGrass();
diff --git a/Engine/Source/Runtime/Landscape/Classes/LandscapeStreamingProxy.h b/Engine/Source/Runtime/Landscape/Classes/LandscapeStreamingProxy.h
--- a/Engine/Source/Runtime/Landscape/Classes/LandscapeStreamingProxy.h
+++ b/Engine/Source/Runtime/Landscape/Classes/LandscapeStreamingProxy.h
@@ -29,6 +29,7 @@ public:
 
 	//~ Begin ALandscapeBase Interface
 	virtual ALandscape* GetLandscapeActor() override;
+	virtual const ALandscape* GetLandscapeActor() const override;
 #if WITH_EDITOR
 	virtual UMaterialInterface* GetLandscapeMaterial(int8 InLODIndex = INDEX_NONE) const override;
 	virtual UMaterialInterface* GetLandscapeHoleMaterial() const override;
diff --git a/Engine/Source/Runtime/Landscape/Private/Landscape.cpp b/Engine/Source/Runtime/Landscape/Private/Landscape.cpp
--- a/Engine/Source/Runtime/Landscape/Private/Landscape.cpp
+++ b/Engine/Source/Runtime/Landscape/Private/Landscape.cpp
@@ -1102,11 +1102,21 @@ ALandscapeStreamingProxy::ALandscapeStreamingProxy(const FObjectInitializer& Obj
 #endif // WITH_EDITORONLY_DATA
 }
 
+const ALandscape* ALandscape::GetLandscapeActor() const
+{
+	return this;
+}
+
 ALandscape* ALandscape::GetLandscapeActor()
 {
 	return this;
 }
 
+const ALandscape* ALandscapeStreamingProxy::GetLandscapeActor() const
+{
+	return LandscapeActor.Get();
+}
+
 ALandscape* ALandscapeStreamingProxy::GetLandscapeActor()
 {
 	return LandscapeActor.Get();
@@ -2269,13 +2279,6 @@ void ALandscapeProxy::PostLoad()
 		ULandscapeInfo* LandscapeInfo = CreateLandscapeInfo();
 
 		// Cache the value at this point as RegisterActor might create/destroy layers content if there was a mismatch between landscape & proxy
-		const bool bDeleteLayerContent = HasLayersContent() && !CanHaveLayersContent();
-		if (bDeleteLayerContent && LandscapeInfo->LandscapeActor.IsValid())
-		{
-			LandscapeInfo->LandscapeActor->DeleteLayers();
-			UE_LOG(LogLandscape, Warning, TEXT("Landscape layer data is being deleted because World Settings doesn't have layers enabled."));
-		}
-		
 		const bool bNeedOldDataMigration = !HasLayersContent() && CanHaveLayersContent();
 				
 		LandscapeInfo->RegisterActor(this, true);
diff --git a/Engine/Source/Runtime/Landscape/Private/LandscapeEditLayers.cpp b/Engine/Source/Runtime/Landscape/Private/LandscapeEditLayers.cpp
--- a/Engine/Source/Runtime/Landscape/Private/LandscapeEditLayers.cpp
+++ b/Engine/Source/Runtime/Landscape/Private/LandscapeEditLayers.cpp
@@ -1424,6 +1424,28 @@ void ALandscape::CreateLayersRenderingResource()
 	InitializeLayersWeightmapResources();
 }
 
+void ALandscape::ToggleCanHaveLayersContent()
+{
+	bCanHaveLayersContent = !bCanHaveLayersContent;
+
+	if (!bCanHaveLayersContent)
+	{
+		ReleaseLayersRenderingResource();
+		DeleteLayers();
+	}
+	else
+	{
+		check(GetLayerCount() == 0);
+		CreateDefaultLayer();
+		CopyOldDataToDefaultLayer();
+	}
+
+	if (LandscapeEdMode)
+	{
+		LandscapeEdMode->OnCanHaveLayersContentChanged();
+	}
+}
+
 void ALandscape::ReleaseLayersRenderingResource()
 {
 	check(!CanHaveLayersContent());
@@ -4407,91 +4429,6 @@ bool ALandscape::UpdateCollisionAndClients(const TArray<ULandscapeComponent*>& I
 	return bAllClientsUpdated;
 }
 
-bool ALandscape::ChangeLandscapeLayersState(UWorld* InWorld, bool bInEnable)
-{
-	if (bInEnable)
-	{
-		auto& LandscapeInfoMap = ULandscapeInfoMap::GetLandscapeInfoMap(InWorld);
-
-		for (auto& ItPair : LandscapeInfoMap.Map)
-		{
-			ULandscapeInfo* Info = ItPair.Value;
-			
-			if (Info->LandscapeActor.IsValid())
-			{
-				ALandscape* Landscape = Info->LandscapeActor.Get();
-
-				if (Landscape->GetLayerCount() == 0)
-				{
-					Landscape->CreateDefaultLayer();
-				}
-
-				Landscape->CopyOldDataToDefaultLayer();
-			}
-		}
-
-		return true;
-	}
-	else
-	{
-		bool HasHiddenLayers = false;
-
-		auto& LandscapeInfoMap = ULandscapeInfoMap::GetLandscapeInfoMap(InWorld);
-
-		for (auto& ItPair : LandscapeInfoMap.Map)
-		{
-			ULandscapeInfo* Info = ItPair.Value;
-			ALandscape* Landscape = Info->LandscapeActor.Get();
-
-			if (Landscape != nullptr)
-			{
-				for (int32 i = 0; i < Landscape->GetLayerCount(); ++i)
-				{
-					const FLandscapeLayer* Layer = Landscape->GetLayer(i);
-					check(Layer != nullptr);
-
-					if (!Layer->bVisible)
-					{
-						HasHiddenLayers = true;
-						break;
-					}
-				}
-			}
-		}
-
-		FText Reason;
-
-		if (HasHiddenLayers)
-		{
-			Reason = NSLOCTEXT("UnrealEd", "LandscapeDisableLayers_HiddenLayers", "Are you sure you want to disable the layers system?\n\nDoing so, will result in losing the data stored for each layers, but the current visual output will be kept. Be aware that some layers are currently hidden, continuing will result in their data being lost.");
-		}
-		else
-		{
-			Reason = NSLOCTEXT("UnrealEd", "LandscapeDisableLayers", "Are you sure you want to disable the layers system?\n\nDoing so, will result in losing the data stored for each layers, but the current visual output will be kept.");
-		}
-
-		if (FMessageDialog::Open(EAppMsgType::YesNo, Reason) == EAppReturnType::Yes)
-		{
-			for (auto& ItPair : LandscapeInfoMap.Map)
-			{
-				ULandscapeInfo* Info = ItPair.Value;
-				ALandscape* Landscape = Info->LandscapeActor.Get();
-
-				// Only clean up data if the Landscape actor is loaded
-				if (Landscape != nullptr)
-				{
-					Landscape->ReleaseLayersRenderingResource();
-					Landscape->DeleteLayers();
-				}			
-			}
-
-			return true;
-		}
-	}
-
-	return false;
-}
-
 void ALandscape::InitializeLayers()
 {
 	check(HasLayersContent());
@@ -4607,9 +4544,17 @@ void ALandscapeProxy::FinishDestroy()
 #if WITH_EDITOR
 bool ALandscapeProxy::CanHaveLayersContent() const
 {
-	UWorld* World = GetWorld();
-	AWorldSettings* WorldSettings = World != nullptr ? World->GetWorldSettings() : nullptr;
-	return WorldSettings != nullptr ? WorldSettings->bEnableLandscapeLayers : false;
+	if (HasAnyFlags(RF_ClassDefaultObject | RF_ArchetypeObject))
+	{
+		return false;
+	}
+
+	if (const ALandscape* LandscapeActor = GetLandscapeActor())
+	{
+		return LandscapeActor->bCanHaveLayersContent;
+	}
+
+	return false;
 }
 
 bool ALandscapeProxy::HasLayersContent() const
@@ -5610,18 +5555,6 @@ TArray<ALandscapeBlueprintCustomBrush*> ALandscape::GetBrushesForLayer(int32 InL
 	return Brushes;
 }
 
-FChangeLandscapeLayersState ALandscape::ChangeLandscapeLayersStateEvent;
-
-void ALandscape::RegisterChangeLandscapeLayersStateDelegate()
-{
-	ALandscape::ChangeLandscapeLayersStateEvent.BindStatic(&ALandscape::ChangeLandscapeLayersState);
-}
-
-void ALandscape::UnregisterChangeLandscapeLayersStateDelegate()
-{
-	ALandscape::ChangeLandscapeLayersStateEvent.Unbind();	
-}
-
 #endif // WITH_EDITOR
 
 #undef LOCTEXT_NAMESPACE
\ No newline at end of file
diff --git a/Engine/Source/Runtime/Landscape/Public/LandscapeEdit.h b/Engine/Source/Runtime/Landscape/Public/LandscapeEdit.h
--- a/Engine/Source/Runtime/Landscape/Public/LandscapeEdit.h
+++ b/Engine/Source/Runtime/Landscape/Public/LandscapeEdit.h
@@ -50,6 +50,7 @@ public:
 	virtual ELandscapeToolTargetType::Type GetLandscapeToolTargetType() const = 0;
 	virtual const FLandscapeLayer* GetLandscapeSelectedLayer() const = 0;
 	virtual ULandscapeLayerInfoObject* GetSelectedLandscapeLayerInfo() const = 0;
+	virtual void OnCanHaveLayersContentChanged() = 0;
 };
 
 struct FLandscapeTextureDataInfo
