- Build Fix missing include

#rnx

diff --git a/Engine/Source/Editor/DetailCustomizations/Private/LandscapeUIDetails.h b/Engine/Source/Editor/DetailCustomizations/Private/LandscapeUIDetails.h
--- a/Engine/Source/Editor/DetailCustomizations/Private/LandscapeUIDetails.h
+++ b/Engine/Source/Editor/DetailCustomizations/Private/LandscapeUIDetails.h
@@ -6,6 +6,7 @@
 #include "UObject/WeakObjectPtr.h"
 #include "Types/SlateEnums.h"
 #include "IDetailCustomization.h"
+#include "Input/Reply.h"
 
 class IDetailLayoutBuilder;
 class ALandscape;
