-Landscape Layer System: Clear the undo/redo transaction buffer when enabling/disabling the Layer System.

#rnx

diff --git a/Engine/Source/Editor/DetailCustomizations/Private/LandscapeUIDetails.cpp b/Engine/Source/Editor/DetailCustomizations/Private/LandscapeUIDetails.cpp
--- a/Engine/Source/Editor/DetailCustomizations/Private/LandscapeUIDetails.cpp
+++ b/Engine/Source/Editor/DetailCustomizations/Private/LandscapeUIDetails.cpp
@@ -11,6 +11,7 @@
 #include "Widgets/Input/SButton.h"
 #include "Widgets/Text/STextBlock.h"
 #include "Misc/MessageDialog.h"
+#include "Editor.h"
 
 FLandscapeUIDetails::FLandscapeUIDetails()
 {
@@ -115,24 +116,25 @@ FReply FLandscapeUIDetails::ToggleCanHaveLayersContent(ALandscape* Landscape)
 
 		if (bHasHiddenLayers)
 		{
-			Reason = NSLOCTEXT("UnrealEd", "LandscapeDisableLayers_HiddenLayers", "Are you sure you want to disable the layers system?\n\nDoing so, will result in losing the data stored for each layers, but the current visual output will be kept. Be aware that some layers are currently hidden, continuing will result in their data being lost.");
+			Reason = NSLOCTEXT("UnrealEd", "LandscapeDisableLayers_HiddenLayers", "Are you sure you want to disable the layers system?\n\nDoing so, will result in losing the data stored for each layers, but the current visual output will be kept. Be aware that some layers are currently hidden, continuing will result in their data being lost. Undo/redo buffer will also be cleared.");
 		}
 		else
 		{
-			Reason = NSLOCTEXT("UnrealEd", "LandscapeDisableLayers", "Are you sure you want to disable the layers system?\n\nDoing so, will result in losing the data stored for each layers, but the current visual output will be kept.");
+			Reason = NSLOCTEXT("UnrealEd", "LandscapeDisableLayers", "Are you sure you want to disable the layers system?\n\nDoing so, will result in losing the data stored for each layers, but the current visual output will be kept. Undo/redo buffer will also be cleared.");
 		}
 
 		bToggled = FMessageDialog::Open(EAppMsgType::YesNo, Reason) == EAppReturnType::Yes;
 	}
 	else
 	{
-		FMessageDialog::Open(EAppMsgType::Ok, NSLOCTEXT("UnrealEd", "LandscapeEnableLayers", "Layer system as been enabled on this landscape and streaming proxies."));
-		bToggled = true;
+		
+		bToggled = FMessageDialog::Open(EAppMsgType::YesNo, NSLOCTEXT("UnrealEd", "LandscapeEnableLayers", "Are you sure you want to enable the layer system on this landscape and streaming proxies? Doing so will clear the undo/redo buffer."));
 	}
 
 	if (DetailLayoutBuilder && bToggled)
 	{
 		Landscape->ToggleCanHaveLayersContent();
+		GEditor->ResetTransaction(NSLOCTEXT("UnrealEd", "ToggleLanscapeLayers", "Toggling Landscape Layers"));
 		DetailLayoutBuilder->ForceRefreshDetails();
 	}
 
