Fix Compile Error (VS2019)

#rnx

diff --git a/Engine/Source/Editor/DetailCustomizations/Private/LandscapeUIDetails.cpp b/Engine/Source/Editor/DetailCustomizations/Private/LandscapeUIDetails.cpp
--- a/Engine/Source/Editor/DetailCustomizations/Private/LandscapeUIDetails.cpp
+++ b/Engine/Source/Editor/DetailCustomizations/Private/LandscapeUIDetails.cpp
@@ -128,7 +128,7 @@ FReply FLandscapeUIDetails::ToggleCanHaveLayersContent(ALandscape* Landscape)
 	else
 	{
 		
-		bToggled = FMessageDialog::Open(EAppMsgType::YesNo, NSLOCTEXT("UnrealEd", "LandscapeEnableLayers", "Are you sure you want to enable the layer system on this landscape and streaming proxies? Doing so will clear the undo/redo buffer."));
+		bToggled = FMessageDialog::Open(EAppMsgType::YesNo, NSLOCTEXT("UnrealEd", "LandscapeEnableLayers", "Are you sure you want to enable the layer system on this landscape and streaming proxies? Doing so will clear the undo/redo buffer.")) == EAppReturnType::Yes;
 	}
 
 	if (DetailLayoutBuilder && bToggled)
