-Landscape Layer System: Replace Button by Checkbox to enable/disable layer system

#rnx

diff --git a/Engine/Source/Editor/DetailCustomizations/Private/LandscapeUIDetails.cpp b/Engine/Source/Editor/DetailCustomizations/Private/LandscapeUIDetails.cpp
--- a/Engine/Source/Editor/DetailCustomizations/Private/LandscapeUIDetails.cpp
+++ b/Engine/Source/Editor/DetailCustomizations/Private/LandscapeUIDetails.cpp
@@ -61,22 +61,25 @@ void FLandscapeUIDetails::CustomizeDetails( IDetailLayoutBuilder& DetailBuilder
 				Category.AddCustomRow(FText(), true)
 				.NameContent()
 				[
-					PropertyHandle->CreatePropertyNameWidget(FText::FromString("Layer System"))
+					PropertyHandle->CreatePropertyNameWidget(NSLOCTEXT("UnrealEd", "LandscapeToggleLayerName", "Enable Layer System"))
 				]
 				.ValueContent()
-				[
-					SNew(SButton)
-					.OnClicked_Lambda([this, Landscape]() -> FReply
-					{
-						return ToggleCanHaveLayersContent(Landscape);
-					})
 					[
-						SNew(STextBlock)
-						.Text_Lambda([Landscape]()
-						{
-							return Landscape->bCanHaveLayersContent ? NSLOCTEXT("UnrealEd", "LandscapeDisableLayers", "Disable") : NSLOCTEXT("UnrealEd", "LandscapeEnableLayers", "Enable");	
+						SNew(SCheckBox)
+						.ToolTipText(NSLOCTEXT("UnrealEd", "LandscapeToggleLayerToolTip", "Toggle whether or not to support layers on this Landscape and its streaming proxies. Toggling this will clear the undo stack."))
+						.Type(ESlateCheckBoxType::CheckBox)
+						.IsChecked_Lambda([=]() 
+						{ 
+							return Landscape->CanHaveLayersContent() ? ECheckBoxState::Checked : ECheckBoxState::Unchecked;
+						})
+						.OnCheckStateChanged_Lambda([=](ECheckBoxState NewState) 
+						{ 
+							bool bChecked = (NewState == ECheckBoxState::Checked);
+							if (Landscape->CanHaveLayersContent() != bChecked)
+							{
+								ToggleCanHaveLayersContent(Landscape);
+							}
 						})
-					]
 				];
 			}
 			
@@ -92,7 +95,7 @@ void FLandscapeUIDetails::OnEditorExperimentalSettingsChanged(FName PropertyName
 	}
 }
 
-FReply FLandscapeUIDetails::ToggleCanHaveLayersContent(ALandscape* Landscape)
+void FLandscapeUIDetails::ToggleCanHaveLayersContent(ALandscape* Landscape)
 {
 	bool bToggled = false;
 
@@ -110,8 +113,7 @@ FReply FLandscapeUIDetails::ToggleCanHaveLayersContent(ALandscape* Landscape)
 				break;
 			}
 		}
-
-		
+				
 		FText Reason;
 
 		if (bHasHiddenLayers)
@@ -131,12 +133,12 @@ FReply FLandscapeUIDetails::ToggleCanHaveLayersContent(ALandscape* Landscape)
 		bToggled = FMessageDialog::Open(EAppMsgType::YesNo, NSLOCTEXT("UnrealEd", "LandscapeEnableLayers", "Are you sure you want to enable the layer system on this landscape and streaming proxies? Doing so will clear the undo/redo buffer.")) == EAppReturnType::Yes;
 	}
 
-	if (DetailLayoutBuilder && bToggled)
+	if (bToggled)
 	{
 		Landscape->ToggleCanHaveLayersContent();
-		GEditor->ResetTransaction(NSLOCTEXT("UnrealEd", "ToggleLanscapeLayers", "Toggling Landscape Layers"));
-		DetailLayoutBuilder->ForceRefreshDetails();
+		if (GEditor)
+		{
+			GEditor->ResetTransaction(NSLOCTEXT("UnrealEd", "ToggleLanscapeLayers", "Toggling Landscape Layers"));
+		}
 	}
-
-	return FReply::Handled();
 }
\ No newline at end of file
diff --git a/Engine/Source/Editor/DetailCustomizations/Private/LandscapeUIDetails.h b/Engine/Source/Editor/DetailCustomizations/Private/LandscapeUIDetails.h
--- a/Engine/Source/Editor/DetailCustomizations/Private/LandscapeUIDetails.h
+++ b/Engine/Source/Editor/DetailCustomizations/Private/LandscapeUIDetails.h
@@ -25,7 +25,7 @@ public:
 
 private:
 	void OnEditorExperimentalSettingsChanged(FName PropertyName);
-	FReply ToggleCanHaveLayersContent(ALandscape* Landscape);
+	void ToggleCanHaveLayersContent(ALandscape* Landscape);
 
 	/** Use MakeInstance to create an instance of this class */
 	FLandscapeUIDetails();
