PR #5713: Fixed editor crash from GetDefault<UNavigationSystemV1>(GEngine->NavigationSystemClass) when GEngine->NavigationSystemClass is neither null nor UNavigationSystemV1, but rather its parent class UNavigationSystemBase. (Contributed by [REDACTED_PERSON])

diff --git a/Engine/Source/Editor/DetailCustomizations/Private/NavAgentSelectorCustomization.cpp b/Engine/Source/Editor/DetailCustomizations/Private/NavAgentSelectorCustomization.cpp
--- a/Engine/Source/Editor/DetailCustomizations/Private/NavAgentSelectorCustomization.cpp
+++ b/Engine/Source/Editor/DetailCustomizations/Private/NavAgentSelectorCustomization.cpp
@@ -39,7 +39,7 @@ void FNavAgentSelectorCustomization::CustomizeChildren(TSharedRef<class IPropert
 	StructPropertyHandle->GetNumChildren(NumChildren);
 
 	FString AgentPrefix("bSupportsAgent");
-	const UNavigationSystemV1* NavSysCDO = (*GEngine->NavigationSystemClass != nullptr)
+	const UNavigationSystemV1* NavSysCDO = (*GEngine->NavigationSystemClass != nullptr && GEngine->NavigationSystemClass->IsChildOf(UNavigationSystemV1::StaticClass()))
 		? GetDefault<UNavigationSystemV1>(GEngine->NavigationSystemClass)
 		: GetDefault<UNavigationSystemV1>();
 
@@ -85,7 +85,7 @@ void FNavAgentSelectorCustomization::CustomizeChildren(TSharedRef<class IPropert
 
 void FNavAgentSelectorCustomization::OnAgentStateChanged()
 {
-	const UNavigationSystemV1* NavSysCDO = (*GEngine->NavigationSystemClass != nullptr)
+	const UNavigationSystemV1* NavSysCDO = (*GEngine->NavigationSystemClass != nullptr && GEngine->NavigationSystemClass->IsChildOf(UNavigationSystemV1::StaticClass()))
 		? GetDefault<UNavigationSystemV1>(GEngine->NavigationSystemClass)
 		: GetDefault<UNavigationSystemV1>();
 
