Clean up logic for displaying available pins when multi selecting make, break, or get class default nodes

diff --git a/Engine/Source/Editor/DetailCustomizations/Private/SkeletalControlNodeDetails.cpp b/Engine/Source/Editor/DetailCustomizations/Private/SkeletalControlNodeDetails.cpp
--- a/Engine/Source/Editor/DetailCustomizations/Private/SkeletalControlNodeDetails.cpp
+++ b/Engine/Source/Editor/DetailCustomizations/Private/SkeletalControlNodeDetails.cpp
@@ -34,79 +34,105 @@ void FSkeletalControlNodeDetails::CustomizeDetails(class IDetailLayoutBuilder& D
 	TSharedRef<IPropertyHandle> AvailablePins = DetailBuilder.GetProperty("ShowPinForProperties");
 	ArrayProperty = AvailablePins->AsArray();
 
+	bool bShowAvailablePins = true;
+
 	const TArray< TWeakObjectPtr<UObject> >& SelectedObjects = DetailBuilder.GetSelectedObjects();
-	for (const TWeakObjectPtr<UObject>& CurrentObject : SelectedObjects)
+	if (SelectedObjects.Num() == 1)
+	{
+		UObject* CurObj = SelectedObjects[0].Get();
+		HideUnconnectedPinsNode = ((CurObj && (CurObj->IsA<UK2Node_BreakStruct>() || CurObj->IsA<UK2Node_GetClassDefaults>())) ? static_cast<UK2Node*>(CurObj) : nullptr);
+	}
+	else
 	{
-		if (Cast<UK2Node_BreakStruct>(CurrentObject.Get()) || Cast<UK2Node_GetClassDefaults>(CurrentObject.Get()))
+		UStruct* CommonStruct = nullptr;
+		for (const TWeakObjectPtr<UObject>& CurrentObject : SelectedObjects)
 		{
-			if (HideUnconnectedPinsNode.IsValid())
+			UObject* CurObj = CurrentObject.Get();
+			if (UK2Node_StructOperation* StructOpNode = Cast<UK2Node_StructOperation>(CurObj))
 			{
-				// Have more than one break struct node, don't cache so we don't
-				// create the hide unconnected pins UI
-				HideUnconnectedPinsNode = nullptr;
-				break;
+				if (CommonStruct && CommonStruct != StructOpNode->StructType)
+				{
+					bShowAvailablePins = false;
+					break;
+				}
+				CommonStruct = StructOpNode->StructType;
+			}
+			else if (UK2Node_GetClassDefaults* CDONode = Cast<UK2Node_GetClassDefaults>(CurObj))
+			{
+				if (CommonStruct && CommonStruct != CDONode->GetInputClass())
+				{
+					bShowAvailablePins = false;
+					break;
+				}
+				CommonStruct = CDONode->GetInputClass();
 			}
-			HideUnconnectedPinsNode = Cast<UK2Node>(CurrentObject.Get());
 		}
 	}
 
-	TSet<FName> UniqueCategoryNames;
+	if (bShowAvailablePins)
 	{
-		int32 NumElements = 0;
+		TSet<FName> UniqueCategoryNames;
 		{
-			uint32 UnNumElements = 0;
-			if (ArrayProperty.IsValid() && (FPropertyAccess::Success == ArrayProperty->GetNumElements(UnNumElements)))
+			int32 NumElements = 0;
 			{
-				NumElements = static_cast<int32>(UnNumElements);
+				uint32 UnNumElements = 0;
+				if (ArrayProperty.IsValid() && (FPropertyAccess::Success == ArrayProperty->GetNumElements(UnNumElements)))
+				{
+					NumElements = static_cast<int32>(UnNumElements);
+				}
 			}
-		}
-		for (int32 Index = 0; Index < NumElements; ++Index)
-		{
-			TSharedRef<IPropertyHandle> StructPropHandle = ArrayProperty->GetElement(Index);
-			TSharedPtr<IPropertyHandle> CategoryPropHandle = StructPropHandle->GetChildHandle(GET_MEMBER_NAME_CHECKED(FOptionalPinFromProperty, CategoryName));
-			check(CategoryPropHandle.IsValid());
-			FName CategoryNameValue;
-			const FPropertyAccess::Result Result = CategoryPropHandle->GetValue(CategoryNameValue);
-			if (ensure(FPropertyAccess::Success == Result))
+			for (int32 Index = 0; Index < NumElements; ++Index)
 			{
-				UniqueCategoryNames.Add(CategoryNameValue);
+				TSharedRef<IPropertyHandle> StructPropHandle = ArrayProperty->GetElement(Index);
+				TSharedPtr<IPropertyHandle> CategoryPropHandle = StructPropHandle->GetChildHandle(GET_MEMBER_NAME_CHECKED(FOptionalPinFromProperty, CategoryName));
+				check(CategoryPropHandle.IsValid());
+				FName CategoryNameValue;
+				const FPropertyAccess::Result Result = CategoryPropHandle->GetValue(CategoryNameValue);
+				if (ensure(FPropertyAccess::Success == Result))
+				{
+					UniqueCategoryNames.Add(CategoryNameValue);
+				}
 			}
 		}
-	}
 
-	//@TODO: Shouldn't show this if the available pins array is empty!
-	const bool bGenerateHeader = true;
-	const bool bDisplayResetToDefault = false;
-	const bool bDisplayElementNum = false;
-	const bool bForAdvanced = false;
-	for (const FName& CategoryName : UniqueCategoryNames)
-	{
-		//@TODO: Pay attention to category filtering here
+		//@TODO: Shouldn't show this if the available pins array is empty!
+		const bool bGenerateHeader = true;
+		const bool bDisplayResetToDefault = false;
+		const bool bDisplayElementNum = false;
+		const bool bForAdvanced = false;
+		for (const FName& CategoryName : UniqueCategoryNames)
+		{
+			//@TODO: Pay attention to category filtering here
 		
 		
-		TSharedRef<FDetailArrayBuilder> AvailablePinsBuilder = MakeShareable(new FDetailArrayBuilder(AvailablePins, bGenerateHeader, bDisplayResetToDefault, bDisplayElementNum));
-		AvailablePinsBuilder->OnGenerateArrayElementWidget(FOnGenerateArrayElementWidget::CreateSP(this, &FSkeletalControlNodeDetails::OnGenerateElementForPropertyPin, CategoryName));
-		AvailablePinsBuilder->SetDisplayName((CategoryName == NAME_None) ? LOCTEXT("DefaultCategory", "Default Category") : FText::FromName(CategoryName));
-		DetailCategory->AddCustomBuilder(AvailablePinsBuilder, bForAdvanced);
-	}
+			TSharedRef<FDetailArrayBuilder> AvailablePinsBuilder = MakeShareable(new FDetailArrayBuilder(AvailablePins, bGenerateHeader, bDisplayResetToDefault, bDisplayElementNum));
+			AvailablePinsBuilder->OnGenerateArrayElementWidget(FOnGenerateArrayElementWidget::CreateSP(this, &FSkeletalControlNodeDetails::OnGenerateElementForPropertyPin, CategoryName));
+			AvailablePinsBuilder->SetDisplayName((CategoryName == NAME_None) ? LOCTEXT("DefaultCategory", "Default Category") : FText::FromName(CategoryName));
+			DetailCategory->AddCustomBuilder(AvailablePinsBuilder, bForAdvanced);
+		}
 
-	// Add the action buttons
-	if(HideUnconnectedPinsNode.IsValid())
-	{
-		FDetailWidgetRow& GroupActionsRow = DetailCategory->AddCustomRow(LOCTEXT("GroupActionsSearchText", "Split Sort"))
-		.ValueContent()
-		.HAlign(HAlign_Left)
-		.MaxDesiredWidth(250.f)
-		[
-			SNew(SButton)
-			.OnClicked(this, &FSkeletalControlNodeDetails::HideAllUnconnectedPins)
-			.ToolTipText(LOCTEXT("HideAllUnconnectedPinsTooltip", "All unconnected pins get hidden (removed from the graph node)"))
+		// Add the action buttons
+		if(HideUnconnectedPinsNode.IsValid())
+		{
+			FDetailWidgetRow& GroupActionsRow = DetailCategory->AddCustomRow(LOCTEXT("GroupActionsSearchText", "Split Sort"))
+			.ValueContent()
+			.HAlign(HAlign_Left)
+			.MaxDesiredWidth(250.f)
 			[
-				SNew(STextBlock)
-				.Text(LOCTEXT("HideAllUnconnectedPins", "Hide Unconnected Pins"))
-				.Font(DetailBuilder.GetDetailFont())
-			]
-		];
+				SNew(SButton)
+				.OnClicked(this, &FSkeletalControlNodeDetails::HideAllUnconnectedPins)
+				.ToolTipText(LOCTEXT("HideAllUnconnectedPinsTooltip", "All unconnected pins get hidden (removed from the graph node)"))
+				[
+					SNew(STextBlock)
+					.Text(LOCTEXT("HideAllUnconnectedPins", "Hide Unconnected Pins"))
+					.Font(DetailBuilder.GetDetailFont())
+				]
+			];
+		}
+	}
+	else
+	{
+		DetailBuilder.HideProperty(AvailablePins);
 	}
 }
 
