Fix mirror scale option in details panel

diff --git a/Engine/Source/Editor/DetailCustomizations/Private/ComponentTransformDetails.cpp b/Engine/Source/Editor/DetailCustomizations/Private/ComponentTransformDetails.cpp
--- a/Engine/Source/Editor/DetailCustomizations/Private/ComponentTransformDetails.cpp
+++ b/Engine/Source/Editor/DetailCustomizations/Private/ComponentTransformDetails.cpp
@@ -856,14 +856,14 @@ void FComponentTransformDetails::OnYScaleMirrored()
 {
 	FSlateApplication::Get().ClearKeyboardFocus(EFocusCause::Mouse);
 	FScopedTransaction Transaction(LOCTEXT("MirrorActorScaleY", "Mirror actor scale Y"));
-	OnSetTransform(ETransformField::Scale, EAxisList::X, FVector(1.0f), true, true);
+	OnSetTransform(ETransformField::Scale, EAxisList::Y, FVector(1.0f), true, true);
 }
 
 void FComponentTransformDetails::OnZScaleMirrored()
 {
 	FSlateApplication::Get().ClearKeyboardFocus(EFocusCause::Mouse);
 	FScopedTransaction Transaction(LOCTEXT("MirrorActorScaleZ", "Mirror actor scale Z"));
-	OnSetTransform(ETransformField::Scale, EAxisList::X, FVector(1.0f), true, true);
+	OnSetTransform(ETransformField::Scale, EAxisList::Z, FVector(1.0f), true, true);
 }
 
 void FComponentTransformDetails::CacheTransform()
@@ -1060,7 +1060,7 @@ void FComponentTransformDetails::OnSetTransform(ETransformField::Type TransformF
 				// Set the incoming value
 				if (bMirror)
 				{
-					NewComponentValue = -OldComponentValue;
+					NewComponentValue = GetAxisFilteredVector(Axis, -OldComponentValue, OldComponentValue);
 				}
 				else
 				{
