Landscape Layers no longer Experimental Feature

#rnx

diff --git a/Engine/Source/Editor/DetailCustomizations/Private/DetailCustomizations.cpp b/Engine/Source/Editor/DetailCustomizations/Private/DetailCustomizations.cpp
--- a/Engine/Source/Editor/DetailCustomizations/Private/DetailCustomizations.cpp
+++ b/Engine/Source/Editor/DetailCustomizations/Private/DetailCustomizations.cpp
@@ -131,7 +131,6 @@
 #include "MotionControllerPinFactory.h"
 #include "LandscapeUIDetails.h"
 #include "LandscapeProxyUIDetails.h"
-#include "LandscapeSettingsDetails.h"
 #include "PerPlatformPropertyCustomization.h"
 #include "SkeletalMeshReductionSettingsDetails.h"
 #include "SkeletalMeshLODSettingsDetails.h"
@@ -387,7 +386,6 @@ void FDetailCustomizationsModule::RegisterObjectCustomizations()
 
 	RegisterCustomClassLayout("Landscape", FOnGetDetailCustomizationInstance::CreateStatic(&FLandscapeUIDetails::MakeInstance));
 	RegisterCustomClassLayout("LandscapeProxy", FOnGetDetailCustomizationInstance::CreateStatic(&FLandscapeProxyUIDetails::MakeInstance));
-	RegisterCustomClassLayout("LandscapeSettings", FOnGetDetailCustomizationInstance::CreateStatic(&FLandscapeSettingsDetails::MakeInstance));
 }
 
 
diff --git a/Engine/Source/Editor/DetailCustomizations/Private/LandscapeSettingsDetails.cpp b/Engine/Source/Editor/DetailCustomizations/Private/LandscapeSettingsDetails.cpp
deleted file mode 100644
diff --git a/Engine/Source/Editor/DetailCustomizations/Private/LandscapeSettingsDetails.h b/Engine/Source/Editor/DetailCustomizations/Private/LandscapeSettingsDetails.h
deleted file mode 100644
diff --git a/Engine/Source/Editor/DetailCustomizations/Private/LandscapeUIDetails.cpp b/Engine/Source/Editor/DetailCustomizations/Private/LandscapeUIDetails.cpp
--- a/Engine/Source/Editor/DetailCustomizations/Private/LandscapeUIDetails.cpp
+++ b/Engine/Source/Editor/DetailCustomizations/Private/LandscapeUIDetails.cpp
@@ -23,10 +23,6 @@ FLandscapeUIDetails::FLandscapeUIDetails()
 
 FLandscapeUIDetails::~FLandscapeUIDetails()
 {
-	if (DetailLayoutBuilder)
-	{
-		GetMutableDefault<UEditorExperimentalSettings>()->OnSettingChanged().RemoveAll(this);
-	}
 }
 
 TSharedRef<IDetailCustomization> FLandscapeUIDetails::MakeInstance()
@@ -52,50 +48,33 @@ void FLandscapeUIDetails::CustomizeDetails( IDetailLayoutBuilder& DetailBuilder
 			TSharedRef<IPropertyHandle> ComponentScreenSizeToUseSubSectionsProp = DetailBuilder.GetProperty(GET_MEMBER_NAME_CHECKED(ALandscapeProxy, ComponentScreenSizeToUseSubSections));
 			DetailBuilder.HideProperty(ComponentScreenSizeToUseSubSectionsProp);
 		}
-
-		if (!DetailLayoutBuilder)
-		{
-			DetailLayoutBuilder = &DetailBuilder;
-			GetMutableDefault<UEditorExperimentalSettings>()->OnSettingChanged().AddSP(this, &FLandscapeUIDetails::OnEditorExperimentalSettingsChanged);
-		}
-			   
+							   
 		TSharedRef<IPropertyHandle> PropertyHandle = DetailBuilder.GetProperty(GET_MEMBER_NAME_CHECKED(ALandscape, bCanHaveLayersContent));
 		DetailBuilder.HideProperty(PropertyHandle);
-		if (GetMutableDefault<UEditorExperimentalSettings>()->bLandscapeLayerSystem)
-		{
-			const FText DisplayAndFilterText(LOCTEXT("LandscapeToggleLayerName", "Enable Layer System"));
-			DetailBuilder.AddCustomRowToCategory(PropertyHandle, DisplayAndFilterText)
-			.NameContent()
-			[
-				PropertyHandle->CreatePropertyNameWidget(DisplayAndFilterText)
-			]
-			.ValueContent()
-			[
-				SNew(SCheckBox)
-				.ToolTipText(LOCTEXT("LandscapeToggleLayerToolTip", "Toggle whether or not to support layers on this Landscape and its streaming proxies. Toggling this will clear the undo stack."))
-				.Type(ESlateCheckBoxType::CheckBox)
-				.IsChecked_Lambda([=]()
+		const FText DisplayAndFilterText(LOCTEXT("LandscapeToggleLayerName", "Enable Layer System"));
+		DetailBuilder.AddCustomRowToCategory(PropertyHandle, DisplayAndFilterText)
+		.NameContent()
+		[
+			PropertyHandle->CreatePropertyNameWidget(DisplayAndFilterText)
+		]
+		.ValueContent()
+		[
+			SNew(SCheckBox)
+			.ToolTipText(LOCTEXT("LandscapeToggleLayerToolTip", "Toggle whether or not to support layers on this Landscape and its streaming proxies. Toggling this will clear the undo stack."))
+			.Type(ESlateCheckBoxType::CheckBox)
+			.IsChecked_Lambda([=]()
+			{
+				return Landscape->CanHaveLayersContent() ? ECheckBoxState::Checked : ECheckBoxState::Unchecked;
+			})
+			.OnCheckStateChanged_Lambda([=](ECheckBoxState NewState)
+			{
+				bool bChecked = (NewState == ECheckBoxState::Checked);
+				if (Landscape->CanHaveLayersContent() != bChecked)
 				{
-					return Landscape->CanHaveLayersContent() ? ECheckBoxState::Checked : ECheckBoxState::Unchecked;
-			    })
-			    .OnCheckStateChanged_Lambda([=](ECheckBoxState NewState)
-			    {
-			        bool bChecked = (NewState == ECheckBoxState::Checked);
-				    if (Landscape->CanHaveLayersContent() != bChecked)
-				    {
-					    ToggleCanHaveLayersContent(Landscape);
-				    }
-			    })
-			];
-		}
-	}
-}
-
-void FLandscapeUIDetails::OnEditorExperimentalSettingsChanged(FName PropertyName)
-{
-	if (DetailLayoutBuilder && PropertyName == GET_MEMBER_NAME_CHECKED(UEditorExperimentalSettings, bLandscapeLayerSystem))
-	{
-		DetailLayoutBuilder->ForceRefreshDetails();
+					ToggleCanHaveLayersContent(Landscape);
+				}
+			})
+		];
 	}
 }
 
diff --git a/Engine/Source/Editor/DetailCustomizations/Private/LandscapeUIDetails.h b/Engine/Source/Editor/DetailCustomizations/Private/LandscapeUIDetails.h
--- a/Engine/Source/Editor/DetailCustomizations/Private/LandscapeUIDetails.h
+++ b/Engine/Source/Editor/DetailCustomizations/Private/LandscapeUIDetails.h
@@ -24,11 +24,8 @@ public:
 	/** End IDetailCustomization interface */
 
 private:
-	void OnEditorExperimentalSettingsChanged(FName PropertyName);
 	void ToggleCanHaveLayersContent(ALandscape* Landscape);
 
 	/** Use MakeInstance to create an instance of this class */
 	FLandscapeUIDetails();
-
-	IDetailLayoutBuilder* DetailLayoutBuilder = nullptr;
 };
diff --git a/Engine/Source/Editor/DetailCustomizations/Private/WorldSettingsDetails.cpp b/Engine/Source/Editor/DetailCustomizations/Private/WorldSettingsDetails.cpp
--- a/Engine/Source/Editor/DetailCustomizations/Private/WorldSettingsDetails.cpp
+++ b/Engine/Source/Editor/DetailCustomizations/Private/WorldSettingsDetails.cpp
@@ -23,15 +23,6 @@
 
 FWorldSettingsDetails::~FWorldSettingsDetails()
 {
-	GetMutableDefault<UEditorExperimentalSettings>()->OnSettingChanged().Remove(ExperimentalDelegateHandle);
-}
-
-void FWorldSettingsDetails::OnEditorExperimentalSettingsChanged(FName InPropertyName)
-{
-	if (DetailLayoutBuilder && InPropertyName == GET_MEMBER_NAME_CHECKED(UEditorExperimentalSettings, bLandscapeLayerSystem))
-	{
-		DetailLayoutBuilder->ForceRefreshDetails();
-	}
 }
 
 /* IDetailCustomization overrides
@@ -39,13 +30,6 @@ void FWorldSettingsDetails::OnEditorExperimentalSettingsChanged(FName InProperty
 
 void FWorldSettingsDetails::CustomizeDetails( IDetailLayoutBuilder& DetailBuilder )
 {
-	DetailLayoutBuilder = &DetailBuilder;
-
-	if (!ExperimentalDelegateHandle.IsValid())
-	{
-		ExperimentalDelegateHandle = GetMutableDefault<UEditorExperimentalSettings>()->OnSettingChanged().AddSP(this, &FWorldSettingsDetails::OnEditorExperimentalSettingsChanged);
-	}
-
 	IDetailCategoryBuilder& Category = DetailBuilder.EditCategory("GameMode");
 	CustomizeGameInfoProperty("DefaultGameMode", DetailBuilder, Category);
 
diff --git a/Engine/Source/Editor/DetailCustomizations/Private/WorldSettingsDetails.h b/Engine/Source/Editor/DetailCustomizations/Private/WorldSettingsDetails.h
--- a/Engine/Source/Editor/DetailCustomizations/Private/WorldSettingsDetails.h
+++ b/Engine/Source/Editor/DetailCustomizations/Private/WorldSettingsDetails.h
@@ -27,8 +27,6 @@ public:
 
 	virtual ~FWorldSettingsDetails();
 
-	void OnEditorExperimentalSettingsChanged(FName InPropertyName);
-
 	// IDetailCustomization interface
 
 	virtual void CustomizeDetails( IDetailLayoutBuilder& DetailBuilder ) override;
@@ -74,8 +72,6 @@ private:
 
 	/** Helper class to customizer GameMode property */
 	TSharedPtr<FGameModeInfoCustomizer>	GameInfoModeCustomizer;
-	FDelegateHandle ExperimentalDelegateHandle;
-	IDetailLayoutBuilder* DetailLayoutBuilder = nullptr;
 };
 
 
diff --git a/Engine/Source/Editor/LandscapeEditor/Private/LandscapeEditorDetailCustomization_NewLandscape.cpp b/Engine/Source/Editor/LandscapeEditor/Private/LandscapeEditorDetailCustomization_NewLandscape.cpp
--- a/Engine/Source/Editor/LandscapeEditor/Private/LandscapeEditorDetailCustomization_NewLandscape.cpp
+++ b/Engine/Source/Editor/LandscapeEditor/Private/LandscapeEditorDetailCustomization_NewLandscape.cpp
@@ -89,15 +89,8 @@ void FLandscapeEditorDetailCustomization_NewLandscape::CustomizeDetails(IDetailL
 	];
 
 	TSharedRef<IPropertyHandle> PropertyHandle_CanHaveLayersContent = DetailBuilder.GetProperty(GET_MEMBER_NAME_CHECKED(ULandscapeEditorObject, bCanHaveLayersContent));
-	if (GetMutableDefault<UEditorExperimentalSettings>()->bLandscapeLayerSystem)
-	{
-		NewLandscapeCategory.AddProperty(PropertyHandle_CanHaveLayersContent);
-	}
-	else
-	{
-		DetailBuilder.HideProperty(PropertyHandle_CanHaveLayersContent);
-	}
-
+	NewLandscapeCategory.AddProperty(PropertyHandle_CanHaveLayersContent);
+	
 	TSharedRef<IPropertyHandle> PropertyHandle_HeightmapFilename = DetailBuilder.GetProperty(GET_MEMBER_NAME_CHECKED(ULandscapeEditorObject, ImportLandscape_HeightmapFilename));
 	TSharedRef<IPropertyHandle> PropertyHandle_HeightmapImportResult = DetailBuilder.GetProperty(GET_MEMBER_NAME_CHECKED(ULandscapeEditorObject, ImportLandscape_HeightmapImportResult));
 	TSharedRef<IPropertyHandle> PropertyHandle_HeightmapErrorMessage = DetailBuilder.GetProperty(GET_MEMBER_NAME_CHECKED(ULandscapeEditorObject, ImportLandscape_HeightmapErrorMessage));
diff --git a/Engine/Source/Editor/UnrealEd/Classes/Settings/EditorExperimentalSettings.h b/Engine/Source/Editor/UnrealEd/Classes/Settings/EditorExperimentalSettings.h
--- a/Engine/Source/Editor/UnrealEd/Classes/Settings/EditorExperimentalSettings.h
+++ b/Engine/Source/Editor/UnrealEd/Classes/Settings/EditorExperimentalSettings.h
@@ -23,11 +23,7 @@ public:
 	/** Allows usage of the procedural foliage system */
 	UPROPERTY(EditAnywhere, config, Category = Foliage, meta = (DisplayName = "Procedural Foliage"))
 	bool bProceduralFoliage;
-
-	/** Allows usage of the procedural landscape system*/
-	UPROPERTY(EditAnywhere, config, Category = Landscape, meta = (DisplayName = "Landscape Layer System"))
-	bool bLandscapeLayerSystem;
-
+		
 	/** Allows usage of the Localization Dashboard */
 	UPROPERTY(EditAnywhere, config, Category = Tools, meta = (DisplayName = "Localization Dashboard"))
 	bool bEnableLocalizationDashboard;
diff --git a/Engine/Source/Editor/UnrealEd/Private/Settings/SettingsClasses.cpp b/Engine/Source/Editor/UnrealEd/Private/Settings/SettingsClasses.cpp
--- a/Engine/Source/Editor/UnrealEd/Private/Settings/SettingsClasses.cpp
+++ b/Engine/Source/Editor/UnrealEd/Private/Settings/SettingsClasses.cpp
@@ -133,7 +133,6 @@ USkeletalMeshEditorSettings::USkeletalMeshEditorSettings(const FObjectInitialize
 
 UEditorExperimentalSettings::UEditorExperimentalSettings( const FObjectInitializer& ObjectInitializer )
 	: Super(ObjectInitializer)
-	, bLandscapeLayerSystem(false)
 	, bEnableLocalizationDashboard(true)
 	, bUseOpenCLForConvexHullDecomp(false)
 	, bAllowPotentiallyUnsafePropertyEditing(false)
