- Fix for FilePathCustomization with RelativeToGameDir metadata not working when its content is already in the content dir.

diff --git a/Engine/Source/Editor/DetailCustomizations/Private/FilePathStructCustomization.cpp b/Engine/Source/Editor/DetailCustomizations/Private/FilePathStructCustomization.cpp
--- a/Engine/Source/Editor/DetailCustomizations/Private/FilePathStructCustomization.cpp
+++ b/Engine/Source/Editor/DetailCustomizations/Private/FilePathStructCustomization.cpp
@@ -95,16 +95,22 @@ void FFilePathStructCustomization::HandleFilePathPickerPathPicked( const FString
 	}
 	else if (bRelativeToGameDir)
 	{
-		//If the path is part of the project directory, make it relative to it or keep the absolute path.
-		const FString AbsolutePickedPath = FPaths::ConvertRelativePathToFull(PickedPath);
-		const FString AbsoluteProjectDir = FPaths::ConvertRelativePathToFull(FPaths::ProjectDir());
-		if (AbsolutePickedPath.StartsWith(AbsoluteProjectDir))
+		//If path is already relative to Content dir, nothing to do.
+		const FString ContentDir = TEXT("Content/");
+		if (!PickedPath.IsEmpty() && !PickedPath.StartsWith(ContentDir))
 		{
-			FinalPath = AbsolutePickedPath.RightChop(AbsoluteProjectDir.Len());
-		}
-		else
-		{
-			FinalPath = AbsolutePickedPath;
+			//If the path is part of the project directory, make it relative to it or keep the absolute path.
+			const FString AbsolutePickedPath = FPaths::ConvertRelativePathToFull(PickedPath);
+			const FString AbsoluteProjectDir = FPaths::ConvertRelativePathToFull(FPaths::ProjectDir());
+			const FString AbsoluteProjectWithContentDir = AbsoluteProjectDir + ContentDir;
+			if (AbsolutePickedPath.StartsWith(AbsoluteProjectWithContentDir))
+			{
+				FinalPath = AbsolutePickedPath.RightChop(AbsoluteProjectDir.Len());
+			}
+			else
+			{
+				FinalPath = AbsolutePickedPath;
+			}
 		}
 	}
 
