Show warnings when using a deprecated RHI to start the editor or when compiling shaders for that RHI

Arciel.Rekman, Chris.Bunner

diff --git a/Engine/Config/BaseEngine.ini b/Engine/Config/BaseEngine.ini
--- a/Engine/Config/BaseEngine.ini
+++ b/Engine/Config/BaseEngine.ini
@@ -2071,7 +2071,6 @@ bEnable2DPhysics=false
 
 [/Script/WindowsTargetPlatform.WindowsTargetSettings]
 +TargetedRHIs=PCD3D_SM5
-+TargetedRHIs=PCD3D_SM4
 MinimumOSVersion=MSOS_Vista
 
 [/Script/WindowsMixedRealityRuntimeSettings.WindowsMixedRealityRuntimeSettings]
@@ -2082,10 +2081,8 @@ MaxBitrate=4000
 ; Possible values for architecture: X86_64UnknownLinuxGnu, ArmUnknownLinuxGnueabihf, AArch64UnknownLinuxGnueabi, I686UnknownLinuxGnu
 TargetArchitecture=X86_64UnknownLinuxGnu
 
-; When neither -vulkan nor -opengl4 are passed on the commandline, the engine will default to the first targeted RHI.
+; the engine will default to the first targeted RHI.
 +TargetedRHIs=SF_VULKAN_SM5
-+TargetedRHIs=GLSL_430
-
  
 [/Script/MacTargetPlatform.MacTargetSettings]
 MaxShaderLanguageVersion=3
diff --git a/Engine/Source/Editor/DetailCustomizations/Private/LinuxTargetSettingsDetails.cpp b/Engine/Source/Editor/DetailCustomizations/Private/LinuxTargetSettingsDetails.cpp
--- a/Engine/Source/Editor/DetailCustomizations/Private/LinuxTargetSettingsDetails.cpp
+++ b/Engine/Source/Editor/DetailCustomizations/Private/LinuxTargetSettingsDetails.cpp
@@ -61,21 +61,22 @@ static FText GetFriendlyNameFromLinuxRHIName(const FString& InRHIName)
 	{
 		FriendlyRHIName = LOCTEXT("OpenGL3ES31", "OpenGL 3 (ES3.1, Experimental)");
 	}
-	else if (InRHIName == TEXT("GLSL_430"))
-	{
-		FriendlyRHIName = LOCTEXT("OpenGL4", "OpenGL 4 (SM5)");
-	}
 	else if (InRHIName == TEXT("SF_VULKAN_ES31_ANDROID") || InRHIName == TEXT("SF_VULKAN_ES31"))
 	{
-		FriendlyRHIName = LOCTEXT("Vulkan ES31", "Vulkan Mobile (ES3.1, Experimental)");
+		FriendlyRHIName = LOCTEXT("Vulkan ES31", "Vulkan Mobile (ES3.1)");
 	}
 	else if (InRHIName == TEXT("SF_VULKAN_SM4"))
 	{
-		FriendlyRHIName = LOCTEXT("VulkanSM4", "Vulkan Desktop (SM4, Experimental)");
+		FriendlyRHIName = LOCTEXT("VulkanSM4", "Vulkan Desktop (SM4)");
 	}
 	else if (InRHIName == TEXT("SF_VULKAN_SM5"))
 	{
-		FriendlyRHIName = LOCTEXT("VulkanSM5", "Vulkan Desktop (SM5, Experimental)");
+		FriendlyRHIName = LOCTEXT("VulkanSM5", "Vulkan Desktop (SM5)");
+	}
+	else if (InRHIName == TEXT("GLSL_430"))
+	{
+		// Explicitly remove these formats as they are obsolete/not quite supported; users can still target them by adding them as +TargetedRHIs in the TargetPlatform ini.
+		FriendlyRHIName = FText::GetEmpty();
 	}
 
 	return FriendlyRHIName;
diff --git a/Engine/Source/Editor/DetailCustomizations/Private/WindowsTargetSettingsDetails.cpp b/Engine/Source/Editor/DetailCustomizations/Private/WindowsTargetSettingsDetails.cpp
--- a/Engine/Source/Editor/DetailCustomizations/Private/WindowsTargetSettingsDetails.cpp
+++ b/Engine/Source/Editor/DetailCustomizations/Private/WindowsTargetSettingsDetails.cpp
@@ -56,21 +56,13 @@ static FText GetFriendlyNameFromWindowsRHIName(const FString& InRHIName)
 	{
 		FriendlyRHIName = LOCTEXT("DirectX11", "DirectX 11 & 12 (SM5)");
 	}
-	else if (InRHIName == TEXT("PCD3D_SM4"))
-	{
-		FriendlyRHIName = LOCTEXT("DirectX10", "DirectX 10 (SM4)");
-	}
 	else if (InRHIName == TEXT("PCD3D_ES31"))
 	{
 		FriendlyRHIName = LOCTEXT("DirectXES31", "DirectX Mobile Emulation (ES3.1)");
 	}
-	else if (InRHIName == TEXT("GLSL_430"))
-	{
-		FriendlyRHIName = LOCTEXT("OpenGL4", "OpenGL 4 (SM5, Experimental)");
-	}
 	else if (InRHIName == TEXT("SF_VULKAN_SM5"))
 	{
-		FriendlyRHIName = LOCTEXT("VulkanSM5", "Vulkan (SM5, Experimental)");
+		FriendlyRHIName = LOCTEXT("VulkanSM5", "Vulkan (SM5)");
 	}
 	else if (InRHIName == TEXT("GLSL_SWITCH"))
 	{
@@ -82,7 +74,8 @@ static FText GetFriendlyNameFromWindowsRHIName(const FString& InRHIName)
 	}
 	else if (InRHIName == TEXT("GLSL_150_ES2") || InRHIName == TEXT("GLSL_150_ES31") || InRHIName == TEXT("GLSL_150")
 		|| InRHIName == TEXT("SF_VULKAN_ES31_ANDROID") || InRHIName == TEXT("SF_VULKAN_ES31")
-		|| InRHIName == TEXT("SF_VULKAN_SM4") || InRHIName == TEXT("PCD3D_ES2"))
+		|| InRHIName == TEXT("SF_VULKAN_SM4") || InRHIName == TEXT("PCD3D_ES2")
+		|| InRHIName == TEXT("GLSL_430") || InRHIName == TEXT("PCD3D_SM4"))
 	{
 		// Explicitly remove these formats as they are obsolete/not quite supported; users can still target them by adding them as +TargetedRHIs in the TargetPlatform ini.
 		FriendlyRHIName = FText::GetEmpty();
diff --git a/Engine/Source/Runtime/Engine/Private/ShaderCompiler/ShaderCompiler.cpp b/Engine/Source/Runtime/Engine/Private/ShaderCompiler/ShaderCompiler.cpp
--- a/Engine/Source/Runtime/Engine/Private/ShaderCompiler/ShaderCompiler.cpp
+++ b/Engine/Source/Runtime/Engine/Private/ShaderCompiler/ShaderCompiler.cpp
@@ -4446,6 +4446,12 @@ void RecompileShadersForRemote(
 			// Only compile for the desired platform if requested
 			if (ShaderPlatform == ShaderPlatformToCompile || ShaderPlatformToCompile == SP_NumPlatforms)
 			{
+				// These platforms are deprecated and we should warn about their use
+				if (ShaderPlatform == SP_OPENGL_SM5 || ShaderPlatform == SP_PCD3D_SM4)
+				{
+					UE_LOG(LogShaders, Warning, TEXT("You are compiling shaders for a deprecated platform '%s'"), *LegacyShaderPlatformToShaderFormat(ShaderPlatform).ToString());
+				}
+
 				if (bCompileChangedShaders)
 				{
 					// Kick off global shader recompiles
diff --git a/Engine/Source/Runtime/RHI/Private/Linux/LinuxDynamicRHI.cpp b/Engine/Source/Runtime/RHI/Private/Linux/LinuxDynamicRHI.cpp
--- a/Engine/Source/Runtime/RHI/Private/Linux/LinuxDynamicRHI.cpp
+++ b/Engine/Source/Runtime/RHI/Private/Linux/LinuxDynamicRHI.cpp
@@ -76,6 +76,11 @@ FDynamicRHI* PlatformCreateDynamicRHI()
 				FApp::SetGraphicsRHI(TEXT("OpenGL"));
 				FPlatformApplicationMisc::UsingOpenGL();
 
+				if (!UE_BUILD_SHIPPING)
+				{
+					FMessageDialog::Open(EAppMsgType::Ok, NSLOCTEXT("LinuxDynamicRHI", "OpenGLDeprecated", "Warning: OpenGL is deprecated, please use Vulkan."));
+				}
+
 				FName ShaderFormatName(*TargetedShaderFormats[SfIdx]);
 				EShaderPlatform TargetedPlatform = ShaderFormatToLegacyShaderPlatform(ShaderFormatName);
 				RequestedFeatureLevel = GetMaxSupportedFeatureLevel(TargetedPlatform);
diff --git a/Engine/Source/Runtime/RHI/Private/Windows/WindowsDynamicRHI.cpp b/Engine/Source/Runtime/RHI/Private/Windows/WindowsDynamicRHI.cpp
--- a/Engine/Source/Runtime/RHI/Private/Windows/WindowsDynamicRHI.cpp
+++ b/Engine/Source/Runtime/RHI/Private/Windows/WindowsDynamicRHI.cpp
@@ -163,6 +163,12 @@ static IDynamicRHIModule* LoadDynamicRHIModule(ERHIFeatureLevel::Type& DesiredFe
 			FPlatformMisc::RequestExit(1);
 			DynamicRHIModule = NULL;
 		}
+
+		if (!UE_BUILD_SHIPPING)
+		{
+			FMessageDialog::Open(EAppMsgType::Ok, NSLOCTEXT("WindowsDynamicRHI", "OpenGLDeprecated", "Warning: OpenGL is deprecated, please use a different RHI."));
+		}
+
 		LoadedRHIModuleName = OpenGLRHIModuleName;
 	}
 	else if (bForceVulkan)
