Removing project settings for audio device selection.

diff --git a/Engine/Source/Editor/DetailCustomizations/Private/WindowsTargetSettingsDetails.cpp b/Engine/Source/Editor/DetailCustomizations/Private/WindowsTargetSettingsDetails.cpp
--- a/Engine/Source/Editor/DetailCustomizations/Private/WindowsTargetSettingsDetails.cpp
+++ b/Engine/Source/Editor/DetailCustomizations/Private/WindowsTargetSettingsDetails.cpp
@@ -287,49 +287,6 @@ void FWindowsTargetSettingsDetails::CustomizeDetails( IDetailLayoutBuilder& Deta
 		]
 	];
 
-	TSharedPtr<IPropertyHandle> AudioDevicePropertyHandle = DetailBuilder.GetProperty("AudioDevice");
-	IDetailCategoryBuilder& AudioDeviceCategory = DetailBuilder.EditCategory("Audio");
-	IDetailPropertyRow& AudioDevicePropertyRow = AudioDeviceCategory.AddProperty(AudioDevicePropertyHandle);
-
-	AudioDevicePropertyRow.CustomWidget()
-	.NameContent()
-	[
-		AudioDevicePropertyHandle->CreatePropertyNameWidget()
-	]
-	.ValueContent()
-	.MaxDesiredWidth(500.0f)
-	.MinDesiredWidth(100.0f)
-	[
-		SNew(SHorizontalBox)
-
-		+ SHorizontalBox::Slot()
-		.FillWidth(1.0f)
-		[
-			SNew(SEditableTextBox)
-			.ForegroundColor(this, &FWindowsTargetSettingsDetails::HandleAudioDeviceBoxForegroundColor, AudioDevicePropertyHandle)
-			.OnTextChanged(this, &FWindowsTargetSettingsDetails::HandleAudioDeviceTextBoxTextChanged, AudioDevicePropertyHandle)
-			.OnTextCommitted(this, &FWindowsTargetSettingsDetails::HandleAudioDeviceTextBoxTextComitted, AudioDevicePropertyHandle)
-			.Text(this, &FWindowsTargetSettingsDetails::HandleAudioDeviceTextBoxText, AudioDevicePropertyHandle)
-			.ToolTipText(AudioDevicePropertyHandle->GetToolTipText())
-		]
-
-		+ SHorizontalBox::Slot()
-			.AutoWidth()
-			.Padding(2.0f, 0.0f, 0.0f, 0.0f)
-			[
-				SNew(SComboButton)
-				.ButtonContent()
-				[
-					SNullWidget::NullWidget
-				]
-				.ContentPadding(FMargin(6.0f, 1.0f))
-				.MenuContent()
-				[
-					MakeAudioDeviceMenu(AudioDevicePropertyHandle)
-				]
-				.ToolTipText(LOCTEXT("AudioDevicesButtonToolTip", "Pick from the list of available audio devices"))
-			]
-	];
 
 	AudioPluginWidgetManager.BuildAudioCategory(DetailBuilder, EAudioPlatform::Windows);
 
@@ -377,11 +334,6 @@ bool FWindowsTargetSettingsDetails::HandlePostExternalIconCopy(const FString& In
 	return true;
 }
 
-void FWindowsTargetSettingsDetails::HandleAudioDeviceSelected(FString AudioDeviceName, TSharedPtr<IPropertyHandle> PropertyHandle)
-{
-	PropertyHandle->SetValue(AudioDeviceName);
-}
-
 void FWindowsTargetSettingsDetails::HandleAudioStreamCachingToggled(ECheckBoxState EnableStreamCaching, TSharedPtr<IPropertyHandle> PropertyHandle)
 {
 	PropertyHandle->SetValue(EnableStreamCaching == ECheckBoxState::Checked);
@@ -406,110 +358,4 @@ ECheckBoxState FWindowsTargetSettingsDetails::GetAudioStreamCachingToggled(TShar
 	}
 }
 
-FSlateColor FWindowsTargetSettingsDetails::HandleAudioDeviceBoxForegroundColor(TSharedPtr<IPropertyHandle> PropertyHandle) const
-{
-	FString Value;
-
-	if (PropertyHandle->GetValue(Value) == FPropertyAccess::Success)
-	{
-		if (Value.IsEmpty() || IsValidAudioDeviceName(Value))
-		{
-			static const FName InvertedForegroundName("InvertedForeground");
-
-			// Return a valid slate color for a valid audio device
-			return FEditorStyle::GetSlateColor(InvertedForegroundName);
-		}
-	}
-
-	// Return Red, which means its an invalid audio device
-	return FLinearColor::Red;
-}
-
-FText FWindowsTargetSettingsDetails::HandleAudioDeviceTextBoxText(TSharedPtr<IPropertyHandle> PropertyHandle) const
-{
-	FString Value;
-
-	if (PropertyHandle->GetValue(Value) == FPropertyAccess::Success)
-	{
-		FString WindowsAudioDeviceName;
-		GConfig->GetString(TEXT("/Script/WindowsTargetPlatform.WindowsTargetSettings"), TEXT("AudioDevice"), WindowsAudioDeviceName, GEngineIni);
-		return FText::FromString(WindowsAudioDeviceName);
-	}
-
-	return FText::GetEmpty();
-}
-
-void FWindowsTargetSettingsDetails::HandleAudioDeviceTextBoxTextChanged(const FText& InText, TSharedPtr<IPropertyHandle> PropertyHandle)
-{
-	PropertyHandle->SetValue(InText.ToString());
-}
-
-void FWindowsTargetSettingsDetails::HandleAudioDeviceTextBoxTextComitted(const FText& InText, ETextCommit::Type CommitType, TSharedPtr<IPropertyHandle> PropertyHandle)
-{
-	FString Value;
-
-	// Clear the property if its not valid
-	if ((PropertyHandle->GetValue(Value) != FPropertyAccess::Success) || !IsValidAudioDeviceName(Value))
-	{
-		PropertyHandle->SetValue(FString());
-	}
-}
-
-bool FWindowsTargetSettingsDetails::IsValidAudioDeviceName(const FString& InDeviceName) const
-{
-	bool bIsValid = false;
-
-#if WITH_ENGINE
-	FAudioDevice* AudioDevice = GEngine->GetMainAudioDevice();
-	if (AudioDevice)
-	{
-		TArray<FString> DeviceNames;
-		AudioDevice->GetAudioDeviceList(DeviceNames);
-
-		for (FString& DeviceName : DeviceNames)
-		{
-			if (InDeviceName == DeviceName)
-			{
-				bIsValid = true;
-				break;
-			}
-		}
-	}
-#endif
-
-	return bIsValid;
-}
-
-TSharedRef<SWidget> FWindowsTargetSettingsDetails::MakeAudioDeviceMenu(const TSharedPtr<IPropertyHandle>& PropertyHandle)
-{
-	FMenuBuilder MenuBuilder(true, nullptr);
-
-#if WITH_ENGINE
-	FAudioDevice* AudioDevice = GEngine->GetMainAudioDevice();
-	if (AudioDevice)
-	{
-		TArray<FString> AudioDeviceNames;
-		AudioDevice->GetAudioDeviceList(AudioDeviceNames);
-
-		// Construct the custom menu widget from the list of device names
-		MenuBuilder.BeginSection(NAME_None, LOCTEXT("AudioDevicesSectionHeader", "Audio Devices"));
-		{
-			for (int32 i = 0; i < AudioDeviceNames.Num(); i++)
-			{
-				FUIAction Action(FExecuteAction::CreateRaw(this, &FWindowsTargetSettingsDetails::HandleAudioDeviceSelected, AudioDeviceNames[i], PropertyHandle));
-				MenuBuilder.AddMenuEntry(
-					FText::FromString(AudioDeviceNames[i]),
-					FText::FromString(TEXT("")),
-					FSlateIcon(),
-					Action
-					);
-			}
-		}
-		MenuBuilder.EndSection();
-	}
-#endif
-
-	return MenuBuilder.MakeWidget();
-}
-
 #undef LOCTEXT_NAMESPACE
diff --git a/Engine/Source/Editor/DetailCustomizations/Private/WindowsTargetSettingsDetails.h b/Engine/Source/Editor/DetailCustomizations/Private/WindowsTargetSettingsDetails.h
--- a/Engine/Source/Editor/DetailCustomizations/Private/WindowsTargetSettingsDetails.h
+++ b/Engine/Source/Editor/DetailCustomizations/Private/WindowsTargetSettingsDetails.h
@@ -35,34 +35,15 @@ private:
 	/** Delegate handler to set the path to start picking from */
 	bool HandlePostExternalIconCopy(const FString& InChosenImage);
 
-	/** Handles when a new audio device is selected from list of available audio devices. */
-	void HandleAudioDeviceSelected(FString AudioDeviceName, TSharedPtr<IPropertyHandle> PropertyHandle);
-
 	/** Handles when Stream Caching is toggled. */
 	void HandleAudioStreamCachingToggled(ECheckBoxState bEnableStreamCaching, TSharedPtr<IPropertyHandle> PropertyHandle);
 
 	/** This gets the current value of the audio stream caching bool property. */
 	ECheckBoxState GetAudioStreamCachingToggled(TSharedPtr<IPropertyHandle> PropertyHandle) const;
 
-	/** Handles changing the foreground color of the audio device box. */
-	FSlateColor HandleAudioDeviceBoxForegroundColor(TSharedPtr<IPropertyHandle> PropertyHandle) const;
-
-	/** Handles getting text of the audio device list text block. */
-	FText HandleAudioDeviceTextBoxText(TSharedPtr<IPropertyHandle> PropertyHandle) const;
-
-	/** Handles text changes in the audio device list text block. */
-	void HandleAudioDeviceTextBoxTextChanged(const FText& InText, TSharedPtr<IPropertyHandle> PropertyHandle);
-
-	/** Handles committing changes in the audio list text block. */
-	void HandleAudioDeviceTextBoxTextComitted(const FText& InText, ETextCommit::Type CommitType, TSharedPtr<IPropertyHandle> PropertyHandle);
-
 protected:
 
-	/** Checks if the device name is valid. */
-	bool IsValidAudioDeviceName(const FString& InDeviceName) const;
 
-	/** Creates a widget for the audio device picker. */
-	TSharedRef<SWidget> MakeAudioDeviceMenu(const TSharedPtr<IPropertyHandle>& PropertyHandle);
 
 private:
 	/** Reference to the target shader formats property view */
