Fix build error introduced by bad merge from 4.23.

diff --git a/Engine/Source/Editor/DetailCustomizations/Private/WindowsTargetSettingsDetails.cpp b/Engine/Source/Editor/DetailCustomizations/Private/WindowsTargetSettingsDetails.cpp
--- a/Engine/Source/Editor/DetailCustomizations/Private/WindowsTargetSettingsDetails.cpp
+++ b/Engine/Source/Editor/DetailCustomizations/Private/WindowsTargetSettingsDetails.cpp
@@ -289,11 +289,12 @@ void FWindowsTargetSettingsDetails::CustomizeDetails( IDetailLayoutBuilder& Deta
 
 
 	AudioPluginWidgetManager.BuildAudioCategory(DetailBuilder, EAudioPlatform::Windows);
+	IDetailCategoryBuilder& AudioCategory = DetailBuilder.EditCategory("Audio");
 
 	// Here we add a callback when the 
 	TSharedPtr<IPropertyHandle> AudioStreamCachingPropertyHandle = DetailBuilder.GetProperty("bUseAudioStreamCaching");
 	IDetailCategoryBuilder& AudioStreamCachingCategory = DetailBuilder.EditCategory("Audio");
-	IDetailPropertyRow& AudioStreamCachingPropertyRow = AudioDeviceCategory.AddProperty(AudioStreamCachingPropertyHandle);
+	IDetailPropertyRow& AudioStreamCachingPropertyRow = AudioCategory.AddProperty(AudioStreamCachingPropertyHandle);
 	AudioStreamCachingPropertyRow.CustomWidget()
 		.NameContent()
 		[
