Fixed ComponentTransformDetails so that committing a slider transaction always calls the postamble, which was currently not happening. This means that properties correctly get their PostEditChangeChainProperty() called, along with their Start/EndBroadcastMovement events.  Note that these will also get called in the case that the slider interaction did not change the value, although no transaction will be added to the undo stack.

diff --git a/Engine/Source/Editor/DetailCustomizations/Private/ComponentTransformDetails.cpp b/Engine/Source/Editor/DetailCustomizations/Private/ComponentTransformDetails.cpp
--- a/Engine/Source/Editor/DetailCustomizations/Private/ComponentTransformDetails.cpp
+++ b/Engine/Source/Editor/DetailCustomizations/Private/ComponentTransformDetails.cpp
@@ -85,6 +85,7 @@ FComponentTransformDetails::FComponentTransformDetails( const TArray< TWeakObjec
 	, NotifyHook( DetailBuilder.GetPropertyUtilities()->GetNotifyHook() )
 	, bPreserveScaleRatio( false )
 	, bEditingRotationInUI( false )
+	, bIsSliderTransaction( false )
 	, HiddenFieldMask( 0 )
 {
 	GConfig->GetBool(TEXT("SelectionDetails"), TEXT("PreserveScaleRatio"), bPreserveScaleRatio, GEditorPerProjectIni);
@@ -1079,12 +1080,12 @@ void FComponentTransformDetails::OnSetTransform(ETransformField::Type TransformF
 					NewComponentValue = GetAxisFilteredVector(Axis, NewValue, OldComponentValue);
 				}
 
-				// If we're committing during a rotation edit then we need to force it
-				if (OldComponentValue != NewComponentValue || (bCommitted && bEditingRotationInUI))
+				// If we're committing during a slider transaction then we need to force it, in order that PostEditChangeChainProperty be called.
+				// Note: this will even happen if the slider hasn't changed the value.
+				if (OldComponentValue != NewComponentValue || (bCommitted && bIsSliderTransaction))
 				{
 					if (!bBeganTransaction && bCommitted)
 					{
-						// Begin a transaction the first time an actors rotation is about to change.
 						// NOTE: One transaction per change, not per actor
 						GEditor->BeginTransaction(TransactionText);
 						bBeganTransaction = true;
@@ -1369,6 +1370,7 @@ void FComponentTransformDetails::OnBeginRotationSlider()
 	BeginSliderTransaction(ActorTransaction, ComponentTransaction);
 
 	bEditingRotationInUI = true;
+	bIsSliderTransaction = true;
 
 	for (TWeakObjectPtr<UObject> ObjectPtr : SelectedObjects)
 	{
@@ -1392,11 +1394,13 @@ void FComponentTransformDetails::OnEndRotationSlider(float NewValue)
 {
 	// Commit gets called right before this, only need to end the transaction
 	bEditingRotationInUI = false;
+	bIsSliderTransaction = false;
 	GEditor->EndTransaction();
 }
 
 void FComponentTransformDetails::OnBeginLocationSlider()
 {
+	bIsSliderTransaction = true;
 	FText ActorTransaction = LOCTEXT("OnSetLocation", "Set Location");
 	FText ComponentTransaction = LOCTEXT("OnSetLocation_ComponentDirect", "Modify Component Location");
 	BeginSliderTransaction(ActorTransaction, ComponentTransaction);
@@ -1404,11 +1408,13 @@ void FComponentTransformDetails::OnBeginLocationSlider()
 
 void FComponentTransformDetails::OnEndLocationSlider(float NewValue)
 {
+	bIsSliderTransaction = false;
 	GEditor->EndTransaction();
 }
 
 void FComponentTransformDetails::OnBeginScaleSlider()
 {
+	bIsSliderTransaction = true;
 	FText ActorTransaction = LOCTEXT("OnSetScale", "Set Scale");
 	FText ComponentTransaction = LOCTEXT("OnSetScale_ComponentDirect", "Modify Component Scale");
 	BeginSliderTransaction(ActorTransaction, ComponentTransaction);
@@ -1416,6 +1422,7 @@ void FComponentTransformDetails::OnBeginScaleSlider()
 
 void FComponentTransformDetails::OnEndScaleSlider(float NewValue)
 {
+	bIsSliderTransaction = false;
 	GEditor->EndTransaction();
 }
 
diff --git a/Engine/Source/Editor/DetailCustomizations/Private/ComponentTransformDetails.h b/Engine/Source/Editor/DetailCustomizations/Private/ComponentTransformDetails.h
--- a/Engine/Source/Editor/DetailCustomizations/Private/ComponentTransformDetails.h
+++ b/Engine/Source/Editor/DetailCustomizations/Private/ComponentTransformDetails.h
@@ -282,6 +282,8 @@ private:
 	FOptionalVector CachedScale;
 	/** Notify hook to use */
 	FNotifyHook* NotifyHook;
+	/** Mapping from object to relative rotation values which are not affected by Quat->Rotator conversions during transform calculations */
+	TMap< UObject*, FRotator > ObjectToRelativeRotationMap;
 	/** Whether or not we are in absolute translation mode */
 	bool bAbsoluteLocation;
 	/** Whether or not we are in absolute rotation mode */
@@ -290,10 +292,10 @@ private:
 	bool bAbsoluteScale;
 	/** Whether or not to preserve scale ratios */
 	bool bPreserveScaleRatio;
-	/** Mapping from object to relative rotation values which are not affected by Quat->Rotator conversions during transform calculations */
-	TMap< UObject*, FRotator > ObjectToRelativeRotationMap;
 	/** Flag to indicate we are currently editing the rotation in the UI, so we should rely on the cached value in objectToRelativeRotationMap, not the value from the object */
 	bool bEditingRotationInUI;
+	/** Flag to indicate we are currently performing a slider transaction */
+	bool bIsSliderTransaction;
 	/** Bitmask to indicate which fields should be hidden (if any) */
 	uint8 HiddenFieldMask;
 };
