Margin property editor remains selected and keyboard focused after value is committed (like most editable text properties).

diff --git a/Engine/Source/Editor/DetailCustomizations/Private/MarginCustomization.cpp b/Engine/Source/Editor/DetailCustomizations/Private/MarginCustomization.cpp
--- a/Engine/Source/Editor/DetailCustomizations/Private/MarginCustomization.cpp
+++ b/Engine/Source/Editor/DetailCustomizations/Private/MarginCustomization.cpp
@@ -84,6 +84,8 @@ TSharedRef<SEditableTextBox> FMarginStructCustomization::MakePropertyWidget()
 		.OnTextCommitted( this, &FMarginStructCustomization::OnMarginTextCommitted )
 		.Font( IDetailLayoutBuilder::GetDetailFont() )
 		.SelectAllTextWhenFocused( true )
+		.SelectAllTextOnCommit( true )
+		.ClearKeyboardFocusOnCommit( false )
 	;
 }
 
@@ -224,7 +226,7 @@ void FMarginStructCustomization::OnMarginTextCommitted( const FText& InText, ETe
 							}
 						}
 
-						StructPropertyHandle->NotifyPostChange();
+						StructPropertyHandle->NotifyPostChange(EPropertyChangeType::ValueSet);
 					}
 				}
 
