Hack fix for sequencer render UI page getting jammed up deferring refresh calls

diff --git a/Engine/Source/Editor/DetailCustomizations/Private/MovieSceneCaptureCustomization.cpp b/Engine/Source/Editor/DetailCustomizations/Private/MovieSceneCaptureCustomization.cpp
--- a/Engine/Source/Editor/DetailCustomizations/Private/MovieSceneCaptureCustomization.cpp
+++ b/Engine/Source/Editor/DetailCustomizations/Private/MovieSceneCaptureCustomization.cpp
@@ -15,6 +15,12 @@ TSharedRef<IDetailCustomization> FMovieSceneCaptureCustomization::MakeInstance()
 	return MakeShareable(new FMovieSceneCaptureCustomization);
 }
 
+FMovieSceneCaptureCustomization::FMovieSceneCaptureCustomization()
+{
+	ObjectsReplacedHandle = GEditor->OnObjectsReplaced().AddRaw(this, &FMovieSceneCaptureCustomization::OnObjectsReplaced);
+	PropertyChangedHandle = FCoreUObjectDelegates::OnObjectPropertyChanged.AddRaw(this, &FMovieSceneCaptureCustomization::OnObjectPostEditChange);
+}
+
 FMovieSceneCaptureCustomization::~FMovieSceneCaptureCustomization()
 {
 	FCoreUObjectDelegates::OnObjectPropertyChanged.Remove(PropertyChangedHandle);
@@ -25,16 +31,21 @@ void FMovieSceneCaptureCustomization::CustomizeDetails(IDetailLayoutBuilder& Det
 {
 	PropertyUtilities = DetailBuilder.GetPropertyUtilities();
 	DetailBuilder.GetObjectsBeingCustomized(ObjectsBeingCustomized);
-
-	PropertyChangedHandle = FCoreUObjectDelegates::OnObjectPropertyChanged.AddRaw(this, &FMovieSceneCaptureCustomization::OnObjectPostEditChange);
-	ObjectsReplacedHandle = GEditor->OnObjectsReplaced().AddRaw(this, &FMovieSceneCaptureCustomization::OnObjectsReplaced);
 }
 
 void FMovieSceneCaptureCustomization::OnObjectsReplaced(const TMap<UObject*, UObject*>& OldToNewInstanceMap)
 {
+	static bool bQueuedRefresh = false;
+
+	if (bQueuedRefresh)
+	{
+		return;
+	}
+	bQueuedRefresh = true;
+
 	// Defer the update 1 frame to ensure that we don't end up in a recursive loop adding bindings to the OnObjectsReplaced delegate that is currently being triggered
 	// (since the bindings are added in FMovieSceneCaptureCustomization::CustomizeDetails)
-	PropertyUtilities->EnqueueDeferredAction(FSimpleDelegate::CreateLambda([LocalPropertyUtilities = PropertyUtilities]{ LocalPropertyUtilities->ForceRefresh(); }));
+	PropertyUtilities->EnqueueDeferredAction(FSimpleDelegate::CreateLambda([LocalPropertyUtilities = PropertyUtilities] { bQueuedRefresh = false; LocalPropertyUtilities->ForceRefresh(); }));
 }
 
 void FMovieSceneCaptureCustomization::OnObjectPostEditChange( UObject* Object, FPropertyChangedEvent& PropertyChangedEvent )
diff --git a/Engine/Source/Editor/DetailCustomizations/Private/MovieSceneCaptureCustomization.h b/Engine/Source/Editor/DetailCustomizations/Private/MovieSceneCaptureCustomization.h
--- a/Engine/Source/Editor/DetailCustomizations/Private/MovieSceneCaptureCustomization.h
+++ b/Engine/Source/Editor/DetailCustomizations/Private/MovieSceneCaptureCustomization.h
@@ -12,6 +12,7 @@ struct FPropertyChangedEvent;
 class FMovieSceneCaptureCustomization : public IDetailCustomization
 {
 public:
+	FMovieSceneCaptureCustomization();
 	~FMovieSceneCaptureCustomization();
 
 	static TSharedRef<IDetailCustomization> MakeInstance();
