Null ptr check in slate brush

#rnx

diff --git a/Engine/Source/Editor/DetailCustomizations/Private/SlateBrushCustomization.cpp b/Engine/Source/Editor/DetailCustomizations/Private/SlateBrushCustomization.cpp
--- a/Engine/Source/Editor/DetailCustomizations/Private/SlateBrushCustomization.cpp
+++ b/Engine/Source/Editor/DetailCustomizations/Private/SlateBrushCustomization.cpp
@@ -1062,7 +1062,11 @@ private:
 			// of the Data Table Editor...
 			if (RawData.Num() > 0)
 			{
-				TemporaryBrush = *static_cast<FSlateBrush*>(RawData[0]);
+				FSlateBrush* SlateBrush = static_cast<FSlateBrush*>(RawData[0]);
+				if (SlateBrush)
+				{
+					TemporaryBrush = *SlateBrush;
+				}				
 			}
 		}
 	}
