Updated BodyInstanceCustomization to notify when collision is set back to default so we can update collision visualizations
Updated Static Mesh Editor to refresh collision visualization on instances when source mesh collision profile is changed

diff --git a/Engine/Source/Editor/DetailCustomizations/Private/BodyInstanceCustomization.cpp b/Engine/Source/Editor/DetailCustomizations/Private/BodyInstanceCustomization.cpp
--- a/Engine/Source/Editor/DetailCustomizations/Private/BodyInstanceCustomization.cpp
+++ b/Engine/Source/Editor/DetailCustomizations/Private/BodyInstanceCustomization.cpp
@@ -824,6 +824,8 @@ void FBodyInstanceCustomization::OnCollisionProfileComboOpening()
 
 void FBodyInstanceCustomization::MarkAllBodiesDefaultCollision(bool bUseDefaultCollision)
 {
+	CollisionResponsesHandle->NotifyPreChange();
+
 	if(PrimComponents.Num() && UseDefaultCollisionHandle.IsValid())	//If we have prim components we might be coming from bp editor which needs to propagate all instances
 	{
 		for(UPrimitiveComponent* PrimComp : PrimComponents)
@@ -850,6 +852,8 @@ void FBodyInstanceCustomization::MarkAllBodiesDefaultCollision(bool bUseDefaultC
 			}
 		}
 	}
+
+	CollisionResponsesHandle->NotifyPostChange();
 }
 
 void FBodyInstanceCustomization::OnCollisionProfileChanged( TSharedPtr<FString> NewSelection, ESelectInfo::Type SelectInfo, IDetailGroup* CollisionGroup )
diff --git a/Engine/Source/Editor/StaticMeshEditor/Private/StaticMeshEditor.cpp b/Engine/Source/Editor/StaticMeshEditor/Private/StaticMeshEditor.cpp
--- a/Engine/Source/Editor/StaticMeshEditor/Private/StaticMeshEditor.cpp
+++ b/Engine/Source/Editor/StaticMeshEditor/Private/StaticMeshEditor.cpp
@@ -2138,6 +2138,19 @@ void FStaticMeshEditor::NotifyPostChange( const FPropertyChangedEvent& PropertyC
 		{
 			RefreshTool();
 		}
+		else if (PropertyChangedEvent.GetPropertyName() == TEXT("CollisionResponses"))
+		{
+			for (FObjectIterator Iter(UStaticMeshComponent::StaticClass()); Iter; ++Iter)
+			{
+				UStaticMeshComponent* StaticMeshComponent = Cast<UStaticMeshComponent>(*Iter);
+				if (StaticMeshComponent->GetStaticMesh() == StaticMesh)
+				{
+					StaticMeshComponent->UpdateCollisionFromStaticMesh();
+					StaticMeshComponent->MarkRenderTransformDirty();
+				}
+			}
+		}
+
 	}
 }
 
