Fixed MaterialProxySettings and MaterialOptions customizations to adjust metadata for their FIntPoint instances instead of all FIntPoint properties

#rn FIntPoint properties will no longer occasionally be clamped to (1,11585) in the details panel

diff --git a/Engine/Source/Developer/MaterialBaking/Private/MaterialOptionsCustomization.cpp b/Engine/Source/Developer/MaterialBaking/Private/MaterialOptionsCustomization.cpp
--- a/Engine/Source/Developer/MaterialBaking/Private/MaterialOptionsCustomization.cpp
+++ b/Engine/Source/Developer/MaterialBaking/Private/MaterialOptionsCustomization.cpp
@@ -200,14 +200,14 @@ void AddTextureSizeClamping(TSharedPtr<IPropertyHandle> TextureSizeProperty)
 	const FString MaxTextureResolutionString = FString::FromInt(GetMax2DTextureDimension());
 	TextureSizeProperty->GetProperty()->SetMetaData(TEXT("ClampMax"), *MaxTextureResolutionString);
 	TextureSizeProperty->GetProperty()->SetMetaData(TEXT("UIMax"), *MaxTextureResolutionString);
-	PropertyX->GetProperty()->SetMetaData(TEXT("ClampMax"), *MaxTextureResolutionString);
-	PropertyX->GetProperty()->SetMetaData(TEXT("UIMax"), *MaxTextureResolutionString);
-	PropertyY->GetProperty()->SetMetaData(TEXT("ClampMax"), *MaxTextureResolutionString);
-	PropertyY->GetProperty()->SetMetaData(TEXT("UIMax"), *MaxTextureResolutionString);
+	PropertyX->SetInstanceMetaData(TEXT("ClampMax"), *MaxTextureResolutionString);
+	PropertyX->SetInstanceMetaData(TEXT("UIMax"), *MaxTextureResolutionString);
+	PropertyY->SetInstanceMetaData(TEXT("ClampMax"), *MaxTextureResolutionString);
+	PropertyY->SetInstanceMetaData(TEXT("UIMax"), *MaxTextureResolutionString);
 
 	const FString MinTextureResolutionString("1");
-	PropertyX->GetProperty()->SetMetaData(TEXT("ClampMin"), *MinTextureResolutionString);
-	PropertyX->GetProperty()->SetMetaData(TEXT("UIMin"), *MinTextureResolutionString);
-	PropertyY->GetProperty()->SetMetaData(TEXT("ClampMin"), *MinTextureResolutionString);
-	PropertyY->GetProperty()->SetMetaData(TEXT("UIMin"), *MinTextureResolutionString);
+	PropertyX->SetInstanceMetaData(TEXT("ClampMin"), *MinTextureResolutionString);
+	PropertyX->SetInstanceMetaData(TEXT("UIMin"), *MinTextureResolutionString);
+	PropertyY->SetInstanceMetaData(TEXT("ClampMin"), *MinTextureResolutionString);
+	PropertyY->SetInstanceMetaData(TEXT("UIMin"), *MinTextureResolutionString);
 }
diff --git a/Engine/Source/Editor/DetailCustomizations/Private/MaterialProxySettingsCustomizations.cpp b/Engine/Source/Editor/DetailCustomizations/Private/MaterialProxySettingsCustomizations.cpp
--- a/Engine/Source/Editor/DetailCustomizations/Private/MaterialProxySettingsCustomizations.cpp
+++ b/Engine/Source/Editor/DetailCustomizations/Private/MaterialProxySettingsCustomizations.cpp
@@ -138,16 +138,16 @@ void FMaterialProxySettingsCustomizations::AddTextureSizeClamping(TSharedPtr<IPr
 	const FString MaxTextureResolutionString = FString::FromInt(MaxProxyTextureResolution);
 	TextureSizeProperty->GetProperty()->SetMetaData(TEXT("ClampMax"), *MaxTextureResolutionString);
 	TextureSizeProperty->GetProperty()->SetMetaData(TEXT("UIMax"), *MaxTextureResolutionString);
-	PropertyX->GetProperty()->SetMetaData(TEXT("ClampMax"), *MaxTextureResolutionString);
-	PropertyX->GetProperty()->SetMetaData(TEXT("UIMax"), *MaxTextureResolutionString);
-	PropertyY->GetProperty()->SetMetaData(TEXT("ClampMax"), *MaxTextureResolutionString);
-	PropertyY->GetProperty()->SetMetaData(TEXT("UIMax"), *MaxTextureResolutionString);
+	PropertyX->SetInstanceMetaData(TEXT("ClampMax"), *MaxTextureResolutionString);
+	PropertyX->SetInstanceMetaData(TEXT("UIMax"), *MaxTextureResolutionString);
+	PropertyY->SetInstanceMetaData(TEXT("ClampMax"), *MaxTextureResolutionString);
+	PropertyY->SetInstanceMetaData(TEXT("UIMax"), *MaxTextureResolutionString);
 
 	const FString MinTextureResolutionString("1");
-	PropertyX->GetProperty()->SetMetaData(TEXT("ClampMin"), *MinTextureResolutionString);
-	PropertyX->GetProperty()->SetMetaData(TEXT("UIMin"), *MinTextureResolutionString);
-	PropertyY->GetProperty()->SetMetaData(TEXT("ClampMin"), *MinTextureResolutionString);
-	PropertyY->GetProperty()->SetMetaData(TEXT("UIMin"), *MinTextureResolutionString);
+	PropertyX->SetInstanceMetaData(TEXT("ClampMin"), *MinTextureResolutionString);
+	PropertyX->SetInstanceMetaData(TEXT("UIMin"), *MinTextureResolutionString);
+	PropertyY->SetInstanceMetaData(TEXT("ClampMin"), *MinTextureResolutionString);
+	PropertyY->SetInstanceMetaData(TEXT("UIMin"), *MinTextureResolutionString);
 }
 
 EVisibility FMaterialProxySettingsCustomizations::AreManualOverrideTextureSizesEnabled() const
