Forward Property Metadata on the per platform Properties (needed to support min/max/clamp etc on FPerPlatformInt and other specializations)

diff --git a/Engine/Source/Editor/DetailCustomizations/Private/PerPlatformPropertyCustomization.cpp b/Engine/Source/Editor/DetailCustomizations/Private/PerPlatformPropertyCustomization.cpp
--- a/Engine/Source/Editor/DetailCustomizations/Private/PerPlatformPropertyCustomization.cpp
+++ b/Engine/Source/Editor/DetailCustomizations/Private/PerPlatformPropertyCustomization.cpp
@@ -16,6 +16,7 @@
 #include "SPerPlatformPropertiesWidget.h"
 #include "ScopedTransaction.h"
 #include "IPropertyUtilities.h"
+#include "UObject/MetaData.h"
 
 #define LOCTEXT_NAMESPACE "PerPlatformPropertyCustomization"
 
@@ -26,6 +27,42 @@ void FPerPlatformPropertyCustomization<PerPlatformType>::CustomizeHeader(TShared
 
 	int32 PlatformNumber = PlatformInfo::GetAllPlatformGroupNames().Num();
 
+	// Get Forward Property MetaData to PerPlatform Property
+	UProperty* MetaDataProperty = StructPropertyHandle->GetMetaDataProperty();
+	if (MetaDataProperty != nullptr)
+	{
+		TMap<FName, FString>* MetaData = UMetaData::GetMapForObject(MetaDataProperty);
+		if (MetaData != nullptr)
+		{
+			TSharedPtr<IPropertyHandle> ChildPropertyHandle = StructPropertyHandle->GetChildHandle(FName("Default"));
+			if (ChildPropertyHandle.IsValid())
+			{
+				for (const TPair<FName, FString>& Pair : *MetaData)
+				{
+					ChildPropertyHandle->SetInstanceMetaData(Pair.Key, Pair.Value);
+				}
+			}
+
+			TSharedPtr<IPropertyHandle>	MapProperty = StructPropertyHandle->GetChildHandle(FName("PerPlatform"));
+			if (MapProperty.IsValid())
+			{
+				uint32 NumChildren = 0;
+				MapProperty->GetNumChildren(NumChildren);
+				for (uint32 ChildIdx = 0; ChildIdx < NumChildren; ChildIdx++)
+				{
+					TSharedPtr<IPropertyHandle> ChildProperty = MapProperty->GetChildHandle(ChildIdx);
+					if (ChildProperty.IsValid())
+					{
+						for (const TPair<FName, FString>& Pair : *MetaData)
+						{
+							ChildProperty->SetInstanceMetaData(Pair.Key, Pair.Value);
+						}
+					}
+				}
+			}
+		}
+	}
+
 	HeaderRow.NameContent()
 	[
 		StructPropertyHandle->CreatePropertyNameWidget()
