Revert fix as it was already done by jurre.debaare in the dev-core branch

Undo single //UE4/Dev-Editor/Engine/Source/Editor/DetailCustomizations/Private/PerPlatformPropertyCustomization.cpp revision 7

diff --git a/Engine/Source/Editor/DetailCustomizations/Private/PerPlatformPropertyCustomization.cpp b/Engine/Source/Editor/DetailCustomizations/Private/PerPlatformPropertyCustomization.cpp
--- a/Engine/Source/Editor/DetailCustomizations/Private/PerPlatformPropertyCustomization.cpp
+++ b/Engine/Source/Editor/DetailCustomizations/Private/PerPlatformPropertyCustomization.cpp
@@ -16,7 +16,6 @@
 #include "SPerPlatformPropertiesWidget.h"
 #include "ScopedTransaction.h"
 #include "IPropertyUtilities.h"
-#include "UObject/MetaData.h"
 
 #define LOCTEXT_NAMESPACE "PerPlatformPropertyCustomization"
 
@@ -27,42 +26,6 @@ void FPerPlatformPropertyCustomization<PerPlatformType>::CustomizeHeader(TShared
 
 	int32 PlatformNumber = PlatformInfo::GetAllPlatformGroupNames().Num();
 
-	// Get Forward Property MetaData to PerPlatform Property
-	UProperty* MetaDataProperty = StructPropertyHandle->GetMetaDataProperty();
-	if (MetaDataProperty != nullptr)
-	{
-		TMap<FName, FString>* MetaData = UMetaData::GetMapForObject(MetaDataProperty);
-		if (MetaData != nullptr)
-		{
-			TSharedPtr<IPropertyHandle> ChildPropertyHandle = StructPropertyHandle->GetChildHandle(FName("Default"));
-			if (ChildPropertyHandle.IsValid())
-			{
-				for (const TPair<FName, FString>& Pair : *MetaData)
-				{
-					ChildPropertyHandle->SetInstanceMetaData(Pair.Key, Pair.Value);
-				}
-			}
-
-			TSharedPtr<IPropertyHandle>	MapProperty = StructPropertyHandle->GetChildHandle(FName("PerPlatform"));
-			if (MapProperty.IsValid())
-			{
-				uint32 NumChildren = 0;
-				MapProperty->GetNumChildren(NumChildren);
-				for (uint32 ChildIdx = 0; ChildIdx < NumChildren; ChildIdx++)
-				{
-					TSharedPtr<IPropertyHandle> ChildProperty = MapProperty->GetChildHandle(ChildIdx);
-					if (ChildProperty.IsValid())
-					{
-						for (const TPair<FName, FString>& Pair : *MetaData)
-						{
-							ChildProperty->SetInstanceMetaData(Pair.Key, Pair.Value);
-						}
-					}
-				}
-			}
-		}
-	}
-
 	HeaderRow.NameContent()
 	[
 		StructPropertyHandle->CreatePropertyNameWidget()
