Merging //UE4/Dev-Rendering change 8934666 to Main (//UE4/Main)
color picker fixes - stop updating during edit

diff --git a/Engine/Source/Editor/DetailCustomizations/Private/ColorStructCustomization.cpp b/Engine/Source/Editor/DetailCustomizations/Private/ColorStructCustomization.cpp
--- a/Engine/Source/Editor/DetailCustomizations/Private/ColorStructCustomization.cpp
+++ b/Engine/Source/Editor/DetailCustomizations/Private/ColorStructCustomization.cpp
@@ -199,11 +199,12 @@ void FColorStructCustomization::CreateColorPicker(bool bUseAlpha)
 	GetColorAsLinear(InitialColor);
 
 	const bool bRefreshOnlyOnOk = bDontUpdateWhileEditing || StructPropertyHandle->HasMetaData("DontUpdateWhileEditing");
+	const bool bOnlyRefreshOnMouseUp = StructPropertyHandle->HasMetaData("OnlyUpdateOnInteractionEnd");
 
 	FColorPickerArgs PickerArgs;
 	{
 		PickerArgs.bUseAlpha = !bIgnoreAlpha;
-		PickerArgs.bOnlyRefreshOnMouseUp = false;
+		PickerArgs.bOnlyRefreshOnMouseUp = bOnlyRefreshOnMouseUp;
 		PickerArgs.bOnlyRefreshOnOk = bRefreshOnlyOnOk;
 		PickerArgs.sRGBOverride = sRGBOverride;
 		PickerArgs.DisplayGamma = TAttribute<float>::Create(TAttribute<float>::FGetter::CreateUObject(GEngine, &UEngine::GetDisplayGamma));
@@ -253,10 +254,11 @@ TSharedRef<SColorPicker> FColorStructCustomization::CreateInlineColorPicker(TWea
 	GetColorAsLinear(InitialColor);
 
 	const bool bRefreshOnlyOnOk = /*bDontUpdateWhileEditing ||*/ StructPropertyHandle->HasMetaData("DontUpdateWhileEditing");
-	
+	const bool bOnlyRefreshOnMouseUp = StructPropertyHandle->HasMetaData("OnlyUpdateOnInteractionEnd");
+
 	return SNew(SColorPicker)
 		.DisplayInlineVersion(true)
-		.OnlyRefreshOnMouseUp(false)
+		.OnlyRefreshOnMouseUp(bOnlyRefreshOnMouseUp)
 		.OnlyRefreshOnOk(bRefreshOnlyOnOk)
 		.DisplayGamma(TAttribute<float>::Create(TAttribute<float>::FGetter::CreateUObject(GEngine, &UEngine::GetDisplayGamma)))
 		.OnColorCommitted(FOnLinearColorValueChanged::CreateSP(this, &FColorStructCustomization::OnSetColorFromColorPicker))
diff --git a/Engine/Source/Editor/UnrealEd/Classes/MaterialEditor/DEditorVectorParameterValue.h b/Engine/Source/Editor/UnrealEd/Classes/MaterialEditor/DEditorVectorParameterValue.h
--- a/Engine/Source/Editor/UnrealEd/Classes/MaterialEditor/DEditorVectorParameterValue.h
+++ b/Engine/Source/Editor/UnrealEd/Classes/MaterialEditor/DEditorVectorParameterValue.h
@@ -13,7 +13,7 @@ class UNREALED_API UDEditorVectorParameterValue : public UDEditorParameterValue
 {
 	GENERATED_UCLASS_BODY()
 
-	UPROPERTY(EditAnywhere, Category=DEditorVectorParameterValue)
+	UPROPERTY(EditAnywhere, Category=DEditorVectorParameterValue, meta=(OnlyUpdateOnInteractionEnd))
 	FLinearColor ParameterValue;
 
 	UPROPERTY(Transient)
diff --git a/Engine/Source/Runtime/Engine/Classes/Materials/MaterialExpressionConstant3Vector.h b/Engine/Source/Runtime/Engine/Classes/Materials/MaterialExpressionConstant3Vector.h
--- a/Engine/Source/Runtime/Engine/Classes/Materials/MaterialExpressionConstant3Vector.h
+++ b/Engine/Source/Runtime/Engine/Classes/Materials/MaterialExpressionConstant3Vector.h
@@ -13,7 +13,7 @@ class UMaterialExpressionConstant3Vector
 {
 	GENERATED_UCLASS_BODY()
 
- 	UPROPERTY(EditAnywhere, BlueprintReadWrite, Category=MaterialExpressionConstant3Vector, meta=(DontUpdateWhileEditing, HideAlphaChannel))
+ 	UPROPERTY(EditAnywhere, BlueprintReadWrite, Category=MaterialExpressionConstant3Vector, meta=(OnlyUpdateOnInteractionEnd, HideAlphaChannel))
 	FLinearColor Constant;
 
 public:
diff --git a/Engine/Source/Runtime/Engine/Classes/Materials/MaterialExpressionConstant4Vector.h b/Engine/Source/Runtime/Engine/Classes/Materials/MaterialExpressionConstant4Vector.h
--- a/Engine/Source/Runtime/Engine/Classes/Materials/MaterialExpressionConstant4Vector.h
+++ b/Engine/Source/Runtime/Engine/Classes/Materials/MaterialExpressionConstant4Vector.h
@@ -13,7 +13,7 @@ class UMaterialExpressionConstant4Vector : public UMaterialExpression
 {
 	GENERATED_UCLASS_BODY()
 
-	UPROPERTY(EditAnywhere, BlueprintReadWrite, Category=MaterialExpressionConstant4Vector)
+	UPROPERTY(EditAnywhere, BlueprintReadWrite, Category=MaterialExpressionConstant4Vector, meta=(OnlyUpdateOnInteractionEnd))
 	FLinearColor Constant;
 
 	//~ Begin UMaterialExpression Interface
diff --git a/Engine/Source/Runtime/Engine/Classes/Materials/MaterialExpressionVectorParameter.h b/Engine/Source/Runtime/Engine/Classes/Materials/MaterialExpressionVectorParameter.h
--- a/Engine/Source/Runtime/Engine/Classes/Materials/MaterialExpressionVectorParameter.h
+++ b/Engine/Source/Runtime/Engine/Classes/Materials/MaterialExpressionVectorParameter.h
@@ -14,7 +14,7 @@ class UMaterialExpressionVectorParameter : public UMaterialExpressionParameter
 {
 	GENERATED_UCLASS_BODY()
 
-	UPROPERTY(EditAnywhere, Category=MaterialExpressionVectorParameter)
+	UPROPERTY(EditAnywhere, Category=MaterialExpressionVectorParameter, meta=(OnlyUpdateOnInteractionEnd))
 	FLinearColor DefaultValue;
 
 	UPROPERTY(EditAnywhere, Category=CustomPrimitiveData)
