Update Engine and Plugins to use the new Getters / Setters for UActorComponent, USceneComponent, and AActor replicated properties

diff --git a/Engine/Plugins/2D/Paper2D/Source/Paper2DEditor/Private/SpriteEditor/SpriteEditorViewportClient.cpp b/Engine/Plugins/2D/Paper2D/Source/Paper2DEditor/Private/SpriteEditor/SpriteEditorViewportClient.cpp
--- a/Engine/Plugins/2D/Paper2D/Source/Paper2DEditor/Private/SpriteEditor/SpriteEditorViewportClient.cpp
+++ b/Engine/Plugins/2D/Paper2D/Source/Paper2DEditor/Private/SpriteEditor/SpriteEditorViewportClient.cpp
@@ -118,7 +118,7 @@ FSpriteEditorViewportClient::FSpriteEditorViewportClient(TWeakPtr<FSpriteEditor>
 
 		// Nudge the source texture view back a bit so it doesn't occlude sprites
 		const FTransform Transform(-1.0f * PaperAxisZ);
-		SourceTextureViewComponent->bVisible = false;
+		SourceTextureViewComponent->SetVisibleFlag(false);
 		PreviewScene->AddComponent(SourceTextureViewComponent, Transform);
 	}
 }
diff --git a/Engine/Plugins/Compositing/Composure/Source/Composure/Private/ComposurePostProcessPass.cpp b/Engine/Plugins/Compositing/Composure/Source/Composure/Private/ComposurePostProcessPass.cpp
--- a/Engine/Plugins/Compositing/Composure/Source/Composure/Private/ComposurePostProcessPass.cpp
+++ b/Engine/Plugins/Compositing/Composure/Source/Composure/Private/ComposurePostProcessPass.cpp
@@ -81,7 +81,7 @@ void UComposurePostProcessPass::Activate(bool bReset)
 {
 	Super::Activate(bReset);
 
-	if (bIsActive)
+	if (IsActive())
 	{
 		if (SceneCapture == nullptr)
 		{
diff --git a/Engine/Plugins/Experimental/ImagePlate/Source/ImagePlate/Private/ImagePlateComponent.cpp b/Engine/Plugins/Experimental/ImagePlate/Source/ImagePlate/Private/ImagePlateComponent.cpp
--- a/Engine/Plugins/Experimental/ImagePlate/Source/ImagePlate/Private/ImagePlateComponent.cpp
+++ b/Engine/Plugins/Experimental/ImagePlate/Source/ImagePlate/Private/ImagePlateComponent.cpp
@@ -347,12 +347,12 @@ void UImagePlateComponent::UpdateTransformScale()
 		FVector HorizontalScale		= UImagePlateComponent::TransfromFromProjection(ScreenToLocalSpace, FVector4(Plate.FillScreenAmount.X/100.f, 0.f, ScreenSpaceLocalPosition.Z, 1.0f));
 		FVector VerticalScale		= UImagePlateComponent::TransfromFromProjection(ScreenToLocalSpace, FVector4(0.f, Plate.FillScreenAmount.Y/100.f, ScreenSpaceLocalPosition.Z, 1.0f));
 
-		SetRelativeScale3D(FVector(RelativeScale3D.X, HorizontalScale.Size(), VerticalScale.Size()));
-		SetRelativeLocation(FVector(RelativeLocation.X, 0.f, 0.f));
+		SetRelativeScale3D(FVector(GetRelativeScale3D().X, HorizontalScale.Size(), VerticalScale.Size()));
+		SetRelativeLocation(FVector(GetRelativeLocation().X, 0.f, 0.f));
 	}
 	else
 	{
-		SetRelativeScale3D(FVector(RelativeScale3D.X, Plate.FixedSize.X*.5f, Plate.FixedSize.Y*.5f));
+		SetRelativeScale3D(FVector(GetRelativeScale3D().X, Plate.FixedSize.X*.5f, Plate.FixedSize.Y*.5f));
 	}
 }
 
diff --git a/Engine/Plugins/Experimental/VirtualCamera/Source/VirtualCamera/Private/ConcertVirtualCamera.cpp b/Engine/Plugins/Experimental/VirtualCamera/Source/VirtualCamera/Private/ConcertVirtualCamera.cpp
--- a/Engine/Plugins/Experimental/VirtualCamera/Source/VirtualCamera/Private/ConcertVirtualCamera.cpp
+++ b/Engine/Plugins/Experimental/VirtualCamera/Source/VirtualCamera/Private/ConcertVirtualCamera.cpp
@@ -60,8 +60,8 @@ FConcertVirtualCameraCameraData::FConcertVirtualCameraCameraData(const AActor* I
 {
 	CameraActorLocation = InOwner->GetActorLocation();
 	CameraActorRotation = InOwner->GetActorRotation();
-	CameraComponentLocation = InCineCameraComponent->RelativeLocation;
-	CameraComponentRotation = InCineCameraComponent->RelativeRotation;
+	CameraComponentLocation = InCineCameraComponent->GetRelativeLocation();
+	CameraComponentRotation = InCineCameraComponent->GetRelativeRotation();
 
 	CurrentAperture = InCineCameraComponent->CurrentAperture;
 	CurrentFocalLength = InCineCameraComponent->CurrentFocalLength;
diff --git a/Engine/Plugins/Experimental/VirtualCamera/Source/VirtualCamera/Private/VirtualCameraMovementComponent.cpp b/Engine/Plugins/Experimental/VirtualCamera/Source/VirtualCamera/Private/VirtualCameraMovementComponent.cpp
--- a/Engine/Plugins/Experimental/VirtualCamera/Source/VirtualCamera/Private/VirtualCameraMovementComponent.cpp
+++ b/Engine/Plugins/Experimental/VirtualCamera/Source/VirtualCamera/Private/VirtualCameraMovementComponent.cpp
@@ -252,15 +252,17 @@ FVector UVirtualCameraMovementComponent::GetStabilizedDeltaLocation() const
 
 	FVector TargetLocationWithOffsets = TargetLocation - GetLocationOffset();
 
+	const FVector UpdatedComponentRelativeLocation = UpdatedComponent->GetRelativeLocation();
+
 	// Calculate each component by getting the needed vector component and doing a lerp with the stabilization scale for that axis
-	FMath::PointDistToLine(TargetLocationWithOffsets, UpdatedComponent->GetForwardVector(), UpdatedComponent->RelativeLocation, ClosestPoint);
-	ReturnVector += (ClosestPoint - UpdatedComponent->RelativeLocation) * (1 - ClampStabilizationScale(AxisSettings[EVirtualCameraAxis::LocationX].StabilizationScale));
+	FMath::PointDistToLine(TargetLocationWithOffsets, UpdatedComponent->GetForwardVector(), UpdatedComponentRelativeLocation, ClosestPoint);
+	ReturnVector += (ClosestPoint - UpdatedComponentRelativeLocation) * (1 - ClampStabilizationScale(AxisSettings[EVirtualCameraAxis::LocationX].StabilizationScale));
 
-	FMath::PointDistToLine(TargetLocationWithOffsets, UpdatedComponent->GetRightVector(), UpdatedComponent->RelativeLocation, ClosestPoint);
-	ReturnVector += (ClosestPoint - UpdatedComponent->RelativeLocation) * (1 - ClampStabilizationScale(AxisSettings[EVirtualCameraAxis::LocationY].StabilizationScale));
+	FMath::PointDistToLine(TargetLocationWithOffsets, UpdatedComponent->GetRightVector(), UpdatedComponentRelativeLocation, ClosestPoint);
+	ReturnVector += (ClosestPoint - UpdatedComponentRelativeLocation) * (1 - ClampStabilizationScale(AxisSettings[EVirtualCameraAxis::LocationY].StabilizationScale));
 
-	FMath::PointDistToLine(TargetLocationWithOffsets, UpdatedComponent->GetUpVector(), UpdatedComponent->RelativeLocation, ClosestPoint);
-	ReturnVector += (ClosestPoint - UpdatedComponent->RelativeLocation) * (1 - ClampStabilizationScale(AxisSettings[EVirtualCameraAxis::LocationZ].StabilizationScale));
+	FMath::PointDistToLine(TargetLocationWithOffsets, UpdatedComponent->GetUpVector(), UpdatedComponentRelativeLocation, ClosestPoint);
+	ReturnVector += (ClosestPoint - UpdatedComponentRelativeLocation) * (1 - ClampStabilizationScale(AxisSettings[EVirtualCameraAxis::LocationZ].StabilizationScale));
 
 	ReturnVector = GetOwner()->GetActorRotation().RotateVector(ReturnVector);
 	return ReturnVector;
diff --git a/Engine/Plugins/Experimental/VirtualCamera/Source/VirtualCamera/Private/VirtualCameraPlayerControllerBase.cpp b/Engine/Plugins/Experimental/VirtualCamera/Source/VirtualCamera/Private/VirtualCameraPlayerControllerBase.cpp
--- a/Engine/Plugins/Experimental/VirtualCamera/Source/VirtualCamera/Private/VirtualCameraPlayerControllerBase.cpp
+++ b/Engine/Plugins/Experimental/VirtualCamera/Source/VirtualCamera/Private/VirtualCameraPlayerControllerBase.cpp
@@ -494,7 +494,7 @@ void AVirtualCameraPlayerControllerBase::PilotTargetedCamera(AVirtualCameraPawnB
 		TargetCameraComponent->LensSettings = CameraToFollow->LensSettings;
 		TargetCameraComponent->FilmbackSettings = CameraToFollow->FilmbackSettings;
 	}
-	TargetCameraComponent->SetRelativeLocationAndRotation(CameraToFollow->RelativeLocation, CameraToFollow->RelativeRotation);
+	TargetCameraComponent->SetRelativeLocationAndRotation(CameraToFollow->GetRelativeLocation(), CameraToFollow->GetRelativeRotation());
 }
 
 bool AVirtualCameraPlayerControllerBase::InputTouch(uint32 Handle, ETouchType::Type Type, const FVector2D & TouchLocation, float Force, FDateTime DeviceTimestamp, uint32 TouchpadIndex)
diff --git a/Engine/Plugins/Experimental/VirtualProductionUtilities/Source/VPUtilities/Private/VPRootActor.cpp b/Engine/Plugins/Experimental/VirtualProductionUtilities/Source/VPUtilities/Private/VPRootActor.cpp
--- a/Engine/Plugins/Experimental/VirtualProductionUtilities/Source/VPUtilities/Private/VPRootActor.cpp
+++ b/Engine/Plugins/Experimental/VirtualProductionUtilities/Source/VPUtilities/Private/VPRootActor.cpp
@@ -62,7 +62,7 @@ AVPRootActor::AVPRootActor(const FObjectInitializer& ObjectInitializer)
 		ConstructorHelpers::FObjectFinderOptional<UTexture2D> RootTextureObject = TEXT("/VirtualProductionUtilities/Icons/S_VPRootActor");
 
 		SpriteComponent->Sprite = RootTextureObject.Get();
-		SpriteComponent->RelativeScale3D = FVector(0.5f, 0.5f, 0.5f);
+		SpriteComponent->SetRelativeScale3D_Direct(FVector(0.5f, 0.5f, 0.5f));
 		SpriteComponent->bHiddenInGame = false;
 		SpriteComponent->SpriteInfo.Category = ConstructorStatics.ID_Sprite;
 		SpriteComponent->SpriteInfo.DisplayName = ConstructorStatics.NAME_Sprite;
@@ -188,8 +188,8 @@ void AVPRootActor::PostEditChangeChainProperty(struct FPropertyChangedChainEvent
 
 	if (bMoveLevelWithActor && !bReentrantPostEditMove)
 	{
-		if (PropertyChangedEvent.Property->GetFName() == GET_MEMBER_NAME_CHECKED(USceneComponent, RelativeLocation)
-			|| PropertyChangedEvent.Property->GetFName() == GET_MEMBER_NAME_CHECKED(USceneComponent, RelativeRotation))
+		if (PropertyChangedEvent.Property->GetFName() == USceneComponent::GetRelativeLocationPropertyName()
+			|| PropertyChangedEvent.Property->GetFName() == USceneComponent::GetRelativeRotationPropertyName())
 		{
 			TGuardValue<bool> Tmp(bReentrantPostEditMove, true);
 			MoveLevelToRootActor();
diff --git a/Engine/Plugins/FX/Niagara/Source/Niagara/Private/NiagaraActor.cpp b/Engine/Plugins/FX/Niagara/Source/Niagara/Private/NiagaraActor.cpp
--- a/Engine/Plugins/FX/Niagara/Source/Niagara/Private/NiagaraActor.cpp
+++ b/Engine/Plugins/FX/Niagara/Source/Niagara/Private/NiagaraActor.cpp
@@ -40,7 +40,7 @@ ANiagaraActor::ANiagaraActor(const FObjectInitializer& ObjectInitializer)
 		if (SpriteComponent)
 		{
 			SpriteComponent->Sprite = ConstructorStatics.SpriteTextureObject.Get();
-			SpriteComponent->RelativeScale3D = FVector(0.5f, 0.5f, 0.5f);
+			SpriteComponent->SetRelativeScale3D_Direct(FVector(0.5f, 0.5f, 0.5f));
 			SpriteComponent->bHiddenInGame = true;
 			SpriteComponent->bIsScreenSizeScaled = true;
 			SpriteComponent->SpriteInfo.Category = ConstructorStatics.ID_Effects;
@@ -59,7 +59,7 @@ ANiagaraActor::ANiagaraActor(const FObjectInitializer& ObjectInitializer)
 			ArrowComponent->SpriteInfo.Category = ConstructorStatics.ID_Effects;
 			ArrowComponent->SpriteInfo.DisplayName = ConstructorStatics.NAME_Effects;
 			ArrowComponent->SetupAttachment(NiagaraComponent);
-			ArrowComponent->bAbsoluteScale = true;
+			ArrowComponent->SetUsingAbsoluteScale(true);
 		}
 	}
 #endif // WITH_EDITORONLY_DATA
diff --git a/Engine/Plugins/FX/Niagara/Source/Niagara/Private/NiagaraComponent.cpp b/Engine/Plugins/FX/Niagara/Source/Niagara/Private/NiagaraComponent.cpp
--- a/Engine/Plugins/FX/Niagara/Source/Niagara/Private/NiagaraComponent.cpp
+++ b/Engine/Plugins/FX/Niagara/Source/Niagara/Private/NiagaraComponent.cpp
@@ -86,7 +86,7 @@ void DumpNiagaraComponents(UWorld* World)
 				else
 				{
 					UE_LOG(LogNiagara, Log, TEXT("Component: \"%s\" System: \"%s\" | ReqExecState: %d | ExecState: %d | bIsActive: %d"), *Component->GetName(), *Sys->GetName(),
-						(int32)SysInst->GetRequestedExecutionState(), (int32)SysInst->GetActualExecutionState(), Component->bIsActive);
+						(int32)SysInst->GetRequestedExecutionState(), (int32)SysInst->GetActualExecutionState(), Component->IsActive());
 
 					if (!SysInst->IsComplete())
 					{
@@ -436,7 +436,7 @@ void UNiagaraComponent::ReleaseToPool()
 		return;
 	}
 
-	if (!bIsActive)
+	if (!IsActive())
 	{
 		//If we're already complete then release to the pool straight away.
 		UWorld* World = GetWorld();
@@ -479,13 +479,13 @@ void UNiagaraComponent::TickComponent(float DeltaSeconds, enum ELevelTick TickTy
 		return;
 	}
 
-	if (!bIsActive && bAutoActivate && SystemInstance.Get() && SystemInstance->GetAreDataInterfacesInitialized())
+	if (!IsActive() && bAutoActivate && SystemInstance.Get() && SystemInstance->GetAreDataInterfacesInitialized())
 	{
 		Activate();
 	}
 
 	check(SystemInstance->IsSolo());
-	if (bIsActive && SystemInstance.Get() && !SystemInstance->IsComplete())
+	if (IsActive() && SystemInstance.Get() && !SystemInstance->IsComplete())
 	{
 		// If the interfaces have changed in a meaningful way, we need to potentially rebind and update the values.
 		if (OverrideParameters.GetInterfacesDirty())
@@ -773,9 +773,9 @@ void UNiagaraComponent::Activate(bool bReset /* = false */)
 			{
 				bDidAutoAttach = bWasAutoAttached;
 				CancelAutoAttachment(true);
-				SavedAutoAttachRelativeLocation = RelativeLocation;
-				SavedAutoAttachRelativeRotation = RelativeRotation;
-				SavedAutoAttachRelativeScale3D = RelativeScale3D;
+				SavedAutoAttachRelativeLocation = GetRelativeLocation();
+				SavedAutoAttachRelativeRotation = GetRelativeRotation();
+				SavedAutoAttachRelativeScale3D = GetRelativeScale3D();
 				//bIsChangingAutoAttachment = true;
 				AttachToComponent(NewParent, FAttachmentTransformRules(AutoAttachLocationRule, AutoAttachRotationRule, AutoAttachScaleRule, false), AutoAttachSocketName);
 				//bIsChangingAutoAttachment = false;
@@ -826,12 +826,12 @@ void UNiagaraComponent::Deactivate()
 		SystemInstance->Deactivate();
 
 		// We are considered active until we are complete
-		bIsActive = !SystemInstance->IsComplete();
+		SetActiveFlag(!SystemInstance->IsComplete());
 	}
 	else
 	{
 		Super::Deactivate();
-		bIsActive = false;
+		SetActiveFlag(false);
 	}
 }
 
@@ -844,7 +844,7 @@ void UNiagaraComponent::DeactivateImmediate()
 
 	//UE_LOG(LogNiagara, Log, TEXT("Deactivate %s"), *GetName());
 
-	bIsActive = false;
+	SetActiveFlag(false);
 
 	if (SystemInstance)
 	{
@@ -857,7 +857,7 @@ void UNiagaraComponent::OnSystemComplete()
 	//UE_LOG(LogNiagara, Log, TEXT("OnSystemComplete: %p - %s"), SystemInstance.Get(), *Asset->GetName());
 
 	SetComponentTickEnabled(false);
-	bIsActive = false;
+	SetActiveFlag(false);
 
 	MarkRenderDynamicDataDirty();
 		
@@ -890,7 +890,7 @@ void UNiagaraComponent::DestroyInstance()
 {
 	//UE_LOG(LogNiagara, Log, TEXT("UNiagaraComponent::DestroyInstance: %p  %s\n"), SystemInstance.Get(), *GetAsset()->GetFullName());
 	//UE_LOG(LogNiagara, Log, TEXT("DestroyInstance: %u - %s"), this, *Asset->GetName());
-	bIsActive = false;
+	SetActiveFlag(false);
 	
 	// Rather than setting the unique ptr to null here, we allow it to transition ownership to the system's deferred deletion queue. This allows us to safely
 	// get rid of the system interface should we be doing this in response to a callback invoked during the system interface's lifetime completion cycle.
@@ -939,9 +939,9 @@ void UNiagaraComponent::OnRegister()
 			}
 		}
 
-		SavedAutoAttachRelativeLocation = RelativeLocation;
-		SavedAutoAttachRelativeRotation = RelativeRotation;
-		SavedAutoAttachRelativeScale3D = RelativeScale3D;
+		SavedAutoAttachRelativeLocation = GetRelativeLocation();
+		SavedAutoAttachRelativeRotation = GetRelativeRotation();
+		SavedAutoAttachRelativeScale3D = GetRelativeScale3D();
 	}
 	Super::OnRegister();
 }
@@ -958,7 +958,7 @@ void UNiagaraComponent::OnUnregister()
 {
 	Super::OnUnregister();
 
-	bIsActive = false;
+	SetActiveFlag(false);
 
 	if (SystemInstance)
 	{
@@ -1756,9 +1756,9 @@ void UNiagaraComponent::CancelAutoAttachment(bool bDetachFromParent)
 		if (bDidAutoAttach)
 		{
 			// Restore relative transform from before attachment. Actual transform will be updated as part of DetachFromParent().
-			RelativeLocation = SavedAutoAttachRelativeLocation;
-			RelativeRotation = SavedAutoAttachRelativeRotation;
-			RelativeScale3D = SavedAutoAttachRelativeScale3D;
+			SetRelativeLocation_Direct(SavedAutoAttachRelativeLocation);
+			SetRelativeRotation_Direct(SavedAutoAttachRelativeRotation);
+			SetRelativeScale3D_Direct(SavedAutoAttachRelativeScale3D);
 			bDidAutoAttach = false;
 		}
 
diff --git a/Engine/Plugins/FX/Niagara/Source/Niagara/Private/NiagaraFunctionLibrary.cpp b/Engine/Plugins/FX/Niagara/Source/Niagara/Private/NiagaraFunctionLibrary.cpp
--- a/Engine/Plugins/FX/Niagara/Source/Niagara/Private/NiagaraFunctionLibrary.cpp
+++ b/Engine/Plugins/FX/Niagara/Source/Niagara/Private/NiagaraFunctionLibrary.cpp
@@ -189,25 +189,25 @@ UNiagaraComponent* UNiagaraFunctionLibrary::SpawnSystemAttached(
 						const FTransform ParentToWorld = AttachToComponent->GetSocketTransform(AttachPointName);
 						const FTransform ComponentToWorld(Rotation, Location, Scale);
 						const FTransform RelativeTM = ComponentToWorld.GetRelativeTransform(ParentToWorld);
-						PSC->RelativeLocation = RelativeTM.GetLocation();
-						PSC->RelativeRotation = RelativeTM.GetRotation().Rotator();
-						PSC->RelativeScale3D = RelativeTM.GetScale3D();
+						PSC->SetRelativeLocation_Direct(RelativeTM.GetLocation());
+						PSC->SetRelativeRotation_Direct(RelativeTM.GetRotation().Rotator());
+						PSC->SetRelativeScale3D_Direct(RelativeTM.GetScale3D());
 					}
 					else
 					{
-						PSC->RelativeLocation = Location;
-						PSC->RelativeRotation = Rotation;
+						PSC->SetRelativeLocation_Direct(Location);
+						PSC->SetRelativeRotation_Direct(Rotation);
 
 						if (LocationType == EAttachLocation::SnapToTarget)
 						{
 							// SnapToTarget indicates we "keep world scale", this indicates we we want the inverse of the parent-to-world scale 
 							// to calculate world scale at Scale 1, and then apply the passed in Scale
 							const FTransform ParentToWorld = AttachToComponent->GetSocketTransform(AttachPointName);
-							PSC->RelativeScale3D = Scale * ParentToWorld.GetSafeScaleReciprocal(ParentToWorld.GetScale3D());
+							PSC->SetRelativeScale3D_Direct(Scale * ParentToWorld.GetSafeScaleReciprocal(ParentToWorld.GetScale3D()));
 						}
 						else
 						{
-							PSC->RelativeScale3D = Scale;
+							PSC->SetRelativeScale3D_Direct(Scale);
 						}
 					}
 
diff --git a/Engine/Plugins/FX/Niagara/Source/Niagara/Private/NiagaraPreviewGrid.cpp b/Engine/Plugins/FX/Niagara/Source/Niagara/Private/NiagaraPreviewGrid.cpp
--- a/Engine/Plugins/FX/Niagara/Source/Niagara/Private/NiagaraPreviewGrid.cpp
+++ b/Engine/Plugins/FX/Niagara/Source/Niagara/Private/NiagaraPreviewGrid.cpp
@@ -80,7 +80,7 @@ ANiagaraPreviewGrid::ANiagaraPreviewGrid(const FObjectInitializer& ObjectInitial
 		if (SpriteComponent)
 		{
 			SpriteComponent->Sprite = ConstructorStatics.SpriteTextureObject.Get();
-			SpriteComponent->RelativeScale3D = FVector(0.5f, 0.5f, 0.5f);
+			SpriteComponent->SetRelativeScale3D_Direct(FVector(0.5f, 0.5f, 0.5f));
 			SpriteComponent->bHiddenInGame = false;
 			SpriteComponent->bIsScreenSizeScaled = true;
 			SpriteComponent->SpriteInfo.Category = ConstructorStatics.ID_Effects;
@@ -100,7 +100,7 @@ ANiagaraPreviewGrid::ANiagaraPreviewGrid(const FObjectInitializer& ObjectInitial
 			ArrowComponent->SpriteInfo.Category = ConstructorStatics.ID_Effects;
 			ArrowComponent->SpriteInfo.DisplayName = ConstructorStatics.NAME_Effects;
 			ArrowComponent->SetupAttachment(RootComponent);
-			ArrowComponent->bAbsoluteScale = true;
+			ArrowComponent->SetUsingAbsoluteScale(true);
 		}
 	}
 #endif // WITH_EDITORONLY_DATA
diff --git a/Engine/Plugins/FX/Niagara/Source/NiagaraAnimNotifies/Private/AnimNotifyState_TimedNiagaraEffect.cpp b/Engine/Plugins/FX/Niagara/Source/NiagaraAnimNotifies/Private/AnimNotifyState_TimedNiagaraEffect.cpp
--- a/Engine/Plugins/FX/Niagara/Source/NiagaraAnimNotifies/Private/AnimNotifyState_TimedNiagaraEffect.cpp
+++ b/Engine/Plugins/FX/Niagara/Source/NiagaraAnimNotifies/Private/AnimNotifyState_TimedNiagaraEffect.cpp
@@ -58,7 +58,7 @@ void UAnimNotifyState_TimedNiagaraEffect::NotifyEnd(USkeletalMeshComponent * Mes
 			bTemplateMatch |= PreviousTemplates.Contains(FXSystemComponent->GetFXSystemAsset());
 #endif
 
-			if (bSocketMatch && bTemplateMatch && FXSystemComponent->bIsActive)
+			if (bSocketMatch && bTemplateMatch && FXSystemComponent->IsActive())
 			{
 				// Either destroy the component or deactivate it to have it's active FXSystems finish.
 				// The component will auto destroy once all FXSystem are gone.
diff --git a/Engine/Plugins/FX/Niagara/Source/NiagaraAnimNotifies/Private/AnimNotify_PlayNiagaraEffect.cpp b/Engine/Plugins/FX/Niagara/Source/NiagaraAnimNotifies/Private/AnimNotify_PlayNiagaraEffect.cpp
--- a/Engine/Plugins/FX/Niagara/Source/NiagaraAnimNotifies/Private/AnimNotify_PlayNiagaraEffect.cpp
+++ b/Engine/Plugins/FX/Niagara/Source/NiagaraAnimNotifies/Private/AnimNotify_PlayNiagaraEffect.cpp
@@ -112,7 +112,7 @@ UFXSystemComponent* UAnimNotify_PlayNiagaraEffect::SpawnEffect(USkeletalMeshComp
 			ReturnComp = UNiagaraFunctionLibrary::SpawnSystemAtLocation(MeshComp->GetWorld(), Template, MeshTransform.TransformPosition(LocationOffset), (MeshTransform.GetRotation() * RotationOffsetQuat).Rotator(), FVector(1.0f),true);
 		}
 
-		ReturnComp->RelativeScale3D = Scale;
+		ReturnComp->SetRelativeScale3D_Direct(Scale);
 	}
 
 	return ReturnComp;
diff --git a/Engine/Plugins/MovieScene/LevelSequenceEditor/Source/LevelSequenceEditor/Private/Misc/LevelSequenceEditorSpawnRegister.cpp b/Engine/Plugins/MovieScene/LevelSequenceEditor/Source/LevelSequenceEditor/Private/Misc/LevelSequenceEditorSpawnRegister.cpp
--- a/Engine/Plugins/MovieScene/LevelSequenceEditor/Source/LevelSequenceEditor/Private/Misc/LevelSequenceEditorSpawnRegister.cpp
+++ b/Engine/Plugins/MovieScene/LevelSequenceEditor/Source/LevelSequenceEditor/Private/Misc/LevelSequenceEditorSpawnRegister.cpp
@@ -329,9 +329,9 @@ void FLevelSequenceEditorSpawnRegister::HandleConvertPossessableToSpawnable(UObj
 	if (OldActor)
 	{
 		OutTransformData.Emplace();
-		OutTransformData->Translation = OldActor->GetRootComponent()->RelativeLocation;
-		OutTransformData->Rotation = OldActor->GetRootComponent()->RelativeRotation;
-		OutTransformData->Scale = OldActor->GetRootComponent()->RelativeScale3D;
+		OutTransformData->Translation = OldActor->GetRootComponent()->GetRelativeLocation();
+		OutTransformData->Rotation = OldActor->GetRootComponent()->GetRelativeRotation();
+		OutTransformData->Scale = OldActor->GetRootComponent()->GetRelativeScale3D();
 
 		GEditor->SelectActor(OldActor, false, true);
 		UWorld* World = Cast<UWorld>(Player.GetPlaybackContext());
diff --git a/Engine/Plugins/NetcodeUnitTest/NetcodeUnitTest/Source/NetcodeUnitTest/Private/NUTActor.cpp b/Engine/Plugins/NetcodeUnitTest/NetcodeUnitTest/Source/NetcodeUnitTest/Private/NUTActor.cpp
--- a/Engine/Plugins/NetcodeUnitTest/NetcodeUnitTest/Source/NetcodeUnitTest/Private/NUTActor.cpp
+++ b/Engine/Plugins/NetcodeUnitTest/NetcodeUnitTest/Source/NetcodeUnitTest/Private/NUTActor.cpp
@@ -38,7 +38,7 @@ ANUTActor::ANUTActor(const FObjectInitializer& ObjectInitializer)
 
 	bReplicates = true;
 	bAlwaysRelevant = true;
-	bReplicateMovement = false;
+	SetReplicatingMovement(false);
 	NetUpdateFrequency = 1;
 }
 
diff --git a/Engine/Plugins/NetcodeUnitTest/NetcodeUnitTest/Source/NetcodeUnitTest/Private/Net/UnitTestActorChannel.cpp b/Engine/Plugins/NetcodeUnitTest/NetcodeUnitTest/Source/NetcodeUnitTest/Private/Net/UnitTestActorChannel.cpp
--- a/Engine/Plugins/NetcodeUnitTest/NetcodeUnitTest/Source/NetcodeUnitTest/Private/Net/UnitTestActorChannel.cpp
+++ b/Engine/Plugins/NetcodeUnitTest/NetcodeUnitTest/Source/NetcodeUnitTest/Private/Net/UnitTestActorChannel.cpp
@@ -73,7 +73,7 @@ void UUnitTestActorChannel::NotifyActorChannelOpen(AActor* InActor, FInBunch& In
 		if (NetPlayerIndex == 0 && Connection->Driver != nullptr && Connection == Connection->Driver->ServerConnection)
 		{
 			// Implement only essential parts of the original function, as we want to block most of it (triggers level change code)
-			PC->Role = ROLE_AutonomousProxy;
+			PC->SetRole(ROLE_AutonomousProxy);
 			PC->NetConnection = Connection;
 
 			Connection->PlayerController = PC;
diff --git a/Engine/Plugins/Online/OnlineFramework/Source/Lobby/Private/LobbyBeaconPlayerState.cpp b/Engine/Plugins/Online/OnlineFramework/Source/Lobby/Private/LobbyBeaconPlayerState.cpp
--- a/Engine/Plugins/Online/OnlineFramework/Source/Lobby/Private/LobbyBeaconPlayerState.cpp
+++ b/Engine/Plugins/Online/OnlineFramework/Source/Lobby/Private/LobbyBeaconPlayerState.cpp
@@ -18,9 +18,6 @@ void ALobbyBeaconPlayerState::PostInitializeComponents()
 {
 	Super::PostInitializeComponents();
 
-	if (Role == ROLE_Authority)
-	{
-	}
 }
 
 void ALobbyBeaconPlayerState::GetLifetimeReplicatedProps(TArray< FLifetimeProperty >& OutLifetimeProps) const
diff --git a/Engine/Plugins/Online/OnlineFramework/Source/Lobby/Private/LobbyBeaconState.cpp b/Engine/Plugins/Online/OnlineFramework/Source/Lobby/Private/LobbyBeaconState.cpp
--- a/Engine/Plugins/Online/OnlineFramework/Source/Lobby/Private/LobbyBeaconState.cpp
+++ b/Engine/Plugins/Online/OnlineFramework/Source/Lobby/Private/LobbyBeaconState.cpp
@@ -161,7 +161,7 @@ void ALobbyBeaconState::PostInitializeComponents()
 {
 	Super::PostInitializeComponents();
 
-	if (Role == ROLE_Authority)
+	if (GetLocalRole() == ROLE_Authority)
 	{
 		FTimerDelegate TimerDelegate = FTimerDelegate::CreateUObject(this, &ALobbyBeaconState::OneSecTick);
 		GetWorldTimerManager().SetTimer(OneSecTimerHandle, TimerDelegate, 1.0f, true);
@@ -187,7 +187,7 @@ void ALobbyBeaconState::GetLifetimeReplicatedProps(TArray< FLifetimeProperty >&
 
 void ALobbyBeaconState::StartWaiting()
 {
-	if (Role == ROLE_Authority)
+	if (GetLocalRole() == ROLE_Authority)
 	{
 		WaitForPlayersTimeRemaining = GetClass()->GetDefaultObject<ALobbyBeaconState>()->WaitForPlayersTimeRemaining;
 		OnRep_WaitForPlayersTimeRemaining();
@@ -198,7 +198,7 @@ void ALobbyBeaconState::StartWaiting()
 
 void ALobbyBeaconState::StartLobby()
 {
-	if (Role == ROLE_Authority)
+	if (GetLocalRole() == ROLE_Authority)
 	{
 		WaitForPlayersTimeRemaining = 0.0f;
 		OnRep_WaitForPlayersTimeRemaining();
@@ -228,7 +228,7 @@ void ALobbyBeaconState::OnRep_WaitForPlayersTimeRemaining()
 
 void ALobbyBeaconState::OneSecTick()
 {
-	if (Role == ROLE_Authority)
+	if (GetLocalRole() == ROLE_Authority)
 	{
 		double CurrTime = FPlatformTime::Seconds();
 		double DeltaTime = (CurrTime - LastTickTime);
@@ -310,7 +310,7 @@ ALobbyBeaconPlayerState* ALobbyBeaconState::CreateNewPlayer(const FText& PlayerN
 
 ALobbyBeaconPlayerState* ALobbyBeaconState::AddPlayer(const FText& PlayerName, const FUniqueNetIdRepl& UniqueId)
 {
-	if (Role == ROLE_Authority)
+	if (GetLocalRole() == ROLE_Authority)
 	{
 		return Players.AddPlayer(PlayerName, UniqueId);
 	}
@@ -320,7 +320,7 @@ ALobbyBeaconPlayerState* ALobbyBeaconState::AddPlayer(const FText& PlayerName, c
 
 void ALobbyBeaconState::RemovePlayer(const FUniqueNetIdRepl& UniqueId)
 {
-	if (Role == ROLE_Authority)
+	if (GetLocalRole() == ROLE_Authority)
 	{
 		Players.RemovePlayer(UniqueId);
 	}
@@ -357,7 +357,7 @@ ALobbyBeaconPlayerState* ALobbyBeaconState::GetPlayer(const FUniqueNetIdRepl& Un
 
 ALobbyBeaconPlayerState* ALobbyBeaconState::GetPlayer(const AOnlineBeaconClient* ClientActor)
 {
-	if (Role == ROLE_Authority)
+	if (GetLocalRole() == ROLE_Authority)
 	{
 		return Players.GetPlayer(ClientActor);
 	}
diff --git a/Engine/Plugins/Online/OnlineSubsystemUtils/Source/OnlineSubsystemUtils/Private/OnlineBeaconClient.cpp b/Engine/Plugins/Online/OnlineSubsystemUtils/Source/OnlineSubsystemUtils/Private/OnlineBeaconClient.cpp
--- a/Engine/Plugins/Online/OnlineSubsystemUtils/Source/OnlineSubsystemUtils/Private/OnlineBeaconClient.cpp
+++ b/Engine/Plugins/Online/OnlineSubsystemUtils/Source/OnlineSubsystemUtils/Private/OnlineBeaconClient.cpp
@@ -262,7 +262,7 @@ void AOnlineBeaconClient::ClientOnConnected_Implementation()
 	SetConnectionState(EBeaconConnectionState::Open);
 	BeaconConnection->State = USOCK_Open;
 
-	Role = ROLE_Authority;
+	SetRole(ROLE_Authority);
 	SetReplicates(true);
 	SetAutonomousProxy(true);
 
diff --git a/Engine/Plugins/Online/OnlineSubsystemUtils/Source/OnlineSubsystemUtils/Private/OnlineBeaconHost.cpp b/Engine/Plugins/Online/OnlineSubsystemUtils/Source/OnlineSubsystemUtils/Private/OnlineBeaconHost.cpp
--- a/Engine/Plugins/Online/OnlineSubsystemUtils/Source/OnlineSubsystemUtils/Private/OnlineBeaconHost.cpp
+++ b/Engine/Plugins/Online/OnlineSubsystemUtils/Source/OnlineSubsystemUtils/Private/OnlineBeaconHost.cpp
@@ -174,7 +174,7 @@ void AOnlineBeaconHost::NotifyControlMessage(UNetConnection* Connection, uint8 M
 								Connection->PlayerId = UniqueId;
 								Connection->OwningActor = NewClientActor;
 								Connection->SetClientLoginState(EClientLoginState::ReceivedJoin);
-								NewClientActor->Role = ROLE_Authority;
+								NewClientActor->SetRole(ROLE_Authority);
 								NewClientActor->SetReplicates(false);
 								check(NetDriverName == NetDriver->NetDriverName);
 								NewClientActor->SetNetDriverName(NetDriverName);
diff --git a/Engine/Plugins/Runtime/ApexDestruction/Source/ApexDestruction/Private/DestructibleComponent.cpp b/Engine/Plugins/Runtime/ApexDestruction/Source/ApexDestruction/Private/DestructibleComponent.cpp
--- a/Engine/Plugins/Runtime/ApexDestruction/Source/ApexDestruction/Private/DestructibleComponent.cpp
+++ b/Engine/Plugins/Runtime/ApexDestruction/Source/ApexDestruction/Private/DestructibleComponent.cpp
@@ -45,7 +45,7 @@ UDestructibleComponent::UDestructibleComponent(const FObjectInitializer& ObjectI
 	SetCollisionProfileName(CollisionProfileName);
 
 	bAlwaysCreatePhysicsState = true;
-	bIsActive = true;
+	SetActiveFlag(true);
 	bMultiBodyOverlap = true;
 
 	LargeChunkThreshold = 25.f;
@@ -1159,7 +1159,7 @@ void UDestructibleComponent::Activate( bool bReset/*=false*/ )
 {
 	if (bReset || ShouldActivate()==true)
 	{
-		bIsActive = true;
+		SetActiveFlag(true);
 	}
 }
 
@@ -1167,7 +1167,7 @@ void UDestructibleComponent::Deactivate()
 {
 	if (ShouldActivate()==false)
 	{
-		bIsActive = false;
+		SetActiveFlag(false);
 	}
 }
 
diff --git a/Engine/Plugins/Runtime/GameplayAbilities/Source/GameplayAbilities/Private/Abilities/GameplayAbility.cpp b/Engine/Plugins/Runtime/GameplayAbilities/Source/GameplayAbilities/Private/Abilities/GameplayAbility.cpp
--- a/Engine/Plugins/Runtime/GameplayAbilities/Source/GameplayAbilities/Private/Abilities/GameplayAbility.cpp
+++ b/Engine/Plugins/Runtime/GameplayAbilities/Source/GameplayAbilities/Private/Abilities/GameplayAbility.cpp
@@ -289,7 +289,7 @@ bool UGameplayAbility::CanActivateAbility(const FGameplayAbilitySpecHandle Handl
 
 	// A valid AvatarActor is required. Simulated proxy check means only authority or autonomous proxies should be executing abilities.
 	AActor* const AvatarActor = ActorInfo ? ActorInfo->AvatarActor.Get() : nullptr;
-	if (AvatarActor == nullptr || !ShouldActivateAbility(AvatarActor->Role))
+	if (AvatarActor == nullptr || !ShouldActivateAbility(AvatarActor->GetLocalRole()))
 	{
 		return false;
 	}
diff --git a/Engine/Plugins/Runtime/GameplayAbilities/Source/GameplayAbilities/Private/Abilities/GameplayAbilityWorldReticle_ActorVisualization.cpp b/Engine/Plugins/Runtime/GameplayAbilities/Source/GameplayAbilities/Private/Abilities/GameplayAbilityWorldReticle_ActorVisualization.cpp
--- a/Engine/Plugins/Runtime/GameplayAbilities/Source/GameplayAbilities/Private/Abilities/GameplayAbilityWorldReticle_ActorVisualization.cpp
+++ b/Engine/Plugins/Runtime/GameplayAbilities/Source/GameplayAbilities/Private/Abilities/GameplayAbilityWorldReticle_ActorVisualization.cpp
@@ -19,7 +19,7 @@ AGameplayAbilityWorldReticle_ActorVisualization::AGameplayAbilityWorldReticle_Ac
 	CollisionComponent = CreateDefaultSubobject<UCapsuleComponent>(TEXT("CollisionCapsule0"));
 	CollisionComponent->InitCapsuleSize(0.f, 0.f);
 	CollisionComponent->AlwaysLoadOnClient = true;
-	CollisionComponent->bAbsoluteScale = true;
+	CollisionComponent->SetUsingAbsoluteScale(true);
 	//CollisionComponent->AlwaysLoadOnServer = true;
 	CollisionComponent->SetCanEverAffectNavigation(false);
 	CollisionComponent->SetCollisionEnabled(ECollisionEnabled::NoCollision);
diff --git a/Engine/Plugins/Runtime/GameplayAbilities/Source/GameplayAbilities/Private/Abilities/Tasks/AbilityTask_PlayMontageAndWait.cpp b/Engine/Plugins/Runtime/GameplayAbilities/Source/GameplayAbilities/Private/Abilities/Tasks/AbilityTask_PlayMontageAndWait.cpp
--- a/Engine/Plugins/Runtime/GameplayAbilities/Source/GameplayAbilities/Private/Abilities/Tasks/AbilityTask_PlayMontageAndWait.cpp
+++ b/Engine/Plugins/Runtime/GameplayAbilities/Source/GameplayAbilities/Private/Abilities/Tasks/AbilityTask_PlayMontageAndWait.cpp
@@ -22,8 +22,8 @@ void UAbilityTask_PlayMontageAndWait::OnMontageBlendingOut(UAnimMontage* Montage
 
 			// Reset AnimRootMotionTranslationScale
 			ACharacter* Character = Cast<ACharacter>(GetAvatarActor());
-			if (Character && (Character->Role == ROLE_Authority ||
-							  (Character->Role == ROLE_AutonomousProxy && Ability->GetNetExecutionPolicy() == EGameplayAbilityNetExecutionPolicy::LocalPredicted)))
+			if (Character && (Character->GetLocalRole() == ROLE_Authority ||
+							  (Character->GetLocalRole() == ROLE_AutonomousProxy && Ability->GetNetExecutionPolicy() == EGameplayAbilityNetExecutionPolicy::LocalPredicted)))
 			{
 				Character->SetAnimRootMotionTranslationScale(1.f);
 			}
@@ -120,8 +120,8 @@ void UAbilityTask_PlayMontageAndWait::Activate()
 				AnimInstance->Montage_SetEndDelegate(MontageEndedDelegate, MontageToPlay);
 
 				ACharacter* Character = Cast<ACharacter>(GetAvatarActor());
-				if (Character && (Character->Role == ROLE_Authority ||
-								  (Character->Role == ROLE_AutonomousProxy && Ability->GetNetExecutionPolicy() == EGameplayAbilityNetExecutionPolicy::LocalPredicted)))
+				if (Character && (Character->GetLocalRole() == ROLE_Authority ||
+								  (Character->GetLocalRole() == ROLE_AutonomousProxy && Ability->GetNetExecutionPolicy() == EGameplayAbilityNetExecutionPolicy::LocalPredicted)))
 				{
 					Character->SetAnimRootMotionTranslationScale(AnimRootMotionTranslationScale);
 				}
diff --git a/Engine/Plugins/Runtime/GameplayAbilities/Source/GameplayAbilities/Private/AbilitySystemComponent.cpp b/Engine/Plugins/Runtime/GameplayAbilities/Source/GameplayAbilities/Private/AbilitySystemComponent.cpp
--- a/Engine/Plugins/Runtime/GameplayAbilities/Source/GameplayAbilities/Private/AbilitySystemComponent.cpp
+++ b/Engine/Plugins/Runtime/GameplayAbilities/Source/GameplayAbilities/Private/AbilitySystemComponent.cpp
@@ -2110,24 +2110,53 @@ void UAbilitySystemComponent::Debug_Internal(FAbilitySystemComponentDebugInfo& I
 	{
 		FString DebugTitle("");
 		// Category
-		if (Info.bShowAbilities) DebugTitle += TEXT("ABILITIES ");
-		if (Info.bShowAttributes) DebugTitle += TEXT("ATTRIBUTES ");
-		if (Info.bShowGameplayEffects) DebugTitle += TEXT("GAMEPLAYEFFECTS ");
+		if (Info.bShowAbilities)
+		{
+			DebugTitle += TEXT("ABILITIES ");
+		}
+		if (Info.bShowAttributes)
+		{
+			DebugTitle += TEXT("ATTRIBUTES ");
+		}
+		if (Info.bShowGameplayEffects)
+		{
+			DebugTitle += TEXT("GAMEPLAYEFFECTS ");
+		}
 		// Avatar info
 		if (AvatarActor)
 		{
+			const ENetRole AvatarRole = AvatarActor->GetLocalRole();
 			DebugTitle += FString::Printf(TEXT("for avatar %s "), *AvatarActor->GetName());
-			if (AvatarActor->Role == ROLE_AutonomousProxy) DebugTitle += TEXT("(local player) ");
-			else if (AvatarActor->Role == ROLE_SimulatedProxy) DebugTitle += TEXT("(simulated) ");
-			else if (AvatarActor->Role == ROLE_Authority) DebugTitle += TEXT("(authority) ");
+			if (AvatarRole == ROLE_AutonomousProxy)
+			{
+				DebugTitle += TEXT("(local player) ");
+			}
+			else if (AvatarRole == ROLE_SimulatedProxy)
+			{
+				DebugTitle += TEXT("(simulated) ");
+			}
+			else if (AvatarRole == ROLE_Authority)
+			{
+				DebugTitle += TEXT("(authority) ");
+			}
 		}
 		// Owner info
 		if (OwnerActor && OwnerActor != AvatarActor)
 		{
+			const ENetRole OwnerRole = OwnerActor->GetLocalRole();
 			DebugTitle += FString::Printf(TEXT("for owner %s "), *OwnerActor->GetName());
-			if (OwnerActor->Role == ROLE_AutonomousProxy) DebugTitle += TEXT("(autonomous) ");
-			else if (OwnerActor->Role == ROLE_SimulatedProxy) DebugTitle += TEXT("(simulated) ");
-			else if (OwnerActor->Role == ROLE_Authority) DebugTitle += TEXT("(authority) ");
+			if (OwnerRole == ROLE_AutonomousProxy)
+			{
+				DebugTitle += TEXT("(autonomous) ");
+			}
+			else if (OwnerRole == ROLE_SimulatedProxy)
+			{
+				DebugTitle += TEXT("(simulated) ");
+			}
+			else if (OwnerRole == ROLE_Authority)
+			{
+				DebugTitle += TEXT("(authority) ");
+			}
 		}
 
 		if (Info.Canvas)
diff --git a/Engine/Plugins/Runtime/GameplayAbilities/Source/GameplayAbilities/Private/AbilitySystemComponent_Abilities.cpp b/Engine/Plugins/Runtime/GameplayAbilities/Source/GameplayAbilities/Private/AbilitySystemComponent_Abilities.cpp
--- a/Engine/Plugins/Runtime/GameplayAbilities/Source/GameplayAbilities/Private/AbilitySystemComponent_Abilities.cpp
+++ b/Engine/Plugins/Runtime/GameplayAbilities/Source/GameplayAbilities/Private/AbilitySystemComponent_Abilities.cpp
@@ -1060,7 +1060,7 @@ bool UAbilitySystemComponent::TryActivateAbility(FGameplayAbilitySpecHandle Abil
 	}
 
 		
-	ENetRole NetMode = ActorInfo->AvatarActor->Role;
+	const ENetRole NetMode = ActorInfo->AvatarActor->GetLocalRole();
 
 	// This should only come from button presses/local instigation (AI, etc).
 	if (NetMode == ROLE_SimulatedProxy)
@@ -1161,12 +1161,12 @@ bool UAbilitySystemComponent::InternalTryActivateAbility(FGameplayAbilitySpecHan
 	// Use PC netmode if its there
 	if (APlayerController* PC = ActorInfo->PlayerController.Get())
 	{
-		NetMode = PC->Role;
+		NetMode = PC->GetLocalRole();
 	}
 	// Fallback to avataractor otherwise. Edge case: avatar "dies" and becomes torn off and ROLE_Authority. We don't want to use this case (use PC role instead).
 	else if (AvatarActor)
 	{
-		NetMode = AvatarActor->Role;
+		NetMode = AvatarActor->GetLocalRole();
 	}
 
 	if (NetMode == ROLE_SimulatedProxy)
@@ -1832,7 +1832,7 @@ bool UAbilitySystemComponent::TriggerAbilityFromGameplayEvent(FGameplayAbilitySp
 		// if we're the server and this is coming from a predicted event we should check if the client has already predicted it
 		if (ScopedPredictionKey.IsValidKey()
 			&& Ability->GetNetExecutionPolicy() == EGameplayAbilityNetExecutionPolicy::LocalPredicted
-			&& ActorInfo->OwnerActor->Role == ROLE_Authority)
+			&& ActorInfo->OwnerActor->GetLocalRole() == ROLE_Authority)
 		{
 			bool bPendingClientAbilityFound = false;
 			for (auto PendingAbilityInfo : Component.PendingClientActivatedAbilities)
@@ -2354,7 +2354,7 @@ float UAbilitySystemComponent::PlayMontage(UGameplayAbility* InAnimatingAbility,
 			{
 				UE_LOG(LogRootMotion, Log, TEXT("UAbilitySystemComponent::PlayMontage %s, Role: %s")
 					, *GetNameSafe(NewAnimMontage)
-					, *UEnum::GetValueAsString(TEXT("Engine.ENetRole"), AnimInstance->GetOwningActor()->Role)
+					, *UEnum::GetValueAsString(TEXT("Engine.ENetRole"), AnimInstance->GetOwningActor()->GetLocalRole())
 					);
 			}
 
diff --git a/Engine/Plugins/Runtime/GameplayAbilities/Source/GameplayAbilities/Private/GameplayAbilityTypes.cpp b/Engine/Plugins/Runtime/GameplayAbilities/Source/GameplayAbilities/Private/GameplayAbilityTypes.cpp
--- a/Engine/Plugins/Runtime/GameplayAbilities/Source/GameplayAbilities/Private/GameplayAbilityTypes.cpp
+++ b/Engine/Plugins/Runtime/GameplayAbilities/Source/GameplayAbilities/Private/GameplayAbilityTypes.cpp
@@ -102,7 +102,7 @@ bool FGameplayAbilityActorInfo::IsNetAuthority() const
 	AActor* const OwnerActorPtr = OwnerActor.Get(/*bEvenIfPendingKill=*/ true);
 	if (OwnerActorPtr)
 	{
-		return (OwnerActorPtr->Role == ROLE_Authority);
+		return (OwnerActorPtr->GetLocalRole() == ROLE_Authority);
 	}
 
 	// This rarely happens during shutdown cases for reasons that aren't quite clear
diff --git a/Engine/Plugins/Runtime/GameplayAbilities/Source/GameplayAbilities/Public/GameplayAbilitySpec.h b/Engine/Plugins/Runtime/GameplayAbilities/Source/GameplayAbilities/Public/GameplayAbilitySpec.h
--- a/Engine/Plugins/Runtime/GameplayAbilities/Source/GameplayAbilities/Public/GameplayAbilitySpec.h
+++ b/Engine/Plugins/Runtime/GameplayAbilities/Source/GameplayAbilities/Public/GameplayAbilitySpec.h
@@ -168,7 +168,7 @@ struct GAMEPLAYABILITIES_API FGameplayAbilityActivationInfo
 	{
 		// On Init, we are either Authority or NonAuthority. We haven't been given a PredictionKey and we haven't been confirmed.
 		// NonAuthority essentially means 'I'm not sure what how I'm going to do this yet'.
-		ActivationMode = (InActor->Role == ROLE_Authority ? EGameplayAbilityActivationMode::Authority : EGameplayAbilityActivationMode::NonAuthority);
+		ActivationMode = (InActor->GetLocalRole() == ROLE_Authority ? EGameplayAbilityActivationMode::Authority : EGameplayAbilityActivationMode::NonAuthority);
 	}
 
 	FGameplayAbilityActivationInfo(EGameplayAbilityActivationMode::Type InType)
diff --git a/Engine/Plugins/Runtime/GoogleVR/GoogleVRController/Source/GoogleVRController/Private/GoogleVRMotionControllerComponent.cpp b/Engine/Plugins/Runtime/GoogleVR/GoogleVRController/Source/GoogleVRController/Private/GoogleVRMotionControllerComponent.cpp
--- a/Engine/Plugins/Runtime/GoogleVR/GoogleVRController/Source/GoogleVRController/Private/GoogleVRMotionControllerComponent.cpp
+++ b/Engine/Plugins/Runtime/GoogleVR/GoogleVRController/Source/GoogleVRController/Private/GoogleVRMotionControllerComponent.cpp
@@ -344,7 +344,7 @@ void UGoogleVRMotionControllerComponent::TickComponent( float DeltaTime, ELevelT
 
 			if (PlayerController->IsInputKeyDown(EKeys::Steam_Touch_0))
 			{
-				ControllerTouchPointMeshComponent->SetRelativeScale3D(ControllerTouchPointMeshComponent->RelativeScale3D * TOUCHPAD_POINT_FILTER_STRENGTH +
+				ControllerTouchPointMeshComponent->SetRelativeScale3D(ControllerTouchPointMeshComponent->GetRelativeScale3D() * TOUCHPAD_POINT_FILTER_STRENGTH +
 																	  TouchMeshScale * (1.0f - TOUCHPAD_POINT_FILTER_STRENGTH));
 
 				float TouchPadX = PlayerController->GetInputAnalogKeyState(FGamepadKeyNames::MotionController_Left_Thumbstick_X);
@@ -357,7 +357,7 @@ void UGoogleVRMotionControllerComponent::TickComponent( float DeltaTime, ELevelT
 			}
 			else
 			{
-				ControllerTouchPointMeshComponent->SetRelativeScale3D(ControllerTouchPointMeshComponent->RelativeScale3D * TOUCHPAD_POINT_FILTER_STRENGTH);
+				ControllerTouchPointMeshComponent->SetRelativeScale3D(ControllerTouchPointMeshComponent->GetRelativeScale3D() * TOUCHPAD_POINT_FILTER_STRENGTH);
 			}
 		}
 	}
diff --git a/Engine/Plugins/Runtime/MixedRealityCaptureFramework/Source/MixedRealityCaptureFramework/Private/MixedRealityCaptureComponent.cpp b/Engine/Plugins/Runtime/MixedRealityCaptureFramework/Source/MixedRealityCaptureFramework/Private/MixedRealityCaptureComponent.cpp
--- a/Engine/Plugins/Runtime/MixedRealityCaptureFramework/Source/MixedRealityCaptureFramework/Private/MixedRealityCaptureComponent.cpp
+++ b/Engine/Plugins/Runtime/MixedRealityCaptureFramework/Source/MixedRealityCaptureFramework/Private/MixedRealityCaptureComponent.cpp
@@ -407,7 +407,7 @@ void UMixedRealityCaptureComponent::Activate(bool bReset)
 {
 	Super::Activate(bReset);
 
-	if (bIsActive)
+	if (IsActive())
 	{
 		RefreshDevicePairing();
 		RefreshTrackingOriginOffset();
@@ -448,7 +448,7 @@ void UMixedRealityCaptureComponent::Deactivate()
 {
 	Super::Deactivate();
 
-	if (!bIsActive)
+	if (!IsActive())
 	{
 		FXRTrackingSystemDelegates::OnXRTrackingOriginChanged.RemoveAll(this);
 		MRCaptureComponent_Impl::RemoveAllCVarBindings(this);
@@ -610,7 +610,7 @@ void UMixedRealityCaptureComponent::UpdateSceneCaptureContents(FSceneInterface*
 void UMixedRealityCaptureComponent::RefreshCameraFeed()
 {
 	const UWorld* MyWorld = GetWorld();
-	if (CaptureFeedRef.DeviceURL.IsEmpty() && bIsActive && HasBeenInitialized() && MyWorld && MyWorld->IsGameWorld())
+	if (CaptureFeedRef.DeviceURL.IsEmpty() && IsActive() && HasBeenInitialized() && MyWorld && MyWorld->IsGameWorld())
 	{
 		TArray<FMediaCaptureDeviceInfo> CaptureDevices;
 		MediaCaptureSupport::EnumerateVideoCaptureDevices(CaptureDevices);
@@ -795,7 +795,7 @@ bool UMixedRealityCaptureComponent::IsTracked() const
 void UMixedRealityCaptureComponent::SetCaptureDevice(const FMrcVideoCaptureFeedIndex& FeedRef)
 {
 	const UWorld* MyWorld = GetWorld();
-	if (HasBeenInitialized() && bIsActive && MyWorld && MyWorld->IsGameWorld())
+	if (HasBeenInitialized() && IsActive() && MyWorld && MyWorld->IsGameWorld())
 	{
 		if (MediaSource)
 		{
@@ -1160,7 +1160,7 @@ bool UMixedRealityCaptureComponent::SetGarbageMatteActor(AMrcGarbageMatteActor*
 		GarbageMatteCaptureComponent->SetGarbageMatteActor(Actor);
 		bSuccess = true;
 	}
-	else if (bIsActive)
+	else if (IsActive())
 	{
 		USceneComponent* GarbageMatteOrigin = TrackingOriginOffset ? TrackingOriginOffset : (PairedTracker ? PairedTracker->GetAttachParent() : GetAttachParent());
 		GarbageMatteCaptureComponent = MRCaptureComponent_Impl::CreateGarbageMatteComponent(this, GarbageMatteOrigin);
diff --git a/Engine/Plugins/Runtime/MixedRealityCaptureFramework/Source/MixedRealityCaptureFramework/Private/MrcLatencyViewExtension.cpp b/Engine/Plugins/Runtime/MixedRealityCaptureFramework/Source/MixedRealityCaptureFramework/Private/MrcLatencyViewExtension.cpp
--- a/Engine/Plugins/Runtime/MixedRealityCaptureFramework/Source/MixedRealityCaptureFramework/Private/MrcLatencyViewExtension.cpp
+++ b/Engine/Plugins/Runtime/MixedRealityCaptureFramework/Source/MixedRealityCaptureFramework/Private/MrcLatencyViewExtension.cpp
@@ -147,5 +147,5 @@ void FMrcLatencyViewExtension::BeginRenderViewFamily(FSceneViewFamily& ViewFamil
 //------------------------------------------------------------------------------
 bool FMrcLatencyViewExtension::IsActiveThisFrame(class FViewport* InViewport) const
 {
-	return Owner.IsValid() && Owner->bIsActive && FMotionDelayClient::IsActiveThisFrame(InViewport);
+	return Owner.IsValid() && Owner->IsActive() && FMotionDelayClient::IsActiveThisFrame(InViewport);
 }
diff --git a/Engine/Plugins/Runtime/NetworkPrediction/Source/NetworkPrediction/Private/Movement/BaseMovementComponent.cpp b/Engine/Plugins/Runtime/NetworkPrediction/Source/NetworkPrediction/Private/Movement/BaseMovementComponent.cpp
--- a/Engine/Plugins/Runtime/NetworkPrediction/Source/NetworkPrediction/Private/Movement/BaseMovementComponent.cpp
+++ b/Engine/Plugins/Runtime/NetworkPrediction/Source/NetworkPrediction/Private/Movement/BaseMovementComponent.cpp
@@ -59,7 +59,7 @@ UBaseMovementComponent::UBaseMovementComponent()
 	
 	bWantsInitializeComponent = true;
 	bAutoActivate = true;
-	bReplicates = true;
+	SetIsReplicatedByDefault(true);
 	bWantsInitializeComponent = true;
 }
 
diff --git a/Engine/Plugins/Runtime/NetworkPrediction/Source/NetworkPrediction/Private/Movement/FlyingMovementComponent.cpp b/Engine/Plugins/Runtime/NetworkPrediction/Source/NetworkPrediction/Private/Movement/FlyingMovementComponent.cpp
--- a/Engine/Plugins/Runtime/NetworkPrediction/Source/NetworkPrediction/Private/Movement/FlyingMovementComponent.cpp
+++ b/Engine/Plugins/Runtime/NetworkPrediction/Source/NetworkPrediction/Private/Movement/FlyingMovementComponent.cpp
@@ -82,7 +82,7 @@ void UFlyingMovementComponent::TickComponent(float DeltaTime, enum ELevelTick Ti
 			}
 		}
 
-		Owner->bReplicateMovement = false;
+		Owner->SetReplicatingMovement(false);
 	}
 
 	Super::TickComponent(DeltaTime, TickType, ThisTickFunction);
diff --git a/Engine/Plugins/Runtime/NetworkPrediction/Source/NetworkPrediction/Private/NetworkPredictionComponent.cpp b/Engine/Plugins/Runtime/NetworkPrediction/Source/NetworkPrediction/Private/NetworkPredictionComponent.cpp
--- a/Engine/Plugins/Runtime/NetworkPrediction/Source/NetworkPrediction/Private/NetworkPredictionComponent.cpp
+++ b/Engine/Plugins/Runtime/NetworkPrediction/Source/NetworkPrediction/Private/NetworkPredictionComponent.cpp
@@ -8,7 +8,7 @@
 
 UNetworkPredictionComponent::UNetworkPredictionComponent()
 {
-	bReplicates = true;
+	SetIsReplicatedByDefault(true);
 }
 
 void UNetworkPredictionComponent::InitializeComponent()
@@ -128,7 +128,7 @@ bool UNetworkPredictionComponent::IsLocallyControlled()
 	// override this function.
 
 	bool bIsLocallyControlled = false;
-	if (GetOwner()->Role == ROLE_AutonomousProxy)
+	if (GetOwner()->GetLocalRole() == ROLE_AutonomousProxy)
 	{
 		bIsLocallyControlled = true;
 	}
@@ -142,7 +142,7 @@ bool UNetworkPredictionComponent::IsLocallyControlled()
 
 bool UNetworkPredictionComponent::CheckOwnerRoleChange()
 {
-	ENetRole CurrentRole = GetOwner()->Role;
+	ENetRole CurrentRole = GetOwner()->GetLocalRole();
 	if (CurrentRole != OwnerCachedNetRole)
 	{
 		if (OwnerCachedNetRole == ROLE_AutonomousProxy)
diff --git a/Engine/Plugins/Runtime/NetworkPrediction/Source/NetworkPrediction/Public/NetworkSimulationModelDebugger.h b/Engine/Plugins/Runtime/NetworkPrediction/Source/NetworkPrediction/Public/NetworkSimulationModelDebugger.h
--- a/Engine/Plugins/Runtime/NetworkPrediction/Source/NetworkPrediction/Public/NetworkSimulationModelDebugger.h
+++ b/Engine/Plugins/Runtime/NetworkPrediction/Source/NetworkPrediction/Public/NetworkSimulationModelDebugger.h
@@ -531,10 +531,10 @@ struct TNetworkSimulationModelDebugger : public INetworkSimulationModelDebugger
 		//	Lines
 		// ------------------------------------------------------------------------------------------------------------------------------------------------
 
-		Out.Emit(FString::Printf(TEXT("%s - %s"), *Owner->GetName(), *UEnum::GetValueAsString(TEXT("Engine.ENetRole"), Owner->Role)), FColor::Yellow);
+		Out.Emit(FString::Printf(TEXT("%s - %s"), *Owner->GetName(), *UEnum::GetValueAsString(TEXT("Engine.ENetRole"), Owner->GetLocalRole())), FColor::Yellow);
 		Out.Emit(FString::Printf(TEXT("LastProcessedInputKeyframe: %d (%d Buffered)"), NetworkSim->TickInfo.LastProcessedInputKeyframe, NetworkSim->Buffers.Input.GetHeadKeyframe() - NetworkSim->TickInfo.LastProcessedInputKeyframe));
 				
-		if (Owner->Role == ROLE_AutonomousProxy)
+		if (Owner->GetLocalRole() == ROLE_AutonomousProxy)
 		{			
 			FColor Color = FColor::White;
 			const bool FaultDetected = NetworkSim->RepProxy_Autonomous.IsReconcileFaultDetected();
@@ -568,7 +568,7 @@ struct TNetworkSimulationModelDebugger : public INetworkSimulationModelDebugger
 			FString AllowedSimulationTimeString = FString::Printf(TEXT("Allowed Simulation Time: %s. Keyframe: %d/%d/%d"), *NetworkSim->TickInfo.GetRemainingAllowedSimulationTime().ToString(), NetworkSim->TickInfo.MaxAllowedInputKeyframe, NetworkSim->TickInfo.LastProcessedInputKeyframe, NetworkSim->Buffers.Input.GetHeadKeyframe());
 			Out.Emit(*AllowedSimulationTimeString, Color);
 		}
-		else if (Owner->Role == ROLE_SimulatedProxy)
+		else if (Owner->GetLocalRole() == ROLE_SimulatedProxy)
 		{
 			FColor Color = FColor::White;
 			FString TimeString = FString::Printf(TEXT("Total Processed Simulation Time: %s. Last Serialized Simulation Time: %s. Delta: %s"), *NetworkSim->TickInfo.GetTotalProcessedSimulationTime().ToString(), *NetworkSim->RepProxy_Simulated.GetLastSerializedSimulationTime().ToString(), *(NetworkSim->RepProxy_Simulated.GetLastSerializedSimulationTime() - NetworkSim->TickInfo.GetTotalProcessedSimulationTime()).ToString());
diff --git a/Engine/Plugins/Runtime/NetworkPredictionExtras/Source/NetworkPredictionExtras/Private/MockNetworkSimulation.cpp b/Engine/Plugins/Runtime/NetworkPredictionExtras/Source/NetworkPredictionExtras/Private/MockNetworkSimulation.cpp
--- a/Engine/Plugins/Runtime/NetworkPredictionExtras/Source/NetworkPredictionExtras/Private/MockNetworkSimulation.cpp
+++ b/Engine/Plugins/Runtime/NetworkPredictionExtras/Source/NetworkPredictionExtras/Private/MockNetworkSimulation.cpp
@@ -113,7 +113,7 @@ UMockNetworkSimulationComponent::UMockNetworkSimulationComponent()
 
 	bWantsInitializeComponent = true;
 	bAutoActivate = true;
-	bReplicates = true;
+	SetIsReplicatedByDefault(true);
 
 	bWantsInitializeComponent = true;
 
diff --git a/Engine/Plugins/Runtime/TimeSynth/Source/TimeSynth/Private/TimeSynthComponent.cpp b/Engine/Plugins/Runtime/TimeSynth/Source/TimeSynth/Private/TimeSynthComponent.cpp
--- a/Engine/Plugins/Runtime/TimeSynth/Source/TimeSynth/Private/TimeSynthComponent.cpp
+++ b/Engine/Plugins/Runtime/TimeSynth/Source/TimeSynth/Private/TimeSynthComponent.cpp
@@ -585,7 +585,7 @@ FTimeSynthClipHandle UTimeSynthComponent::PlayClip(UTimeSynthClip* InClip, UTime
 		return FTimeSynthClipHandle();
 	}
 
-	if (!bIsActive)
+	if (!IsActive())
 	{
 		SetActive(true);
 	}
diff --git a/Engine/Plugins/Runtime/nDisplay/Source/DisplayCluster/Private/Game/Classes/Scene/DisplayClusterPawn.cpp b/Engine/Plugins/Runtime/nDisplay/Source/DisplayCluster/Private/Game/Classes/Scene/DisplayClusterPawn.cpp
--- a/Engine/Plugins/Runtime/nDisplay/Source/DisplayCluster/Private/Game/Classes/Scene/DisplayClusterPawn.cpp
+++ b/Engine/Plugins/Runtime/nDisplay/Source/DisplayCluster/Private/Game/Classes/Scene/DisplayClusterPawn.cpp
@@ -53,12 +53,12 @@ ADisplayClusterPawn::ADisplayClusterPawn(const FObjectInitializer& ObjectInitial
 	CameraComponent = CreateDefaultSubobject<UCameraComponent>(TEXT("DisplayCluster_camera"));
 	CameraComponent->AttachToComponent(CollisionOffsetComponent, FAttachmentTransformRules(EAttachmentRule::KeepRelative, false));
 	CameraComponent->bUsePawnControlRotation = false;
-	CameraComponent->bAbsoluteLocation = false;
-	CameraComponent->bAbsoluteRotation = false;
+	CameraComponent->SetUsingAbsoluteLocation(false);
+	CameraComponent->SetUsingAbsoluteRotation(false);
 
 	PrimaryActorTick.bCanEverTick = true;
 	bFindCameraComponentWhenViewTarget = true;
-	bCanBeDamaged = false;
+	SetCanBeDamaged(false);
 	bReplicates = false;
 	SpawnCollisionHandlingMethod = ESpawnActorCollisionHandlingMethod::AdjustIfPossibleButAlwaysSpawn;
 }
diff --git a/Engine/Plugins/Runtime/nDisplay/Source/DisplayCluster/Private/Game/Classes/Scene/DisplayClusterSceneComponentSyncParent.cpp b/Engine/Plugins/Runtime/nDisplay/Source/DisplayCluster/Private/Game/Classes/Scene/DisplayClusterSceneComponentSyncParent.cpp
--- a/Engine/Plugins/Runtime/nDisplay/Source/DisplayCluster/Private/Game/Classes/Scene/DisplayClusterSceneComponentSyncParent.cpp
+++ b/Engine/Plugins/Runtime/nDisplay/Source/DisplayCluster/Private/Game/Classes/Scene/DisplayClusterSceneComponentSyncParent.cpp
@@ -44,15 +44,15 @@ FString UDisplayClusterSceneComponentSyncParent::GetSyncId() const
 bool UDisplayClusterSceneComponentSyncParent::IsDirty() const
 {
 	USceneComponent* const pParent = GetAttachParent();
-	return (LastSyncLoc != pParent->RelativeLocation || LastSyncRot != pParent->RelativeRotation || LastSyncScale != pParent->RelativeScale3D);
+	return (LastSyncLoc != pParent->GetRelativeLocation() || LastSyncRot != pParent->GetRelativeRotation() || LastSyncScale != pParent->GetRelativeScale3D());
 }
 
 void UDisplayClusterSceneComponentSyncParent::ClearDirty()
 {
 	USceneComponent* const pParent = GetAttachParent();
-	LastSyncLoc = pParent->RelativeLocation;
-	LastSyncRot = pParent->RelativeRotation;
-	LastSyncScale = pParent->RelativeScale3D;
+	LastSyncLoc = pParent->GetRelativeLocation();
+	LastSyncRot = pParent->GetRelativeRotation();
+	LastSyncScale = pParent->GetRelativeScale3D();
 }
 
 
diff --git a/Engine/Plugins/Runtime/nDisplay/Source/DisplayCluster/Private/Game/Classes/Scene/DisplayClusterSceneComponentSyncThis.cpp b/Engine/Plugins/Runtime/nDisplay/Source/DisplayCluster/Private/Game/Classes/Scene/DisplayClusterSceneComponentSyncThis.cpp
--- a/Engine/Plugins/Runtime/nDisplay/Source/DisplayCluster/Private/Game/Classes/Scene/DisplayClusterSceneComponentSyncThis.cpp
+++ b/Engine/Plugins/Runtime/nDisplay/Source/DisplayCluster/Private/Game/Classes/Scene/DisplayClusterSceneComponentSyncThis.cpp
@@ -42,14 +42,14 @@ FString UDisplayClusterSceneComponentSyncThis::GetSyncId() const
 
 bool UDisplayClusterSceneComponentSyncThis::IsDirty() const
 {
-	return (LastSyncLoc != RelativeLocation || LastSyncRot != RelativeRotation || LastSyncScale != RelativeScale3D);
+	return (LastSyncLoc != GetRelativeLocation() || LastSyncRot != GetRelativeRotation() || LastSyncScale != GetRelativeScale3D());
 }
 
 void UDisplayClusterSceneComponentSyncThis::ClearDirty()
 {
-	LastSyncLoc = RelativeLocation;
-	LastSyncRot = RelativeRotation;
-	LastSyncScale = RelativeScale3D;
+	LastSyncLoc = GetRelativeLocation();
+	LastSyncRot = GetRelativeRotation();
+	LastSyncScale = GetRelativeScale3D();
 }
 
 
diff --git a/Engine/Plugins/VirtualProduction/Takes/Source/TakeRecorderSources/Private/TakeRecorderActorSource.cpp b/Engine/Plugins/VirtualProduction/Takes/Source/TakeRecorderSources/Private/TakeRecorderActorSource.cpp
--- a/Engine/Plugins/VirtualProduction/Takes/Source/TakeRecorderSources/Private/TakeRecorderActorSource.cpp
+++ b/Engine/Plugins/VirtualProduction/Takes/Source/TakeRecorderSources/Private/TakeRecorderActorSource.cpp
@@ -653,7 +653,7 @@ TArray<UTakeRecorderSource*> UTakeRecorderActorSource::PostRecording(ULevelSeque
 	// Force to authority role in case of capturing replicated actors
 	if (CachedObjectTemplate.IsValid() && !CachedObjectTemplate->HasAuthority())
 	{
-		CachedObjectTemplate->Role = ROLE_Authority;
+		CachedObjectTemplate->SetRole(ROLE_Authority);
 	}
 
 	// No longer need to track the Object Template that was created inside the level sequence.
diff --git a/Engine/Plugins/VirtualProduction/Takes/Source/TakeTrackRecorders/Private/MovieSceneVisibilityTrackRecorder.cpp b/Engine/Plugins/VirtualProduction/Takes/Source/TakeTrackRecorders/Private/MovieSceneVisibilityTrackRecorder.cpp
--- a/Engine/Plugins/VirtualProduction/Takes/Source/TakeTrackRecorders/Private/MovieSceneVisibilityTrackRecorder.cpp
+++ b/Engine/Plugins/VirtualProduction/Takes/Source/TakeTrackRecorders/Private/MovieSceneVisibilityTrackRecorder.cpp
@@ -71,7 +71,7 @@ void UMovieSceneVisibilityTrackRecorder::CreateTrackImpl()
 		}
 		else if (AActor* Actor = Cast<AActor>(ObjectToRecord.Get()))
 		{
-			bWasVisible = !Actor->bHidden;
+			bWasVisible = !Actor->IsHidden();
 		}
 
 		FMovieSceneBoolChannel* Channel = MovieSceneSection->GetChannelProxy().GetChannel<FMovieSceneBoolChannel>(0);
@@ -184,7 +184,7 @@ void UMovieSceneVisibilityTrackRecorder::RemoveRedundantTracks()
 		else if (AActor* Actor = Cast<AActor>(ObjectToRecord.Get()))
 		{
 			AActor* DefaultActor = CastChecked<AActor>(Actor->GetClass()->GetDefaultObject());
-			if (!DefaultActor->bHidden == bDefaultValue)
+			if (!DefaultActor->IsHidden() == bDefaultValue)
 			{
 				bRemoveSection = true;
 			}
@@ -212,7 +212,7 @@ bool UMovieSceneVisibilityTrackRecorder::IsObjectVisible() const
 	}
 	else if (AActor* Actor = Cast<AActor>(ObjectToRecord.Get()))
 	{
-		bVisible = !Actor->bHidden;
+		bVisible = !Actor->IsHidden();
 	}
 	
 	return bVisible;
diff --git a/Engine/Source/Developer/FunctionalTesting/Private/FunctionalTest.cpp b/Engine/Source/Developer/FunctionalTesting/Private/FunctionalTest.cpp
--- a/Engine/Source/Developer/FunctionalTesting/Private/FunctionalTest.cpp
+++ b/Engine/Source/Developer/FunctionalTesting/Private/FunctionalTest.cpp
@@ -107,7 +107,7 @@ AFunctionalTest::AFunctionalTest( const FObjectInitializer& ObjectInitializer )
 	PrimaryActorTick.bStartWithTickEnabled = false;
 	PrimaryActorTick.bTickEvenWhenPaused = true;
 	
-	bCanBeDamaged = false;
+	SetCanBeDamaged(false);
 
 	SpriteComponent = CreateDefaultSubobject<UBillboardComponent>(TEXT("Sprite"));
 	if (SpriteComponent)
diff --git a/Engine/Source/Developer/GameplayDebugger/Private/GameplayDebuggerExtension_Spectator.cpp b/Engine/Source/Developer/GameplayDebugger/Private/GameplayDebuggerExtension_Spectator.cpp
--- a/Engine/Source/Developer/GameplayDebugger/Private/GameplayDebuggerExtension_Spectator.cpp
+++ b/Engine/Source/Developer/GameplayDebugger/Private/GameplayDebuggerExtension_Spectator.cpp
@@ -56,7 +56,7 @@ void FGameplayDebuggerExtension_Spectator::ToggleSpectatorMode()
 		FActorSpawnParameters SpawnInfo;
 		SpawnInfo.SpawnCollisionHandlingOverride = ESpawnActorCollisionHandlingMethod::AlwaysSpawn;
 		SpawnInfo.Owner = OwnerPC->GetWorldSettings();
-		SpawnInfo.Instigator = OwnerPC->Instigator;
+		SpawnInfo.Instigator = OwnerPC->GetInstigator();
 		SpectatorControllerOb = OwnerPC->GetWorld()->SpawnActor<ADebugCameraController>(SpawnInfo);
 
 		if (SpectatorControllerOb)
diff --git a/Engine/Source/Developer/GameplayDebugger/Private/GameplayDebuggerLocalController.cpp b/Engine/Source/Developer/GameplayDebugger/Private/GameplayDebuggerLocalController.cpp
--- a/Engine/Source/Developer/GameplayDebugger/Private/GameplayDebuggerLocalController.cpp
+++ b/Engine/Source/Developer/GameplayDebugger/Private/GameplayDebuggerLocalController.cpp
@@ -633,7 +633,7 @@ void UGameplayDebuggerLocalController::OnSelectActorTick()
 		const FVector ViewDir = CameraRotation.Vector();
 		for (APawn* TestPawn  : TActorRange<APawn>(OwnerPC->GetWorld()))
 		{
-			if (!TestPawn->bHidden && TestPawn->GetActorEnableCollision() &&
+			if (!TestPawn->IsHidden() && TestPawn->GetActorEnableCollision() &&
 				!TestPawn->IsA(ASpectatorPawn::StaticClass()) &&
 				TestPawn != OwnerPC->GetPawn())
 			{
diff --git a/Engine/Source/Developer/HierarchicalLODUtilities/Private/HierarchicalLODUtilities.cpp b/Engine/Source/Developer/HierarchicalLODUtilities/Private/HierarchicalLODUtilities.cpp
--- a/Engine/Source/Developer/HierarchicalLODUtilities/Private/HierarchicalLODUtilities.cpp
+++ b/Engine/Source/Developer/HierarchicalLODUtilities/Private/HierarchicalLODUtilities.cpp
@@ -573,7 +573,7 @@ EClusterGenerationError FHierarchicalLODUtilities::ShouldGenerateCluster(AActor*
 		return EClusterGenerationError::InvalidActor;
 	}
 
-	if (Actor->bHidden)
+	if (Actor->IsHidden())
 	{
 		return EClusterGenerationError::ActorHiddenInGame;
 	}
diff --git a/Engine/Source/Developer/LogVisualizer/Private/VisualLoggerCameraController.cpp b/Engine/Source/Developer/LogVisualizer/Private/VisualLoggerCameraController.cpp
--- a/Engine/Source/Developer/LogVisualizer/Private/VisualLoggerCameraController.cpp
+++ b/Engine/Source/Developer/LogVisualizer/Private/VisualLoggerCameraController.cpp
@@ -58,7 +58,7 @@ void AVisualLoggerCameraController::PostInitializeComponents()
 	}
 	FActorSpawnParameters SpawnInfo;
 	SpawnInfo.Owner = this;
-	SpawnInfo.Instigator = Instigator;
+	SpawnInfo.Instigator = GetInstigator();
 	SpawnInfo.SpawnCollisionHandlingOverride = ESpawnActorCollisionHandlingMethod::AlwaysSpawn;
 	MyHUD = GetWorld()->SpawnActor<AVisualLoggerHUD>(SpawnInfo);
 
@@ -86,7 +86,7 @@ AVisualLoggerCameraController* AVisualLoggerCameraController::EnableCamera(UWorl
 		FActorSpawnParameters SpawnInfo;
 		SpawnInfo.SpawnCollisionHandlingOverride = ESpawnActorCollisionHandlingMethod::AlwaysSpawn;
 		SpawnInfo.Owner = PC->GetWorldSettings();
-		SpawnInfo.Instigator = PC->Instigator;
+		SpawnInfo.Instigator = PC->GetInstigator();
 		Instance = InWorld->SpawnActor<AVisualLoggerCameraController>( SpawnInfo );
 		Instance->Player = PC->Player;
 	}
diff --git a/Engine/Source/Developer/LogVisualizer/Private/VisualLoggerHUD.cpp b/Engine/Source/Developer/LogVisualizer/Private/VisualLoggerHUD.cpp
--- a/Engine/Source/Developer/LogVisualizer/Private/VisualLoggerHUD.cpp
+++ b/Engine/Source/Developer/LogVisualizer/Private/VisualLoggerHUD.cpp
@@ -17,7 +17,7 @@
 AVisualLoggerHUD::AVisualLoggerHUD(const FObjectInitializer& ObjectInitializer)
 	: Super(ObjectInitializer)
 {
-	bHidden = false;
+	SetHidden(false);
 
 	TextRenderInfo.bEnableShadow = true;
 }
diff --git a/Engine/Source/Developer/MeshMergeUtilities/Private/MeshMergeUtilities.cpp b/Engine/Source/Developer/MeshMergeUtilities/Private/MeshMergeUtilities.cpp
--- a/Engine/Source/Developer/MeshMergeUtilities/Private/MeshMergeUtilities.cpp
+++ b/Engine/Source/Developer/MeshMergeUtilities/Private/MeshMergeUtilities.cpp
@@ -3083,7 +3083,7 @@ void FMeshMergeUtilities::MergeComponentsToInstances(const TArray<UPrimitiveComp
 					{
 						ActorToCleanUp->Modify();
 						ActorToCleanUp->bIsEditorOnlyActor = true;
-						ActorToCleanUp->bHidden = true;
+						ActorToCleanUp->SetHidden(true);
 						ActorToCleanUp->bHiddenEd = true;
 						ActorToCleanUp->SetIsTemporarilyHiddenInEditor(true);
 					}
diff --git a/Engine/Source/Editor/BlueprintGraph/Private/K2Node_AddComponent.cpp b/Engine/Source/Editor/BlueprintGraph/Private/K2Node_AddComponent.cpp
--- a/Engine/Source/Editor/BlueprintGraph/Private/K2Node_AddComponent.cpp
+++ b/Engine/Source/Editor/BlueprintGraph/Private/K2Node_AddComponent.cpp
@@ -446,7 +446,7 @@ void UK2Node_AddComponent::ExpandNode(class FKismetCompilerContext& CompilerCont
 		USceneComponent* SceneCompTemplate = Cast<USceneComponent>(CompilerContext.Blueprint->FindTemplateByName(FName(*TemplateName)));
 		if (SceneCompTemplate)
 		{
-			FTransform TemplateTransform = FTransform(SceneCompTemplate->RelativeRotation, SceneCompTemplate->RelativeLocation, SceneCompTemplate->RelativeScale3D);
+			FTransform TemplateTransform = FTransform(SceneCompTemplate->GetRelativeRotation(), SceneCompTemplate->GetRelativeLocation(), SceneCompTemplate->GetRelativeScale3D());
 			DefaultValue = TemplateTransform.ToString();
 		}
 
diff --git a/Engine/Source/Editor/Cascade/Private/Cascade.cpp b/Engine/Source/Editor/Cascade/Private/Cascade.cpp
--- a/Engine/Source/Editor/Cascade/Private/Cascade.cpp
+++ b/Engine/Source/Editor/Cascade/Private/Cascade.cpp
@@ -170,9 +170,9 @@ FCascade::~FCascade()
 		UStaticMeshComponent* FloorComponent = PreviewViewport->GetViewportClient()->GetFloorComponent();
 		if (FloorComponent)
 		{
-			EditorOptions->FloorPosition = FloorComponent->RelativeLocation;
-			EditorOptions->FloorRotation = FloorComponent->RelativeRotation;
-			EditorOptions->FloorScale3D = FloorComponent->RelativeScale3D;
+			EditorOptions->FloorPosition = FloorComponent->GetRelativeLocation();
+			EditorOptions->FloorRotation = FloorComponent->GetRelativeRotation();
+			EditorOptions->FloorScale3D = FloorComponent->GetRelativeScale3D();
 
 			if (FloorComponent->GetStaticMesh())
 			{
@@ -1635,17 +1635,17 @@ void FCascade::Tick(float DeltaTime)
 	{
 		UParticleModuleVectorFieldLocal* VectorFieldModule = (UParticleModuleVectorFieldLocal*)SelectedModule;
 		LocalVectorFieldPreviewComponent->VectorField = VectorFieldModule->VectorField;
-		LocalVectorFieldPreviewComponent->RelativeLocation = VectorFieldModule->RelativeTranslation;
-		LocalVectorFieldPreviewComponent->RelativeRotation = VectorFieldModule->RelativeRotation;
-		LocalVectorFieldPreviewComponent->RelativeScale3D = VectorFieldModule->RelativeScale3D;
+		LocalVectorFieldPreviewComponent->SetRelativeLocation_Direct(VectorFieldModule->RelativeTranslation);
+		LocalVectorFieldPreviewComponent->SetRelativeRotation_Direct(VectorFieldModule->RelativeRotation);
+		LocalVectorFieldPreviewComponent->SetRelativeScale3D_Direct(VectorFieldModule->RelativeScale3D);
 		LocalVectorFieldPreviewComponent->Intensity = VectorFieldModule->Intensity;
-		LocalVectorFieldPreviewComponent->bVisible = true;
+		LocalVectorFieldPreviewComponent->SetVisibleFlag(true);
 		LocalVectorFieldPreviewComponent->bHiddenInGame = false;
 		LocalVectorFieldPreviewComponent->ReregisterComponent();
 	}
-	else if(LocalVectorFieldPreviewComponent->bVisible)
+	else if(LocalVectorFieldPreviewComponent->GetVisibleFlag())
 	{
-		LocalVectorFieldPreviewComponent->bVisible = false;
+		LocalVectorFieldPreviewComponent->SetVisibleFlag(false);
 		LocalVectorFieldPreviewComponent->ReregisterComponent();
 	}
 
@@ -3327,7 +3327,7 @@ void FCascade::ExportSelectedEmitter()
 					{
 						// Force a recache of the view relevance
 						PSysComp->bIsViewRelevanceDirty = true;
-						bool bIsActive = PSysComp->bIsActive;
+						bool bIsActive = PSysComp->IsActive();
 						PSysComp->DeactivateSystem();
 						PSysComp->ResetParticles();
 						if (bIsActive)
diff --git a/Engine/Source/Editor/Cascade/Private/CascadePreviewViewportClient.cpp b/Engine/Source/Editor/Cascade/Private/CascadePreviewViewportClient.cpp
--- a/Engine/Source/Editor/Cascade/Private/CascadePreviewViewportClient.cpp
+++ b/Engine/Source/Editor/Cascade/Private/CascadePreviewViewportClient.cpp
@@ -163,14 +163,14 @@ FCascadeEdPreviewViewportClient::FCascadeEdPreviewViewportClient(TWeakPtr<FCasca
 		FloorComponent->SetVisibility(EditorOptions->bShowFloor);
 		if (ParticleSystem)
 		{
-			FloorComponent->RelativeLocation = ParticleSystem->FloorPosition;
-			FloorComponent->RelativeRotation = ParticleSystem->FloorRotation;
+			FloorComponent->SetRelativeLocation_Direct(ParticleSystem->FloorPosition);
+			FloorComponent->SetRelativeRotation_Direct(ParticleSystem->FloorRotation);
 			FloorComponent->SetRelativeScale3D(ParticleSystem->FloorScale3D);
 		}
 		else
 		{
-			FloorComponent->RelativeLocation = EditorOptions->FloorPosition;
-			FloorComponent->RelativeRotation = EditorOptions->FloorRotation;
+			FloorComponent->SetRelativeLocation_Direct(EditorOptions->FloorPosition);
+			FloorComponent->SetRelativeRotation_Direct(EditorOptions->FloorRotation);
 			FloorComponent->SetRelativeScale3D(EditorOptions->FloorScale3D);
 		}
 
diff --git a/Engine/Source/Editor/ComponentVisualizers/Private/SpringComponentVisualizer.cpp b/Engine/Source/Editor/ComponentVisualizers/Private/SpringComponentVisualizer.cpp
--- a/Engine/Source/Editor/ComponentVisualizers/Private/SpringComponentVisualizer.cpp
+++ b/Engine/Source/Editor/ComponentVisualizers/Private/SpringComponentVisualizer.cpp
@@ -18,7 +18,7 @@ void FSpringComponentVisualizer::DrawVisualization( const UActorComponent* Compo
 		float R = FMath::Lerp(RestColor.R, CompressedColor.R, CompressionNormalized);
 		float G = FMath::Lerp(RestColor.G, CompressedColor.G, CompressionNormalized);
 		float B = FMath::Lerp(RestColor.B, CompressedColor.B, CompressionNormalized);
-		const FColor CurrentColor = SpringComp->bIsActive ? FColor(R, G, B) : DisabledColor;
+		const FColor CurrentColor = SpringComp->IsActive() ? FColor(R, G, B) : DisabledColor;
 		
 		//draw capsule for spring sweep
 		const FTransform& WorldTM = SpringComp->GetComponentToWorld();
diff --git a/Engine/Source/Editor/DetailCustomizations/Private/ActorComponentDetails.cpp b/Engine/Source/Editor/DetailCustomizations/Private/ActorComponentDetails.cpp
--- a/Engine/Source/Editor/DetailCustomizations/Private/ActorComponentDetails.cpp
+++ b/Engine/Source/Editor/DetailCustomizations/Private/ActorComponentDetails.cpp
@@ -69,7 +69,7 @@ void FActorComponentDetails::CustomizeDetails( IDetailLayoutBuilder& DetailBuild
 
 	if (bHideReplicates)
 	{
-		TSharedPtr<IPropertyHandle> ReplicatesProperty = DetailBuilder.GetProperty(GET_MEMBER_NAME_CHECKED(UActorComponent, bReplicates));
+		TSharedPtr<IPropertyHandle> ReplicatesProperty = DetailBuilder.GetProperty(UActorComponent::GetReplicatesPropertyName());
 		ReplicatesProperty->MarkHiddenByCustomization();
 	}
 }
diff --git a/Engine/Source/Editor/DetailCustomizations/Private/ComponentTransformDetails.cpp b/Engine/Source/Editor/DetailCustomizations/Private/ComponentTransformDetails.cpp
--- a/Engine/Source/Editor/DetailCustomizations/Private/ComponentTransformDetails.cpp
+++ b/Engine/Source/Editor/DetailCustomizations/Private/ComponentTransformDetails.cpp
@@ -606,15 +606,15 @@ void FComponentTransformDetails::OnSetAbsoluteTransform(ETransformField::Type Tr
 	switch (TransformField)
 	{
 	case ETransformField::Location:
-		AbsoluteProperty = FindField<UBoolProperty>(USceneComponent::StaticClass(), GET_MEMBER_NAME_CHECKED(USceneComponent, bAbsoluteLocation));
+		AbsoluteProperty = FindField<UBoolProperty>(USceneComponent::StaticClass(), USceneComponent::GetAbsoluteLocationPropertyName());
 		TransactionText = LOCTEXT("ToggleAbsoluteLocation", "Toggle Absolute Location");
 		break;
 	case ETransformField::Rotation:
-		AbsoluteProperty = FindField<UBoolProperty>(USceneComponent::StaticClass(), GET_MEMBER_NAME_CHECKED(USceneComponent, bAbsoluteRotation));
+		AbsoluteProperty = FindField<UBoolProperty>(USceneComponent::StaticClass(), USceneComponent::GetAbsoluteRotationPropertyName());
 		TransactionText = LOCTEXT("ToggleAbsoluteRotation", "Toggle Absolute Rotation");
 		break;
 	case ETransformField::Scale:
-		AbsoluteProperty = FindField<UBoolProperty>(USceneComponent::StaticClass(), GET_MEMBER_NAME_CHECKED(USceneComponent, bAbsoluteScale));
+		AbsoluteProperty = FindField<UBoolProperty>(USceneComponent::StaticClass(), USceneComponent::GetAbsoluteScalePropertyName());
 		TransactionText = LOCTEXT("ToggleAbsoluteScale", "Toggle Absolute Scale");
 		break;
 	default:
@@ -632,7 +632,7 @@ void FComponentTransformDetails::OnSetAbsoluteTransform(ETransformField::Type Tr
 			USceneComponent* SceneComponent = GetSceneComponentFromDetailsObject(Object);
 			if (SceneComponent)
 			{
-				bool bOldValue = TransformField == ETransformField::Location ? SceneComponent->bAbsoluteLocation : (TransformField == ETransformField::Rotation ? SceneComponent->bAbsoluteRotation : SceneComponent->bAbsoluteScale);
+				bool bOldValue = TransformField == ETransformField::Location ? SceneComponent->IsUsingAbsoluteLocation() : (TransformField == ETransformField::Rotation ? SceneComponent->IsUsingAbsoluteRotation() : SceneComponent->IsUsingAbsoluteScale());
 
 				if (bOldValue == bAbsoluteEnabled)
 				{
@@ -665,28 +665,28 @@ void FComponentTransformDetails::OnSetAbsoluteTransform(ETransformField::Type Tr
 				switch (TransformField)
 				{
 				case ETransformField::Location:
-					SceneComponent->bAbsoluteLocation = bAbsoluteEnabled;
+					SceneComponent->SetUsingAbsoluteLocation(bAbsoluteEnabled);
 
 					// Update RelativeLocation to maintain/stabilize position when switching between relative and world.
 					if (SceneComponent->GetAttachParent())
 					{
-						if (SceneComponent->bAbsoluteLocation)
+						if (SceneComponent->IsUsingAbsoluteLocation())
 						{
-							SceneComponent->RelativeLocation = SceneComponent->GetComponentTransform().GetTranslation();
+							SceneComponent->SetRelativeLocation_Direct(SceneComponent->GetComponentTransform().GetTranslation());
 						}
 						else
 						{
 							FTransform ParentToWorld = SceneComponent->GetAttachParent()->GetSocketTransform(SceneComponent->GetAttachSocketName());
 							FTransform RelativeTM = SceneComponent->GetComponentTransform().GetRelativeTransform(ParentToWorld);
-							SceneComponent->RelativeLocation = RelativeTM.GetTranslation();
+							SceneComponent->SetRelativeLocation_Direct(RelativeTM.GetTranslation());
 						}
 					}
 					break;
 				case ETransformField::Rotation:
-					SceneComponent->bAbsoluteRotation = bAbsoluteEnabled;
+					SceneComponent->SetUsingAbsoluteRotation(bAbsoluteEnabled);
 					break;
 				case ETransformField::Scale:
-					SceneComponent->bAbsoluteScale = bAbsoluteEnabled;
+					SceneComponent->SetUsingAbsoluteScale(bAbsoluteEnabled);
 					break;
 				}
 
@@ -759,7 +759,7 @@ struct FGetRootComponentArchetype
 EVisibility FComponentTransformDetails::GetLocationResetVisibility() const
 {
 	const USceneComponent* Archetype = FGetRootComponentArchetype::Get(SelectedObjects[0].Get());
-	const FVector Data = Archetype ? Archetype->RelativeLocation : FVector::ZeroVector;
+	const FVector Data = Archetype ? Archetype->GetRelativeLocation() : FVector::ZeroVector;
 
 	// unset means multiple differing values, so show "Reset to Default" in that case
 	return CachedLocation.IsSet() && CachedLocation.X.GetValue() == Data.X && CachedLocation.Y.GetValue() == Data.Y && CachedLocation.Z.GetValue() == Data.Z ? EVisibility::Hidden : EVisibility::Visible;
@@ -771,7 +771,7 @@ FReply FComponentTransformDetails::OnLocationResetClicked()
 	FScopedTransaction Transaction(TransactionName);
 
 	const USceneComponent* Archetype = FGetRootComponentArchetype::Get(SelectedObjects[0].Get());
-	const FVector Data = Archetype ? Archetype->RelativeLocation : FVector::ZeroVector;
+	const FVector Data = Archetype ? Archetype->GetRelativeLocation() : FVector::ZeroVector;
 
 	OnSetTransform(ETransformField::Location, EAxisList::All, Data, false, true);
 
@@ -781,7 +781,7 @@ FReply FComponentTransformDetails::OnLocationResetClicked()
 EVisibility FComponentTransformDetails::GetRotationResetVisibility() const
 {
 	const USceneComponent* Archetype = FGetRootComponentArchetype::Get(SelectedObjects[0].Get());
-	const FVector Data = Archetype ? Archetype->RelativeRotation.Euler() : FVector::ZeroVector;
+	const FVector Data = Archetype ? Archetype->GetRelativeRotation().Euler() : FVector::ZeroVector;
 
 	// unset means multiple differing values, so show "Reset to Default" in that case
 	return CachedRotation.IsSet() && CachedRotation.X.GetValue() == Data.X && CachedRotation.Y.GetValue() == Data.Y && CachedRotation.Z.GetValue() == Data.Z ? EVisibility::Hidden : EVisibility::Visible;
@@ -793,7 +793,7 @@ FReply FComponentTransformDetails::OnRotationResetClicked()
 	FScopedTransaction Transaction(TransactionName);
 
 	const USceneComponent* Archetype = FGetRootComponentArchetype::Get(SelectedObjects[0].Get());
-	const FVector Data = Archetype ? Archetype->RelativeRotation.Euler() : FVector::ZeroVector;
+	const FVector Data = Archetype ? Archetype->GetRelativeRotation().Euler() : FVector::ZeroVector;
 
 	OnSetTransform(ETransformField::Rotation, EAxisList::All, Data, false, true);
 
@@ -803,7 +803,7 @@ FReply FComponentTransformDetails::OnRotationResetClicked()
 EVisibility FComponentTransformDetails::GetScaleResetVisibility() const
 {
 	const USceneComponent* Archetype = FGetRootComponentArchetype::Get(SelectedObjects[0].Get());
-	const FVector Data = Archetype ? Archetype->RelativeScale3D : FVector(1.0f);
+	const FVector Data = Archetype ? Archetype->GetRelativeScale3D() : FVector(1.0f);
 
 	// unset means multiple differing values, so show "Reset to Default" in that case
 	return CachedScale.IsSet() && CachedScale.X.GetValue() == Data.X && CachedScale.Y.GetValue() == Data.Y && CachedScale.Z.GetValue() == Data.Z ? EVisibility::Hidden : EVisibility::Visible;
@@ -815,7 +815,7 @@ FReply FComponentTransformDetails::OnScaleResetClicked()
 	FScopedTransaction Transaction(TransactionName);
 
 	const USceneComponent* Archetype = FGetRootComponentArchetype::Get(SelectedObjects[0].Get());
-	const FVector Data = Archetype ? Archetype->RelativeScale3D : FVector(1.0f);
+	const FVector Data = Archetype ? Archetype->GetRelativeScale3D() : FVector(1.0f);
 
 	OnSetTransform(ETransformField::Scale, EAxisList::All, Data, false, true);
 
@@ -898,10 +898,10 @@ void FComponentTransformDetails::CacheTransform()
 			FVector Scale;
 			if( SceneComponent )
 			{
-				Loc = SceneComponent->RelativeLocation;
+				Loc = SceneComponent->GetRelativeLocation();
 				FRotator* FoundRotator = ObjectToRelativeRotationMap.Find(SceneComponent);
-				Rot = (bEditingRotationInUI && !Object->IsTemplate() && FoundRotator) ? *FoundRotator : SceneComponent->RelativeRotation;
-				Scale = SceneComponent->RelativeScale3D;
+				Rot = (bEditingRotationInUI && !Object->IsTemplate() && FoundRotator) ? *FoundRotator : SceneComponent->GetRelativeRotation();
+				Scale = SceneComponent->GetRelativeScale3D();
 
 				if( ObjectIndex == 0 )
 				{
@@ -914,9 +914,9 @@ void FComponentTransformDetails::CacheTransform()
 					CachedRotation.Set( Rot );
 					CachedScale.Set( Scale );
 
-					bAbsoluteLocation = SceneComponent->bAbsoluteLocation;
-					bAbsoluteScale = SceneComponent->bAbsoluteScale;
-					bAbsoluteRotation = SceneComponent->bAbsoluteRotation;
+					bAbsoluteLocation = SceneComponent->IsUsingAbsoluteLocation();
+					bAbsoluteScale = SceneComponent->IsUsingAbsoluteScale();
+					bAbsoluteRotation = SceneComponent->IsUsingAbsoluteRotation();
 				}
 				else if( CurLoc != Loc || CurRot != Rot || CurScale != Scale )
 				{
@@ -968,7 +968,7 @@ void FComponentTransformDetails::OnSetTransform(ETransformField::Type TransformF
 	{
 	case ETransformField::Location:
 		TransactionText = LOCTEXT("OnSetLocation", "Set Location");
-		ValueProperty = FindField<UProperty>(USceneComponent::StaticClass(), GET_MEMBER_NAME_CHECKED(USceneComponent, RelativeLocation));
+		ValueProperty = FindField<UProperty>(USceneComponent::StaticClass(), USceneComponent::GetRelativeLocationPropertyName());
 		
 		// Only set axis property for single axis set
 		if (Axis == EAxisList::X)
@@ -986,7 +986,7 @@ void FComponentTransformDetails::OnSetTransform(ETransformField::Type TransformF
 		break;
 	case ETransformField::Rotation:
 		TransactionText = LOCTEXT("OnSetRotation", "Set Rotation");
-		ValueProperty = FindField<UProperty>(USceneComponent::StaticClass(), GET_MEMBER_NAME_CHECKED(USceneComponent, RelativeRotation));
+		ValueProperty = FindField<UProperty>(USceneComponent::StaticClass(), USceneComponent::GetRelativeRotationPropertyName());
 		
 		// Only set axis property for single axis set
 		if (Axis == EAxisList::X)
@@ -1004,7 +1004,7 @@ void FComponentTransformDetails::OnSetTransform(ETransformField::Type TransformF
 		break;
 	case ETransformField::Scale:
 		TransactionText = LOCTEXT("OnSetScale", "Set Scale");
-		ValueProperty = FindField<UProperty>(USceneComponent::StaticClass(), GET_MEMBER_NAME_CHECKED(USceneComponent, RelativeScale3D));
+		ValueProperty = FindField<UProperty>(USceneComponent::StaticClass(), USceneComponent::GetRelativeScale3DPropertyName());
 
 		// If keep scale is set, don't set axis property
 		if (!bPreserveScaleRatio && Axis == EAxisList::X)
@@ -1055,18 +1055,18 @@ void FComponentTransformDetails::OnSetTransform(ETransformField::Type TransformF
 				switch (TransformField)
 				{
 				case ETransformField::Location:
-					OldComponentValue = SceneComponent->RelativeLocation;
+					OldComponentValue = SceneComponent->GetRelativeLocation();
 					break;
 				case ETransformField::Rotation:
 					// Pull from the actual component or from the cache
-					OldComponentValue = SceneComponent->RelativeRotation.Euler();
+					OldComponentValue = SceneComponent->GetRelativeRotation().Euler();
 					if (bEditingRotationInUI && !bIsEditingTemplateObject && ObjectToRelativeRotationMap.Find(SceneComponent))
 					{
 						OldComponentValue = ObjectToRelativeRotationMap.Find(SceneComponent)->Euler();
 					}
 					break;
 				case ETransformField::Scale:
-					OldComponentValue = SceneComponent->RelativeScale3D;
+					OldComponentValue = SceneComponent->GetRelativeScale3D();
 					break;
 				}
 
@@ -1132,13 +1132,13 @@ void FComponentTransformDetails::OnSetTransform(ETransformField::Type TransformF
 							if (!bIsEditingTemplateObject)
 							{
 								// Update local cache for restoring later
-								ObjectToRelativeRotationMap.FindOrAdd(SceneComponent) = SceneComponent->RelativeRotation;
+								ObjectToRelativeRotationMap.FindOrAdd(SceneComponent) = SceneComponent->GetRelativeRotation();
 							}
 
 							SceneComponent->SetRelativeLocation(NewComponentValue);
 
 							// Also forcibly set it as the cache may have changed it slightly
-							SceneComponent->RelativeLocation = NewComponentValue;
+							SceneComponent->SetRelativeLocation_Direct(NewComponentValue);
 
 							// If it's a template, propagate the change out to any current instances of the object
 							if (bIsEditingTemplateObject)
@@ -1299,12 +1299,12 @@ void FComponentTransformDetails::OnSetTransform(ETransformField::Type TransformF
 						if (FoundRotator)
 						{
 							FQuat OldQuat = FoundRotator->GetDenormalized().Quaternion();
-							FQuat NewQuat = SceneComponent->RelativeRotation.GetDenormalized().Quaternion();
+							FQuat NewQuat = SceneComponent->GetRelativeRotation().GetDenormalized().Quaternion();
 
 							if (OldQuat.Equals(NewQuat))
 							{
 								// Need to restore the manually set rotation as it was modified by quat conversion
-								SceneComponent->RelativeRotation = *FoundRotator;
+								SceneComponent->SetRelativeRotation_Direct(*FoundRotator);
 							}
 						}
 					}
@@ -1417,7 +1417,7 @@ void FComponentTransformDetails::OnBeginRotationSlider()
 				FScopedSwitchWorldForObject WorldSwitcher(Object);
 
 				// Add/update cached rotation value prior to slider interaction
-				ObjectToRelativeRotationMap.FindOrAdd(SceneComponent) = SceneComponent->RelativeRotation;
+				ObjectToRelativeRotationMap.FindOrAdd(SceneComponent) = SceneComponent->GetRelativeRotation();
 			}
 		}
 	}
diff --git a/Engine/Source/Editor/DetailCustomizations/Private/SceneComponentDetails.cpp b/Engine/Source/Editor/DetailCustomizations/Private/SceneComponentDetails.cpp
--- a/Engine/Source/Editor/DetailCustomizations/Private/SceneComponentDetails.cpp
+++ b/Engine/Source/Editor/DetailCustomizations/Private/SceneComponentDetails.cpp
@@ -233,12 +233,12 @@ void FSceneComponentDetails::MakeTransformDetails( IDetailLayoutBuilder& DetailB
 
 
 	// Hide the transform properties so they don't show up
-	DetailBuilder.HideProperty( DetailBuilder.GetProperty( GET_MEMBER_NAME_CHECKED(USceneComponent, bAbsoluteLocation) ) );
-	DetailBuilder.HideProperty( DetailBuilder.GetProperty( GET_MEMBER_NAME_CHECKED(USceneComponent, bAbsoluteRotation) ) );
-	DetailBuilder.HideProperty( DetailBuilder.GetProperty( GET_MEMBER_NAME_CHECKED(USceneComponent, bAbsoluteScale) ) );
-	DetailBuilder.HideProperty( DetailBuilder.GetProperty( GET_MEMBER_NAME_CHECKED(USceneComponent, RelativeLocation) ) );
-	DetailBuilder.HideProperty( DetailBuilder.GetProperty( GET_MEMBER_NAME_CHECKED(USceneComponent, RelativeRotation) ) );
-	DetailBuilder.HideProperty( DetailBuilder.GetProperty( GET_MEMBER_NAME_CHECKED(USceneComponent, RelativeScale3D) ) );
+	DetailBuilder.HideProperty(DetailBuilder.GetProperty(USceneComponent::GetAbsoluteLocationPropertyName()));
+	DetailBuilder.HideProperty(DetailBuilder.GetProperty(USceneComponent::GetAbsoluteRotationPropertyName()));
+	DetailBuilder.HideProperty(DetailBuilder.GetProperty(USceneComponent::GetAbsoluteScalePropertyName()));
+	DetailBuilder.HideProperty(DetailBuilder.GetProperty(USceneComponent::GetRelativeLocationPropertyName()));
+	DetailBuilder.HideProperty(DetailBuilder.GetProperty(USceneComponent::GetRelativeRotationPropertyName()));
+	DetailBuilder.HideProperty(DetailBuilder.GetProperty(USceneComponent::GetRelativeScale3DPropertyName()));
 
 	// Determine whether or not we are editing Class Defaults through the CDO
 	bool bIsEditingBlueprintDefaults = false;
diff --git a/Engine/Source/Editor/DetailCustomizations/Private/WorldSettingsDetails.cpp b/Engine/Source/Editor/DetailCustomizations/Private/WorldSettingsDetails.cpp
--- a/Engine/Source/Editor/DetailCustomizations/Private/WorldSettingsDetails.cpp
+++ b/Engine/Source/Editor/DetailCustomizations/Private/WorldSettingsDetails.cpp
@@ -35,7 +35,7 @@ void FWorldSettingsDetails::CustomizeDetails( IDetailLayoutBuilder& DetailBuilde
 
 	AddLightmapCustomization(DetailBuilder);
 
-	DetailBuilder.HideProperty(GET_MEMBER_NAME_CHECKED(AActor, bHidden), AActor::StaticClass());
+	DetailBuilder.HideProperty(AActor::GetHiddenPropertyName(), AActor::StaticClass());
 }
 
 
diff --git a/Engine/Source/Editor/GeometryMode/Private/GeometryModifiers.cpp b/Engine/Source/Editor/GeometryMode/Private/GeometryModifiers.cpp
--- a/Engine/Source/Editor/GeometryMode/Private/GeometryModifiers.cpp
+++ b/Engine/Source/Editor/GeometryMode/Private/GeometryModifiers.cpp
@@ -1023,7 +1023,7 @@ void UGeomModifier_Lathe::Apply( int32 InTotalSegments, int32 InSegments, EAxis:
 	BuilderBrush->Brush->Polys->Element.Empty();
 
 	// Ensure the builder brush is unhidden.
-	BuilderBrush->bHidden = false;
+	BuilderBrush->SetHidden(false);
 	BuilderBrush->bHiddenEdLayer = false;
 	BuilderBrush->SetIsTemporarilyHiddenInEditor( false );
 
@@ -1350,7 +1350,7 @@ void UGeomModifier_Pen::Apply()
 		ResultingBrush->Brush->Polys->Element.Empty();
 
 		// Ensure the brush is unhidden.
-		ResultingBrush->bHidden = false;
+		ResultingBrush->SetHidden(false);
 		ResultingBrush->bHiddenEdLayer = false;
 		ResultingBrush->SetIsTemporarilyHiddenInEditor( false );
 
diff --git a/Engine/Source/Editor/Kismet/Private/BlueprintEditor.cpp b/Engine/Source/Editor/Kismet/Private/BlueprintEditor.cpp
--- a/Engine/Source/Editor/Kismet/Private/BlueprintEditor.cpp
+++ b/Engine/Source/Editor/Kismet/Private/BlueprintEditor.cpp
@@ -8883,9 +8883,9 @@ void FBlueprintEditor::UpdatePreviewActor(UBlueprint* InBlueprint, bool bInForce
 			check(PreviewActor);
 
 			// Ensure that the actor is visible
-			if ( PreviewActor->bHidden )
+			if ( PreviewActor->IsHidden() )
 			{
-				PreviewActor->bHidden = false;
+				PreviewActor->SetHidden(false);
 				PreviewActor->MarkComponentsRenderStateDirty();				
 			}
 
diff --git a/Engine/Source/Editor/Kismet/Private/SCSEditorViewportClient.cpp b/Engine/Source/Editor/Kismet/Private/SCSEditorViewportClient.cpp
--- a/Engine/Source/Editor/Kismet/Private/SCSEditorViewportClient.cpp
+++ b/Engine/Source/Editor/Kismet/Private/SCSEditorViewportClient.cpp
@@ -451,9 +451,9 @@ bool FSCSEditorViewportClient::InputWidgetDelta( FViewport* InViewport, EAxisLis
 						if(SceneComp != NULL && SelectedTemplate != NULL)
 						{
 							// Cache the current default values for propagation
-							FVector OldRelativeLocation = SelectedTemplate->RelativeLocation;
-							FRotator OldRelativeRotation = SelectedTemplate->RelativeRotation;
-							FVector OldRelativeScale3D = SelectedTemplate->RelativeScale3D;
+							FVector OldRelativeLocation = SelectedTemplate->GetRelativeLocation();
+							FRotator OldRelativeRotation = SelectedTemplate->GetRelativeRotation();
+							FVector OldRelativeScale3D = SelectedTemplate->GetRelativeScale3D();
 
 							// Adjust the deltas as necessary
 							FComponentEditorUtils::AdjustComponentDelta(SceneComp, Drag, Rot);
@@ -473,7 +473,7 @@ bool FSCSEditorViewportClient::InputWidgetDelta( FViewport* InViewport, EAxisLis
 										&Drag,
 										&Rot,
 										&ModifiedScale,
-										SelectedTemplate->RelativeLocation );
+										SelectedTemplate->GetRelativeLocation());
 							}
 
 							UBlueprint* PreviewBlueprint = UBlueprint::GetBlueprintFromClass(PreviewActor->GetClass());
@@ -505,9 +505,9 @@ bool FSCSEditorViewportClient::InputWidgetDelta( FViewport* InViewport, EAxisLis
 										SceneComp = Cast<USceneComponent>(ArchetypeInstances[InstanceIndex]);
 										if(SceneComp != nullptr)
 										{
-											FComponentEditorUtils::ApplyDefaultValueChange(SceneComp, SceneComp->RelativeLocation, OldRelativeLocation, SelectedTemplate->RelativeLocation);
-											FComponentEditorUtils::ApplyDefaultValueChange(SceneComp, SceneComp->RelativeRotation, OldRelativeRotation, SelectedTemplate->RelativeRotation);
-											FComponentEditorUtils::ApplyDefaultValueChange(SceneComp, SceneComp->RelativeScale3D,  OldRelativeScale3D,  SelectedTemplate->RelativeScale3D);
+											FComponentEditorUtils::ApplyDefaultValueChange(SceneComp, SceneComp->GetRelativeLocation_DirectMutable(), OldRelativeLocation, SelectedTemplate->GetRelativeLocation());
+											FComponentEditorUtils::ApplyDefaultValueChange(SceneComp, SceneComp->GetRelativeRotation_DirectMutable(), OldRelativeRotation, SelectedTemplate->GetRelativeRotation());
+											FComponentEditorUtils::ApplyDefaultValueChange(SceneComp, SceneComp->GetRelativeScale3D_DirectMutable(),  OldRelativeScale3D,  SelectedTemplate->GetRelativeScale3D());
 										}
 									}
 								}
@@ -519,9 +519,9 @@ bool FSCSEditorViewportClient::InputWidgetDelta( FViewport* InViewport, EAxisLis
 										SceneComp = static_cast<USceneComponent*>(FindObjectWithOuter(ArchetypeInstances[InstanceIndex], SelectedTemplate->GetClass(), SelectedTemplate->GetFName()));
 										if(SceneComp)
 										{
-											FComponentEditorUtils::ApplyDefaultValueChange(SceneComp, SceneComp->RelativeLocation, OldRelativeLocation, SelectedTemplate->RelativeLocation);
-											FComponentEditorUtils::ApplyDefaultValueChange(SceneComp, SceneComp->RelativeRotation, OldRelativeRotation, SelectedTemplate->RelativeRotation);
-											FComponentEditorUtils::ApplyDefaultValueChange(SceneComp, SceneComp->RelativeScale3D, OldRelativeScale3D, SelectedTemplate->RelativeScale3D);
+											FComponentEditorUtils::ApplyDefaultValueChange(SceneComp, SceneComp->GetRelativeLocation_DirectMutable(), OldRelativeLocation, SelectedTemplate->GetRelativeLocation());
+											FComponentEditorUtils::ApplyDefaultValueChange(SceneComp, SceneComp->GetRelativeRotation_DirectMutable(), OldRelativeRotation, SelectedTemplate->GetRelativeRotation());
+											FComponentEditorUtils::ApplyDefaultValueChange(SceneComp, SceneComp->GetRelativeScale3D_DirectMutable(), OldRelativeScale3D, SelectedTemplate->GetRelativeScale3D());
 										}
 									}
 								}
