Landscape Layers: Change some UI to reflect naming: "Landscape Edit Layers"

#rnx

diff --git a/Engine/Source/Editor/DetailCustomizations/Private/LandscapeUIDetails.cpp b/Engine/Source/Editor/DetailCustomizations/Private/LandscapeUIDetails.cpp
--- a/Engine/Source/Editor/DetailCustomizations/Private/LandscapeUIDetails.cpp
+++ b/Engine/Source/Editor/DetailCustomizations/Private/LandscapeUIDetails.cpp
@@ -51,16 +51,17 @@ void FLandscapeUIDetails::CustomizeDetails( IDetailLayoutBuilder& DetailBuilder
 							   
 		TSharedRef<IPropertyHandle> PropertyHandle = DetailBuilder.GetProperty(GET_MEMBER_NAME_CHECKED(ALandscape, bCanHaveLayersContent));
 		DetailBuilder.HideProperty(PropertyHandle);
-		const FText DisplayAndFilterText(LOCTEXT("LandscapeToggleLayerName", "Enable Layer System"));
+		const FText DisplayAndFilterText(LOCTEXT("LandscapeToggleLayerName", "Enable Edit Layers"));
+		const FText ToolTipText(LOCTEXT("LandscapeToggleLayerToolTip", "Toggle whether or not to support edit layers on this Landscape. Toggling this will clear the undo stack."));
 		DetailBuilder.AddCustomRowToCategory(PropertyHandle, DisplayAndFilterText)
 		.NameContent()
 		[
-			PropertyHandle->CreatePropertyNameWidget(DisplayAndFilterText)
+			PropertyHandle->CreatePropertyNameWidget(DisplayAndFilterText, ToolTipText)
 		]
 		.ValueContent()
 		[
 			SNew(SCheckBox)
-			.ToolTipText(LOCTEXT("LandscapeToggleLayerToolTip", "Toggle whether or not to support layers on this Landscape and its streaming proxies. Toggling this will clear the undo stack."))
+			.ToolTipText(ToolTipText)
 			.Type(ESlateCheckBoxType::CheckBox)
 			.IsChecked_Lambda([=]()
 			{
@@ -101,11 +102,11 @@ void FLandscapeUIDetails::ToggleCanHaveLayersContent(ALandscape* Landscape)
 
 		if (bHasHiddenLayers)
 		{
-			Reason = LOCTEXT("LandscapeDisableLayers_HiddenLayers", "Are you sure you want to disable the layers system?\n\nDoing so, will result in losing the data stored for each layers, but the current visual output will be kept. Be aware that some layers are currently hidden, continuing will result in their data being lost. Undo/redo buffer will also be cleared.");
+			Reason = LOCTEXT("LandscapeDisableLayers_HiddenLayers", "Are you sure you want to disable the edit layers on this Landscape?\n\nDoing so, will result in losing the data stored for each edit layer, but the current visual output will be kept. Be aware that some edit layers are currently hidden, continuing will result in their data being lost. Undo/redo buffer will also be cleared.");
 		}
 		else
 		{
-			Reason = LOCTEXT("LandscapeDisableLayers", "Are you sure you want to disable the layers system?\n\nDoing so, will result in losing the data stored for each layers, but the current visual output will be kept. Undo/redo buffer will also be cleared.");
+			Reason = LOCTEXT("LandscapeDisableLayers", "Are you sure you want to disable the edit layers on this Landscape?\n\nDoing so, will result in losing the data stored for each edit layers, but the current visual output will be kept. Undo/redo buffer will also be cleared.");
 		}
 
 		bToggled = FMessageDialog::Open(EAppMsgType::YesNo, Reason) == EAppReturnType::Yes;
@@ -113,7 +114,7 @@ void FLandscapeUIDetails::ToggleCanHaveLayersContent(ALandscape* Landscape)
 	else
 	{
 		
-		bToggled = FMessageDialog::Open(EAppMsgType::YesNo, LOCTEXT("LandscapeEnableLayers", "Are you sure you want to enable the layer system on this landscape and streaming proxies? Doing so will clear the undo/redo buffer.")) == EAppReturnType::Yes;
+		bToggled = FMessageDialog::Open(EAppMsgType::YesNo, LOCTEXT("LandscapeEnableLayers", "Are you sure you want to enable edit layers on this landscape? Doing so will clear the undo/redo buffer.")) == EAppReturnType::Yes;
 	}
 
 	if (bToggled)
@@ -121,7 +122,7 @@ void FLandscapeUIDetails::ToggleCanHaveLayersContent(ALandscape* Landscape)
 		Landscape->ToggleCanHaveLayersContent();
 		if (GEditor)
 		{
-			GEditor->ResetTransaction(LOCTEXT("ToggleLanscapeLayers", "Toggling Landscape Layers"));
+			GEditor->ResetTransaction(LOCTEXT("ToggleLanscapeLayers", "Toggling Landscape Edit Layers"));
 		}
 	}
 }
diff --git a/Engine/Source/Editor/LandscapeEditor/Private/LandscapeEditorDetailCustomization_Layers.cpp b/Engine/Source/Editor/LandscapeEditor/Private/LandscapeEditorDetailCustomization_Layers.cpp
--- a/Engine/Source/Editor/LandscapeEditor/Private/LandscapeEditorDetailCustomization_Layers.cpp
+++ b/Engine/Source/Editor/LandscapeEditor/Private/LandscapeEditorDetailCustomization_Layers.cpp
@@ -59,7 +59,7 @@ TSharedRef<IDetailCustomization> FLandscapeEditorDetailCustomization_Layers::Mak
 BEGIN_SLATE_FUNCTION_BUILD_OPTIMIZATION
 void FLandscapeEditorDetailCustomization_Layers::CustomizeDetails(IDetailLayoutBuilder& DetailBuilder)
 {
-	IDetailCategoryBuilder& LayerCategory = DetailBuilder.EditCategory("Layers");
+	IDetailCategoryBuilder& LayerCategory = DetailBuilder.EditCategory("Edit Layers");
 
 	FEdModeLandscape* LandscapeEdMode = GetEditorMode();
 	if (LandscapeEdMode && LandscapeEdMode->CurrentToolMode != nullptr)
@@ -144,7 +144,7 @@ void FLandscapeEditorCustomNodeBuilder_Layers::GenerateChildContent(IDetailChild
 		LayerList->SetDropIndicator_Above(*FEditorStyle::GetBrush("LandscapeEditor.TargetList.DropZone.Above"));
 		LayerList->SetDropIndicator_Below(*FEditorStyle::GetBrush("LandscapeEditor.TargetList.DropZone.Below"));
 
-		ChildrenBuilder.AddCustomRow(FText::FromString(FString(TEXT("Layers"))))
+		ChildrenBuilder.AddCustomRow(FText::FromString(FString(TEXT("Edit Layers"))))
 			.Visibility(EVisibility::Visible)
 			[
 				LayerList.ToSharedRef()
@@ -428,7 +428,7 @@ TSharedPtr<SWidget> FLandscapeEditorCustomNodeBuilder_Layers::OnLayerContextMenu
 		FLandscapeLayer* Layer = LandscapeEdMode->GetLayer(InLayerIndex);
 		TSharedRef<FLandscapeEditorCustomNodeBuilder_Layers> SharedThis = AsShared();
 		FMenuBuilder MenuBuilder(true, NULL);
-		MenuBuilder.BeginSection("LandscapeEditorLayerActions", LOCTEXT("LandscapeEditorLayerActions.Heading", "Layers"));
+		MenuBuilder.BeginSection("LandscapeEditorLayerActions", LOCTEXT("LandscapeEditorLayerActions.Heading", "Edit Layers"));
 		{
 			// Create Layer
 			FUIAction CreateLayerAction = FUIAction(FExecuteAction::CreateLambda([SharedThis] { SharedThis->CreateLayer(); }), FCanExecuteAction::CreateLambda([Landscape] { return !Landscape->IsMaxLayersReached(); }));
diff --git a/Engine/Source/Editor/LandscapeEditor/Private/LandscapeEditorDetailCustomization_LayersBrushStack.cpp b/Engine/Source/Editor/LandscapeEditor/Private/LandscapeEditorDetailCustomization_LayersBrushStack.cpp
--- a/Engine/Source/Editor/LandscapeEditor/Private/LandscapeEditorDetailCustomization_LayersBrushStack.cpp
+++ b/Engine/Source/Editor/LandscapeEditor/Private/LandscapeEditorDetailCustomization_LayersBrushStack.cpp
@@ -38,12 +38,10 @@ TSharedRef<IDetailCustomization> FLandscapeEditorDetailCustomization_LayersBrush
 BEGIN_SLATE_FUNCTION_BUILD_OPTIMIZATION
 void FLandscapeEditorDetailCustomization_LayersBrushStack::CustomizeDetails(IDetailLayoutBuilder& DetailBuilder)
 {
-	IDetailCategoryBuilder& LayerCategory = DetailBuilder.EditCategory("Current Layer Brushes");
-
 	FEdModeLandscape* LandscapeEdMode = GetEditorMode();
 	if (LandscapeEdMode && LandscapeEdMode->CurrentToolMode != nullptr)
 	{
-		const FName CurrentToolName = LandscapeEdMode->CurrentTool->GetToolName();
+		IDetailCategoryBuilder& LayerCategory = DetailBuilder.EditCategory(FName("Edit Layer Blueprint Brushes"));
 
 		LayerCategory.AddCustomBuilder(MakeShareable(new FLandscapeEditorCustomNodeBuilder_LayersBrushStack(DetailBuilder.GetThumbnailPool().ToSharedRef())));
 	}
@@ -88,7 +86,7 @@ void FLandscapeEditorCustomNodeBuilder_LayersBrushStack::GenerateChildContent(ID
 		BrushesList->SetDropIndicator_Above(*FEditorStyle::GetBrush("LandscapeEditor.TargetList.DropZone.Above"));
 		BrushesList->SetDropIndicator_Below(*FEditorStyle::GetBrush("LandscapeEditor.TargetList.DropZone.Below"));
 
-		ChildrenBuilder.AddCustomRow(FText::FromString(FString(TEXT("Brush Stack"))))
+		ChildrenBuilder.AddCustomRow(FText::FromString(FString(TEXT("Edit Layer Blueprint Brushes"))))
 			.Visibility(EVisibility::Visible)
 			[
 				SNew(SVerticalBox)
@@ -223,7 +221,7 @@ TSharedPtr<SWidget> FLandscapeEditorCustomNodeBuilder_LayersBrushStack::OnBrushC
 			{
 				MenuBuilder.AddSubMenu(
 					LOCTEXT("LandscapeEditorBrushAddSubMenu", "Add"),
-					LOCTEXT("LandscapeEditorBrushAddSubMenuToolTip", "Add brush to current layer"),
+					LOCTEXT("LandscapeEditorBrushAddSubMenuToolTip", "Add brush to selected edit layer"),
 					FNewMenuDelegate::CreateSP(this, &FLandscapeEditorCustomNodeBuilder_LayersBrushStack::FillAddBrushMenu, FilteredBrushes),
 					false,
 					FSlateIcon()
@@ -236,7 +234,7 @@ TSharedPtr<SWidget> FLandscapeEditorCustomNodeBuilder_LayersBrushStack::OnBrushC
 				const FScopedTransaction Transaction(LOCTEXT("LandscapeBrushRemoveTransaction", "Remove Brush"));
 				GetEditorMode()->RemoveBrushFromCurrentLayer(CurrentBrush);
 			}));
-			MenuBuilder.AddMenuEntry(LOCTEXT("LandscapeBrushRemove", "Remove"), LOCTEXT("LandscapeBrushRemoveToolTip", "Remove brush from current layer"), FSlateIcon(), RemoveAction);
+			MenuBuilder.AddMenuEntry(LOCTEXT("LandscapeBrushRemove", "Remove"), LOCTEXT("LandscapeBrushRemoveToolTip", "Remove brush from selected edit layer"), FSlateIcon(), RemoveAction);
 		}
 		MenuBuilder.EndSection();
 		return MenuBuilder.MakeWidget();
@@ -252,7 +250,7 @@ void FLandscapeEditorCustomNodeBuilder_LayersBrushStack::FillAddBrushMenu(FMenuB
 	{
 		FUIAction AddAction = FUIAction(FExecuteAction::CreateLambda([SharedThis, Brush]()
 		{
-			const FScopedTransaction Transaction(LOCTEXT("LandscapeBrushAddToCurrentLayerTransaction", "Add brush to current layer"));
+			const FScopedTransaction Transaction(LOCTEXT("LandscapeBrushAddToCurrentLayerTransaction", "Add brush to selected edit layer"));
 			GetEditorMode()->AddBrushToCurrentLayer(Brush);
 		}));
 		MenuBuilder.AddMenuEntry(FText::FromString(Brush->GetActorLabel()), FText(), FSlateIcon(), AddAction);
diff --git a/Engine/Source/Editor/LandscapeEditor/Public/LandscapeEditorObject.h b/Engine/Source/Editor/LandscapeEditor/Public/LandscapeEditorObject.h
--- a/Engine/Source/Editor/LandscapeEditor/Public/LandscapeEditorObject.h
+++ b/Engine/Source/Editor/LandscapeEditor/Public/LandscapeEditorObject.h
@@ -526,7 +526,7 @@ private:
 	UPROPERTY(NonTransactional)
 	TArray<uint16> ImportLandscape_Data;
 public:
-	UPROPERTY(Category = "New Landscape", EditAnywhere, NonTransactional, meta = (DisplayName= "Enable Layer Support", ShowForTools= "NewLandscape"))
+	UPROPERTY(Category = "New Landscape", EditAnywhere, NonTransactional, meta = (DisplayName= "Enable Edit Layers", ToolTip="Enable support for landscape edit layers.", ShowForTools= "NewLandscape"))
 	bool bCanHaveLayersContent = false;
 
 	// Whether the imported alpha maps are to be interpreted as "layered" or "additive" (UE4 uses additive internally)
