Add/update copyright headers

diff --git a/Engine/Source/Developer/AssetTools/Private/AssetVtConversion.cpp b/Engine/Source/Developer/AssetTools/Private/AssetVtConversion.cpp
--- a/Engine/Source/Developer/AssetTools/Private/AssetVtConversion.cpp
+++ b/Engine/Source/Developer/AssetTools/Private/AssetVtConversion.cpp
@@ -1,3 +1,5 @@
+// Copyright 1998-2019 Epic Games, Inc. All Rights Reserved.
+
 #include "AssetVtConversion.h"
 #include "Engine/Texture2D.h"
 #include "Misc/PackageName.h"
diff --git a/Engine/Source/Developer/AssetTools/Private/AssetVtConversion.h b/Engine/Source/Developer/AssetTools/Private/AssetVtConversion.h
--- a/Engine/Source/Developer/AssetTools/Private/AssetVtConversion.h
+++ b/Engine/Source/Developer/AssetTools/Private/AssetVtConversion.h
@@ -1,3 +1,5 @@
+// Copyright 1998-2019 Epic Games, Inc. All Rights Reserved.
+
 #pragma once
 
 #include "CoreMinimal.h"
diff --git a/Engine/Source/Editor/DetailCustomizations/Private/MaterialExpressionTextureBaseCustomization.cpp b/Engine/Source/Editor/DetailCustomizations/Private/MaterialExpressionTextureBaseCustomization.cpp
--- a/Engine/Source/Editor/DetailCustomizations/Private/MaterialExpressionTextureBaseCustomization.cpp
+++ b/Engine/Source/Editor/DetailCustomizations/Private/MaterialExpressionTextureBaseCustomization.cpp
@@ -1,4 +1,4 @@
-// Copyright 1998-2018 Epic Games, Inc. All Rights Reserved.
+// Copyright 1998-2019 Epic Games, Inc. All Rights Reserved.
 
 #include "MaterialExpressionTextureBaseCustomization.h"
 #include "PropertyCustomizationHelpers.h"
diff --git a/Engine/Source/Editor/DetailCustomizations/Private/MaterialExpressionTextureBaseCustomization.h b/Engine/Source/Editor/DetailCustomizations/Private/MaterialExpressionTextureBaseCustomization.h
--- a/Engine/Source/Editor/DetailCustomizations/Private/MaterialExpressionTextureBaseCustomization.h
+++ b/Engine/Source/Editor/DetailCustomizations/Private/MaterialExpressionTextureBaseCustomization.h
@@ -1,4 +1,4 @@
-// Copyright 1998-2018 Epic Games, Inc. All Rights Reserved.
+// Copyright 1998-2019 Epic Games, Inc. All Rights Reserved.
 
 #pragma once
 
diff --git a/Engine/Source/Editor/UnrealEd/Classes/Factories/TextureImportSettings.h b/Engine/Source/Editor/UnrealEd/Classes/Factories/TextureImportSettings.h
--- a/Engine/Source/Editor/UnrealEd/Classes/Factories/TextureImportSettings.h
+++ b/Engine/Source/Editor/UnrealEd/Classes/Factories/TextureImportSettings.h
@@ -1,4 +1,4 @@
-// Copyright 1998-2018 Epic Games, Inc. All Rights Reserved.
+// Copyright 1998-2019 Epic Games, Inc. All Rights Reserved.
 
 #pragma once
 
diff --git a/Engine/Source/Editor/UnrealEd/Private/Factories/TextureImportSettings.cpp b/Engine/Source/Editor/UnrealEd/Private/Factories/TextureImportSettings.cpp
--- a/Engine/Source/Editor/UnrealEd/Private/Factories/TextureImportSettings.cpp
+++ b/Engine/Source/Editor/UnrealEd/Private/Factories/TextureImportSettings.cpp
@@ -1,4 +1,4 @@
-// Copyright 1998-2018 Epic Games, Inc. All Rights Reserved.
+// Copyright 1998-2019 Epic Games, Inc. All Rights Reserved.
 
 #include "Factories/TextureImportSettings.h"
 #include "HAL/IConsoleManager.h"
diff --git a/Engine/Source/Runtime/Engine/Classes/VT/VirtualTexturePoolConfig.h b/Engine/Source/Runtime/Engine/Classes/VT/VirtualTexturePoolConfig.h
--- a/Engine/Source/Runtime/Engine/Classes/VT/VirtualTexturePoolConfig.h
+++ b/Engine/Source/Runtime/Engine/Classes/VT/VirtualTexturePoolConfig.h
@@ -1,4 +1,4 @@
-// Copyright 1998-2017 Epic Games, Inc. All Rights Reserved.
+// Copyright 1998-2019 Epic Games, Inc. All Rights Reserved.
 
 #pragma once
 
diff --git a/Engine/Source/Runtime/Engine/Private/VT/VirtualTextureChunkDDCCache.cpp b/Engine/Source/Runtime/Engine/Private/VT/VirtualTextureChunkDDCCache.cpp
--- a/Engine/Source/Runtime/Engine/Private/VT/VirtualTextureChunkDDCCache.cpp
+++ b/Engine/Source/Runtime/Engine/Private/VT/VirtualTextureChunkDDCCache.cpp
@@ -1,3 +1,5 @@
+// Copyright 1998-2019 Epic Games, Inc. All Rights Reserved.
+
 #include "VirtualTextureChunkDDCCache.h"
 
 #if WITH_EDITOR
diff --git a/Engine/Source/Runtime/Engine/Private/VT/VirtualTextureChunkDDCCache.h b/Engine/Source/Runtime/Engine/Private/VT/VirtualTextureChunkDDCCache.h
--- a/Engine/Source/Runtime/Engine/Private/VT/VirtualTextureChunkDDCCache.h
+++ b/Engine/Source/Runtime/Engine/Private/VT/VirtualTextureChunkDDCCache.h
@@ -1,3 +1,5 @@
+// Copyright 1998-2019 Epic Games, Inc. All Rights Reserved.
+
 #pragma once
 
 #if WITH_EDITOR
