Material editor only displays texture sampler type enums that are valid for current virtual/non-virtual texture

diff --git a/Engine/Source/Editor/DetailCustomizations/Private/DetailCustomizations.cpp b/Engine/Source/Editor/DetailCustomizations/Private/DetailCustomizations.cpp
--- a/Engine/Source/Editor/DetailCustomizations/Private/DetailCustomizations.cpp
+++ b/Engine/Source/Editor/DetailCustomizations/Private/DetailCustomizations.cpp
@@ -134,6 +134,7 @@
 #include "SkeletalMeshReductionSettingsDetails.h"
 #include "SkeletalMeshLODSettingsDetails.h"
 #include "MaterialExpressionLandscapeGrassCustomization.h"
+#include "MaterialExpressionTextureBaseDetails.h"
 #include "TimecodeDetailsCustomization.h"
 #include "SkeletonDetails.h"
 #include "MaterialShadingModelCustomization.h"
@@ -376,6 +377,7 @@ void FDetailCustomizationsModule::RegisterObjectCustomizations()
 
 	RegisterCustomClassLayout("MaterialExpressionGetMaterialAttributes", FOnGetDetailCustomizationInstance::CreateStatic(&FMaterialAttributePropertyDetails::MakeInstance));
 	RegisterCustomClassLayout("MaterialExpressionSetMaterialAttributes", FOnGetDetailCustomizationInstance::CreateStatic(&FMaterialAttributePropertyDetails::MakeInstance));
+	RegisterCustomClassLayout("MaterialExpressionTextureBase", FOnGetDetailCustomizationInstance::CreateStatic(&FMaterialExpressionTextureBaseDetails::MakeInstance));
 	RegisterCustomClassLayout("SkeletalMeshLODSettings", FOnGetDetailCustomizationInstance::CreateStatic(&FSkeletalMeshLODSettingsDetails::MakeInstance));
 
 	RegisterCustomClassLayout("Skeleton", FOnGetDetailCustomizationInstance::CreateStatic(&FSkeletonDetails::MakeInstance));
diff --git a/Engine/Source/Editor/DetailCustomizations/Private/MaterialExpressionTextureBaseCustomization.cpp b/Engine/Source/Editor/DetailCustomizations/Private/MaterialExpressionTextureBaseCustomization.cpp
deleted file mode 100644
diff --git a/Engine/Source/Editor/DetailCustomizations/Private/MaterialExpressionTextureBaseCustomization.h b/Engine/Source/Editor/DetailCustomizations/Private/MaterialExpressionTextureBaseCustomization.h
deleted file mode 100644
diff --git a/Engine/Source/Editor/DetailCustomizations/Private/MaterialExpressionTextureBaseDetails.cpp b/Engine/Source/Editor/DetailCustomizations/Private/MaterialExpressionTextureBaseDetails.cpp
new file mode 100644
--- /dev/null
+++ b/Engine/Source/Editor/DetailCustomizations/Private/MaterialExpressionTextureBaseDetails.cpp
@@ -0,0 +1,60 @@
+// Copyright 1998-2019 Epic Games, Inc. All Rights Reserved.
+
+#include "MaterialExpressionTextureBaseDetails.h"
+#include "PropertyCustomizationHelpers.h"
+#include "DetailLayoutBuilder.h"
+#include "Materials/MaterialExpressionTextureBase.h"
+#include "Engine/Texture.h"
+#include "Engine/Texture2D.h"
+
+#define LOCTEXT_NAMESPACE "MaterialExpressionTextureBaseDetails"
+
+TSharedRef<IDetailCustomization> FMaterialExpressionTextureBaseDetails::MakeInstance()
+{
+	return MakeShareable(new FMaterialExpressionTextureBaseDetails);
+}
+
+void FMaterialExpressionTextureBaseDetails::CustomizeDetails( IDetailLayoutBuilder& DetailLayout )
+{
+	TArray<TWeakObjectPtr<UObject>> Objects;
+	DetailLayout.GetObjectsBeingCustomized(Objects);
+
+	if (Objects.Num() > 0)
+	{
+		Expression = CastChecked<UMaterialExpressionTextureBase>(Objects[0]);
+	}
+
+	EnumRestriction = MakeShareable(new FPropertyRestriction(LOCTEXT("VirtualTextureSamplerMatch", "Sampler type must match VirtualTexture usage")));
+	TSharedPtr<IPropertyHandle> SamplerTypeHandle = DetailLayout.GetProperty(GET_MEMBER_NAME_CHECKED(UMaterialExpressionTextureBase, SamplerType));
+	SamplerTypeHandle->AddRestriction(EnumRestriction.ToSharedRef());
+
+	TSharedPtr<IPropertyHandle> TextureHandle = DetailLayout.GetProperty(GET_MEMBER_NAME_CHECKED(UMaterialExpressionTextureBase, Texture));
+	TextureHandle->SetOnPropertyValueChanged(FSimpleDelegate::CreateSP(this, &FMaterialExpressionTextureBaseDetails::OnTextureChanged));
+
+	OnTextureChanged();
+}
+
+void FMaterialExpressionTextureBaseDetails::OnTextureChanged()
+{
+	bool bAllowVirtualTexture = true;
+	bool bAllowNonVirtualTexture = true;
+	if (Expression.IsValid() && Expression->Texture)
+	{
+		bAllowVirtualTexture = Expression->Texture->VirtualTextureStreaming;
+		bAllowNonVirtualTexture = !bAllowVirtualTexture;
+	}
+
+	EnumRestriction->RemoveAll();
+
+	const UEnum* MaterialSamplerTypeEnum = StaticEnum<EMaterialSamplerType>();
+	for (int SamplerTypeIndex = 0; SamplerTypeIndex < SAMPLERTYPE_MAX; ++SamplerTypeIndex)
+	{
+		const bool bIsVirtualTexture = IsVirtualSamplerType((EMaterialSamplerType)SamplerTypeIndex);
+		if ((bIsVirtualTexture && !bAllowVirtualTexture) || (!bIsVirtualTexture && !bAllowNonVirtualTexture))
+		{
+			EnumRestriction->AddHiddenValue(MaterialSamplerTypeEnum->GetNameStringByValue(SamplerTypeIndex));
+		}
+	}
+}
+
+#undef LOCTEXT_NAMESPACE
diff --git a/Engine/Source/Editor/DetailCustomizations/Private/MaterialExpressionTextureBaseDetails.h b/Engine/Source/Editor/DetailCustomizations/Private/MaterialExpressionTextureBaseDetails.h
new file mode 100644
--- /dev/null
+++ b/Engine/Source/Editor/DetailCustomizations/Private/MaterialExpressionTextureBaseDetails.h
@@ -0,0 +1,30 @@
+// Copyright 1998-2019 Epic Games, Inc. All Rights Reserved.
+
+#pragma once
+
+#include "CoreMinimal.h"
+#include "UObject/WeakObjectPtr.h"
+#include "IDetailCustomization.h"
+#include "PropertyRestriction.h"
+
+class IPropertyHandle;
+class IDetailLayoutBuilder;
+class SObjectPropertyEntryBox;
+struct FAssetData;
+class UMaterialExpressionTextureBase;
+
+class FMaterialExpressionTextureBaseDetails : public IDetailCustomization
+{
+public:
+	/** Makes a new instance of this detail layout class for a specific detail view requesting it */
+	static TSharedRef<IDetailCustomization> MakeInstance();
+
+	/** IDetailCustomization interface */
+	virtual void CustomizeDetails( IDetailLayoutBuilder& DetailLayout ) override;
+
+	TSharedPtr<FPropertyRestriction> EnumRestriction;
+	TWeakObjectPtr<UMaterialExpressionTextureBase> Expression;
+
+	void OnTextureChanged();
+};
+
