Ensure UI list of valid texture sampler types updates when texture is modified outside the material editor

diff --git a/Engine/Source/Developer/AssetTools/Private/AssetVtConversion.cpp b/Engine/Source/Developer/AssetTools/Private/AssetVtConversion.cpp
--- a/Engine/Source/Developer/AssetTools/Private/AssetVtConversion.cpp
+++ b/Engine/Source/Developer/AssetTools/Private/AssetVtConversion.cpp
@@ -499,6 +499,21 @@ void FVTConversionWorker::FilterList(int32 SizeThreshold)
 	FindAllTexturesAndMaterials(Materials, Functions, Textures);
 }
 
+static void MarkTextureExpressionModified(UMaterialExpressionTextureBase* Expression)
+{
+	FPropertyChangedEvent SamplerTypeChangeEvent(UMaterialExpressionTextureBase::StaticClass()->FindPropertyByName(GET_MEMBER_NAME_CHECKED(UMaterialExpressionTextureBase, SamplerType)));
+	FPropertyChangedEvent TextureChangeEvent(UMaterialExpressionTextureBase::StaticClass()->FindPropertyByName(GET_MEMBER_NAME_CHECKED(UMaterialExpressionTextureBase, Texture)));
+
+	Expression->Modify();
+
+	// Nofity that we changed SamplerType
+	Expression->PostEditChangeProperty(SamplerTypeChangeEvent);
+
+	// Also notify that we changed Texture (technically didn't modify this property)
+	// This way code in FMaterialExpressionTextureBaseDetails will see the event, and refresh the list of valid sampler types for the updated texture
+	Expression->PostEditChangeProperty(TextureChangeEvent);
+}
+
 void FVTConversionWorker::DoConvert(bool backwards)
 {
 	const bool bVirtualTextureEnable = !backwards;
@@ -620,8 +635,7 @@ void FVTConversionWorker::DoConvert(bool backwards)
 						}
 						if (TextureCopy != nullptr || OldType != TexExpr->SamplerType)
 						{
-							TexExpr->Modify();
-							TexExpr->PostEditChange();
+							MarkTextureExpressionModified(TexExpr);
 							MatModified = true;
 						}
 					}
@@ -687,8 +701,7 @@ void FVTConversionWorker::DoConvert(bool backwards)
 						}
 						if (TextureCopy != nullptr || TexExpr->SamplerType != OldType)
 						{
-							TexExpr->Modify();
-							TexExpr->PostEditChange();
+							MarkTextureExpressionModified(TexExpr);
 							FuncModified = true;
 						}
 					}
diff --git a/Engine/Source/Editor/DetailCustomizations/Private/MaterialExpressionTextureBaseDetails.cpp b/Engine/Source/Editor/DetailCustomizations/Private/MaterialExpressionTextureBaseDetails.cpp
--- a/Engine/Source/Editor/DetailCustomizations/Private/MaterialExpressionTextureBaseDetails.cpp
+++ b/Engine/Source/Editor/DetailCustomizations/Private/MaterialExpressionTextureBaseDetails.cpp
@@ -14,6 +14,11 @@ TSharedRef<IDetailCustomization> FMaterialExpressionTextureBaseDetails::MakeInst
 	return MakeShareable(new FMaterialExpressionTextureBaseDetails);
 }
 
+FMaterialExpressionTextureBaseDetails::~FMaterialExpressionTextureBaseDetails()
+{
+	FCoreUObjectDelegates::OnObjectPropertyChanged.Remove(DelegateHandle);
+}
+
 void FMaterialExpressionTextureBaseDetails::CustomizeDetails( IDetailLayoutBuilder& DetailLayout )
 {
 	TArray<TWeakObjectPtr<UObject>> Objects;
@@ -28,9 +33,14 @@ void FMaterialExpressionTextureBaseDetails::CustomizeDetails( IDetailLayoutBuild
 	TSharedPtr<IPropertyHandle> SamplerTypeHandle = DetailLayout.GetProperty(GET_MEMBER_NAME_CHECKED(UMaterialExpressionTextureBase, SamplerType));
 	SamplerTypeHandle->AddRestriction(EnumRestriction.ToSharedRef());
 
+	// IPropertyHandle::SetOnPropertyValueChanged will catch property changes initiated by the editor
 	TSharedPtr<IPropertyHandle> TextureHandle = DetailLayout.GetProperty(GET_MEMBER_NAME_CHECKED(UMaterialExpressionTextureBase, Texture));
 	TextureHandle->SetOnPropertyValueChanged(FSimpleDelegate::CreateSP(this, &FMaterialExpressionTextureBaseDetails::OnTextureChanged));
 
+	// FCoreUObjectDelegates::OnObjectPropertyChanged will catch property changes initiated by C++
+	// This is required to ensure UI is properly updated when the VT conversion tool runs and updates textures/materials
+	DelegateHandle = FCoreUObjectDelegates::OnObjectPropertyChanged.AddRaw(this, &FMaterialExpressionTextureBaseDetails::OnPropertyChanged);
+
 	OnTextureChanged();
 }
 
@@ -57,4 +67,12 @@ void FMaterialExpressionTextureBaseDetails::OnTextureChanged()
 	}
 }
 
+void FMaterialExpressionTextureBaseDetails::OnPropertyChanged(UObject* ObjectBeingModified, FPropertyChangedEvent& PropertyChangedEvent)
+{
+	if (ObjectBeingModified == Expression && PropertyChangedEvent.GetPropertyName() == GET_MEMBER_NAME_CHECKED(UMaterialExpressionTextureBase, Texture))
+	{
+		OnTextureChanged();
+	}
+}
+
 #undef LOCTEXT_NAMESPACE
diff --git a/Engine/Source/Editor/DetailCustomizations/Private/MaterialExpressionTextureBaseDetails.h b/Engine/Source/Editor/DetailCustomizations/Private/MaterialExpressionTextureBaseDetails.h
--- a/Engine/Source/Editor/DetailCustomizations/Private/MaterialExpressionTextureBaseDetails.h
+++ b/Engine/Source/Editor/DetailCustomizations/Private/MaterialExpressionTextureBaseDetails.h
@@ -12,6 +12,7 @@ class IDetailLayoutBuilder;
 class SObjectPropertyEntryBox;
 struct FAssetData;
 class UMaterialExpressionTextureBase;
+struct FPropertyChangedEvent;
 
 class FMaterialExpressionTextureBaseDetails : public IDetailCustomization
 {
@@ -19,12 +20,16 @@ public:
 	/** Makes a new instance of this detail layout class for a specific detail view requesting it */
 	static TSharedRef<IDetailCustomization> MakeInstance();
 
+	virtual ~FMaterialExpressionTextureBaseDetails();
+
 	/** IDetailCustomization interface */
 	virtual void CustomizeDetails( IDetailLayoutBuilder& DetailLayout ) override;
 
 	TSharedPtr<FPropertyRestriction> EnumRestriction;
 	TWeakObjectPtr<UMaterialExpressionTextureBase> Expression;
+	FDelegateHandle DelegateHandle;
 
 	void OnTextureChanged();
+	void OnPropertyChanged(UObject* ObjectBeingModified, FPropertyChangedEvent& PropertyChangedEvent);
 };
 
