Update material list of valid sampler types when texture's VT streaming property changes

diff --git a/Engine/Source/Editor/DetailCustomizations/Private/MaterialExpressionTextureBaseDetails.cpp b/Engine/Source/Editor/DetailCustomizations/Private/MaterialExpressionTextureBaseDetails.cpp
--- a/Engine/Source/Editor/DetailCustomizations/Private/MaterialExpressionTextureBaseDetails.cpp
+++ b/Engine/Source/Editor/DetailCustomizations/Private/MaterialExpressionTextureBaseDetails.cpp
@@ -71,6 +71,12 @@ void FMaterialExpressionTextureBaseDetails::OnPropertyChanged(UObject* ObjectBei
 {
 	if (ObjectBeingModified == Expression && PropertyChangedEvent.GetPropertyName() == GET_MEMBER_NAME_CHECKED(UMaterialExpressionTextureBase, Texture))
 	{
+		// Update enum list if our texture reference is changed
+		OnTextureChanged();
+	}
+	else if(Expression.IsValid() && Expression->Texture == ObjectBeingModified && PropertyChangedEvent.GetPropertyName() == GET_MEMBER_NAME_CHECKED(UTexture, VirtualTextureStreaming))
+	{
+		// Update enum list of currently assigned texture's VT streaming status is changed
 		OnTextureChanged();
 	}
 }
