Merging CL #9618651 from //UE4/Dev-Enterprise

Fixed minor issue with UI of customization of math structs

diff --git a/Engine/Source/Editor/DetailCustomizations/Private/MathStructCustomizations.cpp b/Engine/Source/Editor/DetailCustomizations/Private/MathStructCustomizations.cpp
--- a/Engine/Source/Editor/DetailCustomizations/Private/MathStructCustomizations.cpp
+++ b/Engine/Source/Editor/DetailCustomizations/Private/MathStructCustomizations.cpp
@@ -24,6 +24,7 @@ TSharedRef<IPropertyTypeCustomization> FMathStructCustomization::MakeInstance()
 
 void FMathStructCustomization::CustomizeHeader(TSharedRef<class IPropertyHandle> StructPropertyHandle, class FDetailWidgetRow& HeaderRow, IPropertyTypeCustomizationUtils& StructCustomizationUtils)
 {
+	SortedChildHandles.Empty();
 	GetSortedChildren(StructPropertyHandle, SortedChildHandles);
 	MakeHeaderRow(StructPropertyHandle, HeaderRow);
 }
@@ -146,8 +147,6 @@ void FMathStructCustomization::OnPreserveScaleRatioToggled(ECheckBoxState NewSta
 
 void FMathStructCustomization::GetSortedChildren(TSharedRef<IPropertyHandle> StructPropertyHandle, TArray< TSharedRef<IPropertyHandle> >& OutChildren)
 {
-	OutChildren.Empty();
-
 	uint32 NumChildren;
 	StructPropertyHandle->GetNumChildren(NumChildren);
 
