Fixed data table row customization using a localized default value

diff --git a/Engine/Source/Editor/DetailCustomizations/Private/DataTableCustomization.cpp b/Engine/Source/Editor/DetailCustomizations/Private/DataTableCustomization.cpp
--- a/Engine/Source/Editor/DetailCustomizations/Private/DataTableCustomization.cpp
+++ b/Engine/Source/Editor/DetailCustomizations/Private/DataTableCustomization.cpp
@@ -9,9 +9,9 @@
 
 #define LOCTEXT_NAMESPACE "FDataTableCustomizationLayout"
 
-TSharedPtr<FString> FDataTableCustomizationLayout::InitWidgetContent()
+TSharedPtr<FName> FDataTableCustomizationLayout::InitWidgetContent()
 {
-	TSharedPtr<FString> InitialValue = MakeShareable(new FString(LOCTEXT("DataTable_None", "None").ToString()));;
+	TSharedPtr<FName> InitialValue = MakeShared<FName>();
 
 	FName RowName;
 	const FPropertyAccess::Result RowResult = RowNamePropertyHandle->GetValue(RowName);
@@ -27,7 +27,7 @@ TSharedPtr<FString> FDataTableCustomizationLayout::InitWidgetContent()
 		for (TMap<FName, uint8*>::TConstIterator Iterator(DataTable->GetRowMap()); Iterator; ++Iterator)
 		{
 			/** Create a simple array of the row names */
-			TSharedRef<FString> RowNameItem = MakeShareable(new FString(Iterator.Key().ToString()));
+			TSharedRef<FName> RowNameItem = MakeShared<FName>(Iterator.Key());
 			RowNames.Add(RowNameItem);
 
 			/** Set the initial value to the currently selected item */
@@ -41,8 +41,7 @@ TSharedPtr<FString> FDataTableCustomizationLayout::InitWidgetContent()
 	/** Reset the initial value to ensure a valid entry is set */
 	if (RowResult != FPropertyAccess::MultipleValues)
 	{
-		FName NewValue = FName(**InitialValue);
-		RowNamePropertyHandle->SetValue(NewValue);
+		RowNamePropertyHandle->SetValue(*InitialValue);
 	}
 
 	return InitialValue;
@@ -134,14 +133,13 @@ void FDataTableCustomizationLayout::HandleMenuOpen()
 
 void FDataTableCustomizationLayout::OnSearchForReferences()
 {
-	if (CurrentSelectedItem.IsValid() && !CurrentSelectedItem->IsEmpty() && DataTablePropertyHandle.IsValid() && DataTablePropertyHandle->IsValidHandle())
+	if (CurrentSelectedItem.IsValid() && !CurrentSelectedItem->IsNone() && DataTablePropertyHandle.IsValid() && DataTablePropertyHandle->IsValidHandle())
 	{
 		UObject* SourceDataTable;
 		DataTablePropertyHandle->GetValue(SourceDataTable);
-		FName RowName(**CurrentSelectedItem);
 		
 		TArray<FAssetIdentifier> AssetIdentifiers;
-		AssetIdentifiers.Add(FAssetIdentifier(SourceDataTable, RowName));
+		AssetIdentifiers.Add(FAssetIdentifier(SourceDataTable, *CurrentSelectedItem));
 
 		FEditorDelegates::OnOpenReferenceViewer.Broadcast(AssetIdentifiers, FReferenceViewerParams());
 	}
@@ -149,7 +147,7 @@ void FDataTableCustomizationLayout::OnSearchForReferences()
 
 TSharedRef<SWidget> FDataTableCustomizationLayout::GetListContent()
 {
-	SAssignNew(RowNameComboListView, SListView<TSharedPtr<FString> >)
+	SAssignNew(RowNameComboListView, SListView<TSharedPtr<FName>>)
 		.ListItemsSource(&RowNames)
 		.OnSelectionChanged(this, &FDataTableCustomizationLayout::OnSelectionChanged)
 		.OnGenerateRow(this, &FDataTableCustomizationLayout::HandleRowNameComboBoxGenarateWidget)
@@ -192,22 +190,26 @@ void FDataTableCustomizationLayout::OnDataTableChanged()
 	}
 }
 
-TSharedRef<ITableRow> FDataTableCustomizationLayout::HandleRowNameComboBoxGenarateWidget(TSharedPtr<FString> InItem, const TSharedRef<STableViewBase>& OwnerTable)
+TSharedRef<ITableRow> FDataTableCustomizationLayout::HandleRowNameComboBoxGenarateWidget(TSharedPtr<FName> InItem, const TSharedRef<STableViewBase>& OwnerTable)
 {
 	return
-		SNew(STableRow<TSharedPtr<FString>>, OwnerTable)
+		SNew(STableRow<TSharedPtr<FName>>, OwnerTable)
 		[
-			SNew(STextBlock).Text(FText::FromString(*InItem))
+			SNew(STextBlock).Text(FText::FromName(*InItem))
 		];
 }
 
 FText FDataTableCustomizationLayout::GetRowNameComboBoxContentText() const
 {
-	FString RowNameValue;
+	FName RowNameValue;
 	const FPropertyAccess::Result RowResult = RowNamePropertyHandle->GetValue(RowNameValue);
 	if (RowResult == FPropertyAccess::Success)
 	{
-		return FText::FromString(*RowNameValue);
+		if (RowNameValue.IsNone())
+		{
+			return LOCTEXT("DataTable_None", "None");
+		}
+		return FText::FromName(RowNameValue);
 	}
 	else if (RowResult == FPropertyAccess::Fail)
 	{
@@ -219,13 +221,12 @@ FText FDataTableCustomizationLayout::GetRowNameComboBoxContentText() const
 	}
 }
 
-void FDataTableCustomizationLayout::OnSelectionChanged(TSharedPtr<FString> SelectedItem, ESelectInfo::Type SelectInfo)
+void FDataTableCustomizationLayout::OnSelectionChanged(TSharedPtr<FName> SelectedItem, ESelectInfo::Type SelectInfo)
 {
 	if (SelectedItem.IsValid())
 	{
 		CurrentSelectedItem = SelectedItem;
-		FName NewValue = FName(**SelectedItem);
-		RowNamePropertyHandle->SetValue(NewValue);
+		RowNamePropertyHandle->SetValue(*SelectedItem);
 
 		// Close the combo
 		RowNameComboButton->SetIsOpen(false);
@@ -236,32 +237,29 @@ void FDataTableCustomizationLayout::OnFilterTextChanged(const FText& InFilterTex
 {
 	FString CurrentFilterText = InFilterText.ToString();
 
-	FName RowName;
-	const FPropertyAccess::Result RowResult = RowNamePropertyHandle->GetValue(RowName);
 	RowNames.Empty();
 
 	/** Get the properties we wish to work with */
 	const UDataTable* DataTable = nullptr;
 	DataTablePropertyHandle->GetValue((UObject*&)DataTable);
 
-	TArray<FString> AllRowNames;
+	TArray<FName> AllRowNames;
 	if (DataTable != nullptr)
 	{
 		for (TMap<FName, uint8*>::TConstIterator Iterator(DataTable->GetRowMap()); Iterator; ++Iterator)
 		{
-			FString RowString = Iterator.Key().ToString();
-			AllRowNames.Add(RowString);
+			AllRowNames.Add(Iterator.Key());
 		}
 
 		// Sort the names alphabetically.
-		AllRowNames.Sort();
+		AllRowNames.Sort(FNameLexicalLess());
 	}
 
-	for (const FString& RowString : AllRowNames)
+	for (const FName& RowName : AllRowNames)
 	{
-		if (CurrentFilterText == TEXT("") || RowString.Contains(CurrentFilterText))
+		if (CurrentFilterText.IsEmpty() || RowName.ToString().Contains(CurrentFilterText))
 		{
-			TSharedRef<FString> RowNameItem = MakeShareable(new FString(RowString));
+			TSharedRef<FName> RowNameItem = MakeShared<FName>(RowName);
 			RowNames.Add(RowNameItem);
 		}
 	}
diff --git a/Engine/Source/Editor/DetailCustomizations/Private/DataTableCustomization.h b/Engine/Source/Editor/DetailCustomizations/Private/DataTableCustomization.h
--- a/Engine/Source/Editor/DetailCustomizations/Private/DataTableCustomization.h
+++ b/Engine/Source/Editor/DetailCustomizations/Private/DataTableCustomization.h
@@ -42,7 +42,7 @@ private:
 	bool ShouldFilterAsset(const FAssetData& AssetData);
 
 	/** Init the contents the combobox sources its data off */
-	TSharedPtr<FString> InitWidgetContent();
+	TSharedPtr<FName> InitWidgetContent();
 
 	/** Returns the ListView for the ComboButton */
 	TSharedRef<SWidget> GetListContent();
@@ -51,13 +51,13 @@ private:
 	void OnDataTableChanged();
 
 	/** Return the representation of the the row names to display */
-	TSharedRef<ITableRow> HandleRowNameComboBoxGenarateWidget(TSharedPtr<FString> InItem, const TSharedRef<STableViewBase>& OwnerTable);
+	TSharedRef<ITableRow> HandleRowNameComboBoxGenarateWidget(TSharedPtr<FName> InItem, const TSharedRef<STableViewBase>& OwnerTable);
 
 	/** Display the current selection */
 	FText GetRowNameComboBoxContentText() const;
 
 	/** Update the root data on a change of selection */
-	void OnSelectionChanged(TSharedPtr<FString> SelectedItem, ESelectInfo::Type SelectInfo);
+	void OnSelectionChanged(TSharedPtr<FName> SelectedItem, ESelectInfo::Type SelectInfo);
 
 	/** Called by Slate when the filter box changes text. */
 	void OnFilterTextChanged(const FText& InFilterText);
@@ -69,14 +69,14 @@ private:
 	/** The comboButton objects */
 	TSharedPtr<SComboButton> RowNameComboButton;
 	TSharedPtr<class SSearchBox> SearchBox;
-	TSharedPtr<SListView<TSharedPtr<FString> > > RowNameComboListView;
-	TSharedPtr<FString> CurrentSelectedItem;	
+	TSharedPtr<SListView<TSharedPtr<FName>>> RowNameComboListView;
+	TSharedPtr<FName> CurrentSelectedItem;
 	/** Handle to the struct properties being customized */
 	TSharedPtr<IPropertyHandle> StructPropertyHandle;
 	TSharedPtr<IPropertyHandle> DataTablePropertyHandle;
 	TSharedPtr<IPropertyHandle> RowNamePropertyHandle;
 	/** A cached copy of strings to populate the combo box */
-	TArray<TSharedPtr<FString> > RowNames;
+	TArray<TSharedPtr<FName>> RowNames;
 	/** The MetaData derived filter for the row type */
 	FName RowTypeFilter;
 };
