Attenuation priority.

- New feature to attenuate sound priority vs distance.

diff --git a/Engine/Source/Editor/DetailCustomizations/Private/AttenuationSettingsCustomizations.cpp b/Engine/Source/Editor/DetailCustomizations/Private/AttenuationSettingsCustomizations.cpp
--- a/Engine/Source/Editor/DetailCustomizations/Private/AttenuationSettingsCustomizations.cpp
+++ b/Engine/Source/Editor/DetailCustomizations/Private/AttenuationSettingsCustomizations.cpp
@@ -293,7 +293,9 @@ void FSoundAttenuationSettingsCustomization::CustomizeChildren(TSharedRef<IPrope
 	bIsSpatializedHandle = PropertyHandles.FindChecked(GET_MEMBER_NAME_CHECKED(FSoundAttenuationSettings, bSpatialize)).ToSharedRef();
 	bIsAirAbsorptionEnabledHandle = PropertyHandles.FindChecked(GET_MEMBER_NAME_CHECKED(FSoundAttenuationSettings, bAttenuateWithLPF)).ToSharedRef();
 	bIsReverbSendEnabledHandle = PropertyHandles.FindChecked(GET_MEMBER_NAME_CHECKED(FSoundAttenuationSettings, bEnableReverbSend)).ToSharedRef();
+	bIsPriorityAttenuationEnabledHandle = PropertyHandles.FindChecked(GET_MEMBER_NAME_CHECKED(FSoundAttenuationSettings, bEnablePriorityAttenuation)).ToSharedRef();
 	ReverbSendMethodHandle = PropertyHandles.FindChecked(GET_MEMBER_NAME_CHECKED(FSoundAttenuationSettings, ReverbSendMethod)).ToSharedRef();
+	PriorityAttenuationMethodHandle = PropertyHandles.FindChecked(GET_MEMBER_NAME_CHECKED(FSoundAttenuationSettings, PriorityAttenuationMethod)).ToSharedRef();
 	AbsorptionMethodHandle = PropertyHandles.FindChecked(GET_MEMBER_NAME_CHECKED(FSoundAttenuationSettings, AbsorptionMethod)).ToSharedRef();
 
 	bIsFocusedHandle = PropertyHandles.FindChecked(GET_MEMBER_NAME_CHECKED(FSoundAttenuationSettings, bEnableListenerFocus)).ToSharedRef();
@@ -456,13 +458,45 @@ void FSoundAttenuationSettingsCustomization::CustomizeChildren(TSharedRef<IPrope
 	LayoutBuilder.AddPropertyToCategory(PropertyHandles.FindChecked(GET_MEMBER_NAME_CHECKED(FSoundAttenuationSettings, bUseComplexCollisionForOcclusion)))
 		.EditCondition(GetIsOcclusionEnabledAttribute(), nullptr);
 
+	// Add the attenuation priority
+	LayoutBuilder.AddPropertyToCategory(bIsPriorityAttenuationEnabledHandle)
+		.EditCondition(IsAttenuationOverriddenAttribute(), nullptr);
+
+	LayoutBuilder.AddPropertyToCategory(PriorityAttenuationMethodHandle)
+		.EditCondition(GetIsPriorityAttenuationEnabledAttribute(), nullptr);
+
+	LayoutBuilder.AddPropertyToCategory(PropertyHandles.FindChecked(GET_MEMBER_NAME_CHECKED(FSoundAttenuationSettings, PriorityAttenuationMin)))
+		.Visibility(TAttribute<EVisibility>(this, &FSoundAttenuationSettingsCustomization::IsPriorityAttenuationLinearMethodSelected))
+		.EditCondition(GetIsPriorityAttenuationEnabledAttribute(), nullptr);
+
+	LayoutBuilder.AddPropertyToCategory(PropertyHandles.FindChecked(GET_MEMBER_NAME_CHECKED(FSoundAttenuationSettings, PriorityAttenuationMax)))
+		.Visibility(TAttribute<EVisibility>(this, &FSoundAttenuationSettingsCustomization::IsPriorityAttenuationLinearMethodSelected))
+		.EditCondition(GetIsPriorityAttenuationEnabledAttribute(), nullptr);
+
+	LayoutBuilder.AddPropertyToCategory(PropertyHandles.FindChecked(GET_MEMBER_NAME_CHECKED(FSoundAttenuationSettings, CustomPriorityAttenuationCurve)))
+		.Visibility(TAttribute<EVisibility>(this, &FSoundAttenuationSettingsCustomization::IsCustomPriorityAttenuationCurveSelected))
+		.EditCondition(GetIsPriorityAttenuationEnabledAttribute(), nullptr);
+
+	LayoutBuilder.AddPropertyToCategory(PropertyHandles.FindChecked(GET_MEMBER_NAME_CHECKED(FSoundAttenuationSettings, PriorityAttenuationDistanceMin)))
+		.Visibility(TAttribute<EVisibility>(this, &FSoundAttenuationSettingsCustomization::IsLinearOrCustomPriorityAttenuationSelected))
+		.EditCondition(GetIsPriorityAttenuationEnabledAttribute(), nullptr);
+
+	LayoutBuilder.AddPropertyToCategory(PropertyHandles.FindChecked(GET_MEMBER_NAME_CHECKED(FSoundAttenuationSettings, PriorityAttenuationDistanceMax)))
+		.Visibility(TAttribute<EVisibility>(this, &FSoundAttenuationSettingsCustomization::IsLinearOrCustomPriorityAttenuationSelected))
+		.EditCondition(GetIsPriorityAttenuationEnabledAttribute(), nullptr);
+
+	LayoutBuilder.AddPropertyToCategory(PropertyHandles.FindChecked(GET_MEMBER_NAME_CHECKED(FSoundAttenuationSettings, ManualPriorityAttenuation)))
+		.Visibility(TAttribute<EVisibility>(this, &FSoundAttenuationSettingsCustomization::IsManualPriorityAttenuationSelected))
+		.EditCondition(GetIsPriorityAttenuationEnabledAttribute(), nullptr);
+
+
 	if (bIsAudioMixerEnabled)
 	{
 		LayoutBuilder.AddPropertyToCategory(PropertyHandles.FindChecked(GET_MEMBER_NAME_CHECKED(FSoundAttenuationSettings, PluginSettings)))
 			.EditCondition(IsAttenuationOverriddenAttribute(), nullptr);
 	}
 
-	if (PropertyHandles.Num() != 52)
+	if (PropertyHandles.Num() != 60)
 	{
 		FString PropertyList;
 		for (auto It(PropertyHandles.CreateConstIterator()); It; ++It)
@@ -603,6 +637,24 @@ TAttribute<bool> FSoundAttenuationSettingsCustomization::GetIsReverbSendEnabledA
 	return TAttribute<bool>::Create(TAttribute<bool>::FGetter::CreateLambda(Lambda));
 }
 
+TAttribute<bool> FSoundAttenuationSettingsCustomization::GetIsPriorityAttenuationEnabledAttribute() const
+{
+	TWeakPtr<IPropertyHandle> bOverrideAttenuationPropertyWeakPtr = bOverrideAttenuationHandle;
+	TWeakPtr<IPropertyHandle> bIsPriorityAttenuationEnabledWeakPtr = bIsPriorityAttenuationEnabledHandle;
+
+	auto Lambda = [bOverrideAttenuationPropertyWeakPtr, bIsPriorityAttenuationEnabledWeakPtr]()
+	{
+		TSharedPtr<IPropertyHandle> bOverrideAttenuationProperty = bOverrideAttenuationPropertyWeakPtr.Pin();
+		TSharedPtr<IPropertyHandle> bIsPriorityAttenuationEnabledProperty = bIsPriorityAttenuationEnabledWeakPtr.Pin();
+
+		bool Value = GetValue(bOverrideAttenuationProperty);
+		Value &= GetValue(bIsPriorityAttenuationEnabledProperty);
+		return Value;
+	};
+
+	return TAttribute<bool>::Create(TAttribute<bool>::FGetter::CreateLambda(Lambda));
+}
+
 EVisibility FSoundAttenuationSettingsCustomization::IsLinearMethodSelected() const
 {
 	uint8 SendMethodValue;
@@ -653,6 +705,45 @@ EVisibility FSoundAttenuationSettingsCustomization::IsManualReverbSendSelected()
 	return (SendMethodType == EReverbSendMethod::Manual ? EVisibility::Visible : EVisibility::Hidden);
 }
 
+EVisibility FSoundAttenuationSettingsCustomization::IsPriorityAttenuationLinearMethodSelected() const
+{
+	uint8 Value;
+	PriorityAttenuationMethodHandle->GetValue(Value);
+
+	const EPriorityAttenuationMethod MethodType = (EPriorityAttenuationMethod)Value;
+
+	return (MethodType == EPriorityAttenuationMethod::Linear ? EVisibility::Visible : EVisibility::Hidden);
+}
+
+EVisibility FSoundAttenuationSettingsCustomization::IsCustomPriorityAttenuationCurveSelected() const
+{
+	uint8 Value;
+	PriorityAttenuationMethodHandle->GetValue(Value);
+
+	const EPriorityAttenuationMethod MethodType = (EPriorityAttenuationMethod)Value;
+
+	return (MethodType == EPriorityAttenuationMethod::CustomCurve ? EVisibility::Visible : EVisibility::Hidden);
+}
+
+EVisibility FSoundAttenuationSettingsCustomization::IsLinearOrCustomPriorityAttenuationSelected() const
+{
+	uint8 Value;
+	PriorityAttenuationMethodHandle->GetValue(Value);
+
+	const EPriorityAttenuationMethod MethodType = (EPriorityAttenuationMethod)Value;
+
+	return (MethodType != EPriorityAttenuationMethod::Manual ? EVisibility::Visible : EVisibility::Hidden);
+}
+
+EVisibility FSoundAttenuationSettingsCustomization::IsManualPriorityAttenuationSelected() const
+{
+	uint8 Value;
+	PriorityAttenuationMethodHandle->GetValue(Value);
+
+	const EPriorityAttenuationMethod MethodType = (EPriorityAttenuationMethod)Value;
+
+	return (MethodType == EPriorityAttenuationMethod::Manual ? EVisibility::Visible : EVisibility::Hidden);
+}
 
 EVisibility FBaseAttenuationSettingsCustomization::IsConeSelected() const
 {
diff --git a/Engine/Source/Editor/DetailCustomizations/Private/AttenuationSettingsCustomizations.h b/Engine/Source/Editor/DetailCustomizations/Private/AttenuationSettingsCustomizations.h
--- a/Engine/Source/Editor/DetailCustomizations/Private/AttenuationSettingsCustomizations.h
+++ b/Engine/Source/Editor/DetailCustomizations/Private/AttenuationSettingsCustomizations.h
@@ -53,7 +53,9 @@ protected:
 	TSharedPtr<IPropertyHandle> bIsFocusedHandle;
 	TSharedPtr<IPropertyHandle> bIsOcclusionEnabledHandle;
 	TSharedPtr<IPropertyHandle> bIsReverbSendEnabledHandle;
+	TSharedPtr<IPropertyHandle> bIsPriorityAttenuationEnabledHandle;
 	TSharedPtr<IPropertyHandle> ReverbSendMethodHandle;
+	TSharedPtr<IPropertyHandle> PriorityAttenuationMethodHandle;
 	TSharedPtr<IPropertyHandle> AbsorptionMethodHandle;
 
 	TAttribute<bool> IsFocusEnabledAttribute;
@@ -63,12 +65,17 @@ protected:
 	TAttribute<bool> GetIsSpatializationEnabledAttribute() const;
 	TAttribute<bool> GetIsAirAbsorptionEnabledAttribute() const;
 	TAttribute<bool> GetIsReverbSendEnabledAttribute() const;
+	TAttribute<bool> GetIsPriorityAttenuationEnabledAttribute() const;
 
 	EVisibility IsLinearMethodSelected() const;
 	EVisibility IsCustomReverbSendCurveSelected() const;
 	EVisibility IsLinearOrCustomReverbMethodSelected() const;
 	EVisibility IsManualReverbSendSelected() const;
 	EVisibility IsCustomAirAbsorptionCurveSelected() const;
+	EVisibility IsPriorityAttenuationLinearMethodSelected() const;
+	EVisibility IsCustomPriorityAttenuationCurveSelected() const;
+	EVisibility IsLinearOrCustomPriorityAttenuationSelected() const;
+	EVisibility IsManualPriorityAttenuationSelected() const;
 };
 
 class FForceFeedbackAttenuationSettingsCustomization : public FBaseAttenuationSettingsCustomization
diff --git a/Engine/Source/Runtime/Engine/Classes/Sound/SoundAttenuation.h b/Engine/Source/Runtime/Engine/Classes/Sound/SoundAttenuation.h
--- a/Engine/Source/Runtime/Engine/Classes/Sound/SoundAttenuation.h
+++ b/Engine/Source/Runtime/Engine/Classes/Sound/SoundAttenuation.h
@@ -52,6 +52,20 @@ enum class EReverbSendMethod : uint8
 };
 
 
+UENUM(BlueprintType)
+enum class EPriorityAttenuationMethod : uint8
+{
+	// A priority attenuation based on linear interpolation between a distance range and priority attenuation range
+	Linear,
+
+	// A priority attenuation based on a supplied curve
+	CustomCurve,
+
+	// A manual priority attenuation (Uses the specified constant value. Useful for 2D sounds.)
+	Manual,
+};
+
+
 USTRUCT(BlueprintType)
 struct ENGINE_API FSoundAttenuationPluginSettings
 {
@@ -110,6 +124,10 @@ struct ENGINE_API FSoundAttenuationSettings : public FBaseAttenuationSettings
 	UPROPERTY(EditAnywhere, BlueprintReadWrite, Category = AttenuationReverbSend, meta = (DisplayName = "Enable Reverb Send"))
 	uint8 bEnableReverbSend : 1;
 
+	/** Enables attenuation of sound priority based off distance. */
+	UPROPERTY(EditAnywhere, BlueprintReadWrite, Category = AttenuationPriority, meta = (DisplayName = "Enable Priority Attenuation"))
+	uint8 bEnablePriorityAttenuation : 1;
+
 	/** Enables applying a -6 dB attenuation to stereo assets which are 3d spatialized. Avoids clipping when assets have spread of 0.0 due to channel summing. */
 	UPROPERTY(EditAnywhere, BlueprintReadWrite, Category = AttenuationSpatialization, meta = (DisplayName = "Normalize 3D Stereo Sounds"))
 	uint8 bApplyNormalizationToStereoSounds : 1;
@@ -138,6 +156,10 @@ struct ENGINE_API FSoundAttenuationSettings : public FBaseAttenuationSettings
 	UPROPERTY(EditAnywhere, BlueprintReadWrite, Category = AttenuationReverbSend)
 	EReverbSendMethod ReverbSendMethod;
 
+	/** What method to use to control priority attenuation */
+	UPROPERTY(EditAnywhere, BlueprintReadWrite, Category = AttenuationPriority)
+	EPriorityAttenuationMethod PriorityAttenuationMethod;
+
 #if WITH_EDITORONLY_DATA
 	UPROPERTY()
 	TEnumAsByte<enum ESoundDistanceCalc> DistanceType_DEPRECATED;
@@ -278,6 +300,30 @@ struct ENGINE_API FSoundAttenuationSettings : public FBaseAttenuationSettings
 	UPROPERTY(EditAnywhere, BlueprintReadWrite, Category = AttenuationReverbSend)
 	FRuntimeFloatCurve CustomReverbSendCurve;
 
+	/** The priority attenuation to use when the sound is at the minimum priority attenuation distance. */
+	UPROPERTY(EditAnywhere, BlueprintReadWrite, Category = AttenuationPriority, meta = (DisplayName = "Min Priority Attenuation"))
+	float PriorityAttenuationMin;
+
+	/** The priority attenuation to use when the sound is at the maximum priority attenuation distance. */
+	UPROPERTY(EditAnywhere, BlueprintReadWrite, Category = AttenuationPriority, meta = (DisplayName = "Max Priority Attenuation"))
+	float PriorityAttenuationMax;
+
+	/** The min distance to attenuate priority. */
+	UPROPERTY(EditAnywhere, BlueprintReadWrite, Category = AttenuationPriority, meta = (DisplayName = "Priority Attenuation Min Distance"))
+	float PriorityAttenuationDistanceMin;
+
+	/** The max distance to attenuate priority. */
+	UPROPERTY(EditAnywhere, BlueprintReadWrite, Category = AttenuationPriority, meta = (DisplayName = "Priority Attenuation Max Distance"))
+	float PriorityAttenuationDistanceMax;
+
+	/* The manual priority attenuation to use. Doesn't change as a function of distance. */
+	UPROPERTY(EditAnywhere, BlueprintReadWrite, Category = AttenuationPriority)
+	float ManualPriorityAttenuation;
+
+	/* The custom curve to use for distance-based priority attenuation. */
+	UPROPERTY(EditAnywhere, BlueprintReadWrite, Category = AttenuationPriority)
+	FRuntimeFloatCurve CustomPriorityAttenuationCurve;
+
 	/** Sound attenuation plugin settings to use with sounds that play with this attenuation setting. */
 	UPROPERTY(EditAnywhere, BlueprintReadWrite, Category = AttenuationPluginSettings, meta = (ShowOnlyInnerProperties))
 	FSoundAttenuationPluginSettings PluginSettings;
@@ -291,6 +337,7 @@ struct ENGINE_API FSoundAttenuationSettings : public FBaseAttenuationSettings
 		, bEnableOcclusion(false)
 		, bUseComplexCollisionForOcclusion(false)
 		, bEnableReverbSend(true)
+		, bEnablePriorityAttenuation(false)
 		, bApplyNormalizationToStereoSounds(false)
 		, bEnableLogFrequencyScaling(false)
 		, SpatializationAlgorithm(ESoundSpatializationAlgorithm::SPATIALIZATION_Default)
@@ -298,6 +345,7 @@ struct ENGINE_API FSoundAttenuationSettings : public FBaseAttenuationSettings
 		, AbsorptionMethod(EAirAbsorptionMethod::Linear)
 		, OcclusionTraceChannel(ECC_Visibility)
 		, ReverbSendMethod(EReverbSendMethod::Linear)
+		, PriorityAttenuationMethod(EPriorityAttenuationMethod::Linear)
 #if WITH_EDITORONLY_DATA
 		, DistanceType_DEPRECATED(SOUNDDISTANCE_Normal)
 #endif
@@ -336,6 +384,11 @@ struct ENGINE_API FSoundAttenuationSettings : public FBaseAttenuationSettings
 		, ReverbDistanceMin(AttenuationShapeExtents.X)
 		, ReverbDistanceMax(AttenuationShapeExtents.X + FalloffDistance)
 		, ManualReverbSendLevel(0.0f)
+		, PriorityAttenuationMin(1.0f)
+		, PriorityAttenuationMax(1.0f)
+		, PriorityAttenuationDistanceMin(AttenuationShapeExtents.X)
+		, PriorityAttenuationDistanceMax(AttenuationShapeExtents.X + FalloffDistance)
+		, ManualPriorityAttenuation(1.0f)
 	{
 	}
 
diff --git a/Engine/Source/Runtime/Engine/Private/ActiveSound.cpp b/Engine/Source/Runtime/Engine/Private/ActiveSound.cpp
--- a/Engine/Source/Runtime/Engine/Private/ActiveSound.cpp
+++ b/Engine/Source/Runtime/Engine/Private/ActiveSound.cpp
@@ -1564,6 +1564,33 @@ void FActiveSound::UpdateAttenuation(float DeltaTime, FSoundParseParameters& Par
 	check(AudioDevice);
 	FAttenuationListenerData ListenerData = FAttenuationListenerData::Create(*AudioDevice, Listener.Transform, ParseParams.Transform, *Settings);
 
+	// Apply priority attenuation if it's enabled
+	if (Settings->bEnablePriorityAttenuation)
+	{
+		float PriorityScale = 1.0f;
+		if (Settings->PriorityAttenuationMethod == EPriorityAttenuationMethod::Manual)
+		{
+			PriorityScale = Settings->ManualPriorityAttenuation;
+		}
+		else
+		{
+			const float Denom = FMath::Max(Settings->PriorityAttenuationDistanceMax - Settings->PriorityAttenuationDistanceMin, 1.0f);
+			const float Alpha = FMath::Clamp((ListenerData.ListenerToSoundDistance - Settings->PriorityAttenuationDistanceMin) / Denom, 0.0f, 1.0f);
+
+			if (Settings->PriorityAttenuationMethod == EPriorityAttenuationMethod::Linear)
+			{
+				PriorityScale = FMath::Clamp(FMath::Lerp(Settings->PriorityAttenuationMin, Settings->PriorityAttenuationMax, Alpha), 0.0f, 1.0f);
+			}
+			else
+			{
+				PriorityScale = FMath::Clamp(Settings->CustomPriorityAttenuationCurve.GetRichCurveConst()->Eval(Alpha), 0.0f, 1.0f);
+			}
+		}
+
+		ParseParams.Priority *= FMath::Max(PriorityScale, 0.0f);
+		ParseParams.Priority = FMath::Clamp(ParseParams.Priority, 0.0f, 100.0f);
+	}
+
 	if (Settings->bSpatialize || Settings->bEnableListenerFocus)
 	{
 		// Feed prior focus factor on update to allow for proper interpolation.
