Change the warning for experimental/beta classes properly report the last ancestor class that declared experimental/beta instead of the last ancestor class that introduced a property

diff --git a/Engine/Source/Editor/DetailCustomizations/Private/ObjectDetails.cpp b/Engine/Source/Editor/DetailCustomizations/Private/ObjectDetails.cpp
--- a/Engine/Source/Editor/DetailCustomizations/Private/ObjectDetails.cpp
+++ b/Engine/Source/Editor/DetailCustomizations/Private/ObjectDetails.cpp
@@ -37,16 +37,15 @@ void FObjectDetails::AddExperimentalWarningCategory(IDetailLayoutBuilder& Detail
 {
 	bool bBaseClassIsExperimental = false;
 	bool bBaseClassIsEarlyAccess = false;
-
-	FObjectEditorUtils::GetClassDevelopmentStatus(DetailBuilder.GetBaseClass(), bBaseClassIsExperimental, bBaseClassIsEarlyAccess);
+	FString MostDerivedDevelopmentClassName;
+	FObjectEditorUtils::GetClassDevelopmentStatus(DetailBuilder.GetBaseClass(), bBaseClassIsExperimental, bBaseClassIsEarlyAccess, MostDerivedDevelopmentClassName);
 
 	if (bBaseClassIsExperimental || bBaseClassIsEarlyAccess)
 	{
 		const FName CategoryName(TEXT("Warning"));
 		const FText CategoryDisplayName = LOCTEXT("WarningCategoryDisplayName", "Warning");
-		FString ClassUsed = DetailBuilder.GetTopLevelProperty().ToString();
-		const FText WarningText = bBaseClassIsExperimental ? FText::Format( LOCTEXT("ExperimentalClassWarning", "Uses experimental class: {0}") , FText::FromString(ClassUsed) )
-			: FText::Format( LOCTEXT("EarlyAccessClassWarning", "Uses beta class {0}"), FText::FromString(*ClassUsed) );
+		const FText WarningText = bBaseClassIsExperimental ? FText::Format( LOCTEXT("ExperimentalClassWarning", "Uses experimental class: {0}") , FText::FromString(MostDerivedDevelopmentClassName) )
+			: FText::Format( LOCTEXT("EarlyAccessClassWarning", "Uses beta class {0}"), FText::FromString(MostDerivedDevelopmentClassName) );
 		const FText SearchString = WarningText;
 		const FText Tooltip = bBaseClassIsExperimental ? LOCTEXT("ExperimentalClassTooltip", "Here be dragons!  Uses one or more unsupported 'experimental' classes") : LOCTEXT("EarlyAccessClassTooltip", "Uses one or more 'beta' classes");
 		const FString ExcerptName = bBaseClassIsExperimental ? TEXT("ObjectUsesExperimentalClass") : TEXT("ObjectUsesEarlyAccessClass");
diff --git a/Engine/Source/Editor/UMGEditor/Private/WidgetBlueprintEditorUtils.cpp b/Engine/Source/Editor/UMGEditor/Private/WidgetBlueprintEditorUtils.cpp
--- a/Engine/Source/Editor/UMGEditor/Private/WidgetBlueprintEditorUtils.cpp
+++ b/Engine/Source/Editor/UMGEditor/Private/WidgetBlueprintEditorUtils.cpp
@@ -1509,7 +1509,8 @@ bool FWidgetBlueprintEditorUtils::IsUsableWidgetClass(UClass* WidgetClass)
 	{
 		// We aren't interested in classes that are experimental or cannot be instantiated
 		bool bIsExperimental, bIsEarlyAccess;
-		FObjectEditorUtils::GetClassDevelopmentStatus(WidgetClass, bIsExperimental, bIsEarlyAccess);
+		FString MostDerivedDevelopmentClassName;
+		FObjectEditorUtils::GetClassDevelopmentStatus(WidgetClass, bIsExperimental, bIsEarlyAccess, MostDerivedDevelopmentClassName);
 		const bool bIsInvalid = WidgetClass->HasAnyClassFlags(CLASS_Abstract | CLASS_Deprecated | CLASS_NewerVersionExists);
 		if ( bIsExperimental || bIsEarlyAccess || bIsInvalid )
 		{
diff --git a/Engine/Source/Editor/UnrealEd/Private/AssetSelection.cpp b/Engine/Source/Editor/UnrealEd/Private/AssetSelection.cpp
--- a/Engine/Source/Editor/UnrealEd/Private/AssetSelection.cpp
+++ b/Engine/Source/Editor/UnrealEd/Private/AssetSelection.cpp
@@ -240,7 +240,8 @@ namespace AssetSelectionUtils
 
 							// Check for experimental/early-access classes in the component hierarchy
 							bool bIsExperimental, bIsEarlyAccess;
-							FObjectEditorUtils::GetClassDevelopmentStatus(Component->GetClass(), bIsExperimental, bIsEarlyAccess);
+							FString MostDerivedDevelopmentClassName;
+							FObjectEditorUtils::GetClassDevelopmentStatus(Component->GetClass(), bIsExperimental, bIsEarlyAccess, MostDerivedDevelopmentClassName);
 
 							ActorInfo.bHaveExperimentalClass |= bIsExperimental;
 							ActorInfo.bHaveEarlyAccessClass |= bIsEarlyAccess;
@@ -250,7 +251,8 @@ namespace AssetSelectionUtils
 					// Check for experimental/early-access classes in the actor hierarchy
 					{
 						bool bIsExperimental, bIsEarlyAccess;
-						FObjectEditorUtils::GetClassDevelopmentStatus(CurrentClass, bIsExperimental, bIsEarlyAccess);
+						FString MostDerivedDevelopmentClassName;
+						FObjectEditorUtils::GetClassDevelopmentStatus(CurrentClass, bIsExperimental, bIsEarlyAccess, MostDerivedDevelopmentClassName);
 
 						ActorInfo.bHaveExperimentalClass |= bIsExperimental;
 						ActorInfo.bHaveEarlyAccessClass |= bIsEarlyAccess;
diff --git a/Engine/Source/Runtime/Engine/Private/ObjectEditorUtils.cpp b/Engine/Source/Runtime/Engine/Private/ObjectEditorUtils.cpp
--- a/Engine/Source/Runtime/Engine/Private/ObjectEditorUtils.cpp
+++ b/Engine/Source/Runtime/Engine/Private/ObjectEditorUtils.cpp
@@ -69,20 +69,30 @@ namespace FObjectEditorUtils
 		return bResult;
 	}
 
-	void GetClassDevelopmentStatus(UClass* Class, bool& bIsExperimental, bool& bIsEarlyAccess)
+	void GetClassDevelopmentStatus(UClass* Class, bool& bIsExperimental, bool& bIsEarlyAccess, FString& MostDerivedClassName)
 	{
 		static const FName DevelopmentStatusKey(TEXT("DevelopmentStatus"));
 		static const FString EarlyAccessValue(TEXT("EarlyAccess"));
 		static const FString ExperimentalValue(TEXT("Experimental"));
 
+		MostDerivedClassName = FString();
 		bIsExperimental = bIsEarlyAccess = false;
 
+		// Determine the development status and introducing class
 		FString DevelopmentStatus;
-		if ( Class->GetStringMetaDataHierarchical(DevelopmentStatusKey, /*out*/ &DevelopmentStatus) )
+		for (const UStruct* TestStruct = Class; TestStruct != nullptr; TestStruct = TestStruct->GetSuperStruct())
 		{
-			bIsExperimental = DevelopmentStatus == ExperimentalValue;
-			bIsEarlyAccess = DevelopmentStatus == EarlyAccessValue;
+			if (TestStruct->HasMetaData(DevelopmentStatusKey))
+			{
+				DevelopmentStatus = TestStruct->GetMetaData(DevelopmentStatusKey);
+				MostDerivedClassName = TestStruct->GetName();
+				break;
+			}
 		}
+
+		// Update the flags based on the status
+		bIsExperimental = DevelopmentStatus == ExperimentalValue;
+		bIsEarlyAccess = DevelopmentStatus == EarlyAccessValue;
 	}
 
 	static void CopySinglePropertyRecursive(UObject* SourceObject, const void* const InSourcePtr, UProperty* InSourceProperty, void* const InTargetPtr, UObject* InDestinationObject, UProperty* InDestinationProperty)
diff --git a/Engine/Source/Runtime/Engine/Public/ObjectEditorUtils.h b/Engine/Source/Runtime/Engine/Public/ObjectEditorUtils.h
--- a/Engine/Source/Runtime/Engine/Public/ObjectEditorUtils.h
+++ b/Engine/Source/Runtime/Engine/Public/ObjectEditorUtils.h
@@ -63,8 +63,9 @@ namespace FObjectEditorUtils
 	 * @param Class the class to inspect.
 	 * @param bIsExperimental [out] value indicating if the class is experimental.
 	 * @param bIsEarlyAccess [out] value indicating if the class is early access.
+	 * @param MostDerivedClassName [out] The name of the most derived class that is marked as experimental or early access (or empty string)
 	 */
-	ENGINE_API void GetClassDevelopmentStatus(UClass* Class, bool& bIsExperimental, bool& bIsEarlyAccess);
+	ENGINE_API void GetClassDevelopmentStatus(UClass* Class, bool& bIsExperimental, bool& bIsEarlyAccess, FString& MostDerivedClassName);
 
 	/**
 	 * Copy the value of a property from source object to a destination object.
