Fix DefaultEngine.ini being saved each time project settings tab opened

diff --git a/Engine/Source/Editor/DetailCustomizations/Private/PhysicsSettingsDetails.cpp b/Engine/Source/Editor/DetailCustomizations/Private/PhysicsSettingsDetails.cpp
--- a/Engine/Source/Editor/DetailCustomizations/Private/PhysicsSettingsDetails.cpp
+++ b/Engine/Source/Editor/DetailCustomizations/Private/PhysicsSettingsDetails.cpp
@@ -163,7 +163,6 @@ public:
 		PhysicalSurfaceList.Sort(FComparePhysicalSurface());
 
 		PhysicsSettings->LoadSurfaceType();
-		PhysicsSettings->UpdateDefaultConfigFile();
 
 		RegenerateChildren.ExecuteIfBound();
 	}
