- Fix bug with 'Natural Sound' attenuation falloff, where dBAttenuationAtMax when set to a value larger than the default minimum of -60.f causes looping sounds to virtualize even though volume doesn't go to 0.
- Add ability to specify if sound should continue attenuating, go silent, or hold the max value when beyond the bounds of the attenuation shape.

diff --git a/Engine/Source/Editor/DetailCustomizations/Private/AttenuationSettingsCustomizations.cpp b/Engine/Source/Editor/DetailCustomizations/Private/AttenuationSettingsCustomizations.cpp
--- a/Engine/Source/Editor/DetailCustomizations/Private/AttenuationSettingsCustomizations.cpp
+++ b/Engine/Source/Editor/DetailCustomizations/Private/AttenuationSettingsCustomizations.cpp
@@ -119,10 +119,15 @@ void FBaseAttenuationSettingsCustomization::CustomizeChildren(TSharedRef<IProper
 		.Visibility(TAttribute<EVisibility>(this, &FBaseAttenuationSettingsCustomization::IsCustomCurveSelected))
 		.EditCondition(GetIsAttenuationEnabledAttribute(), nullptr);
 
-	LayoutBuilder.AddPropertyToCategory(PropertyHandles.FindChecked(GET_MEMBER_NAME_CHECKED(FBaseAttenuationSettings, dBAttenuationAtMax)))
+	DbAttenuationAtMaxHandle = PropertyHandles.FindChecked(GET_MEMBER_NAME_CHECKED(FBaseAttenuationSettings, dBAttenuationAtMax));
+	LayoutBuilder.AddPropertyToCategory(DbAttenuationAtMaxHandle)
 		.Visibility(TAttribute<EVisibility>(this, &FBaseAttenuationSettingsCustomization::IsNaturalSoundSelected))
 		.EditCondition(GetIsAttenuationEnabledAttribute(), nullptr);
 
+	LayoutBuilder.AddPropertyToCategory(PropertyHandles.FindChecked(GET_MEMBER_NAME_CHECKED(FBaseAttenuationSettings, FalloffMode)))
+		.Visibility(TAttribute<EVisibility>(this, &FBaseAttenuationSettingsCustomization::IsNaturalSoundSelected))
+		.EditCondition(GetIsFalloffModeEnabledAttribute(), nullptr);
+
 	LayoutBuilder.AddPropertyToCategory(AttenuationShapeHandle)
 		.EditCondition(GetIsAttenuationEnabledAttribute(), nullptr);
 	
@@ -269,6 +274,33 @@ TAttribute<bool> FBaseAttenuationSettingsCustomization::GetIsAttenuationEnabledA
 	return TAttribute<bool>::Create(TAttribute<bool>::FGetter::CreateLambda(Lambda));
 }
 
+TAttribute<bool> FBaseAttenuationSettingsCustomization::GetIsFalloffModeEnabledAttribute() const
+{
+	TWeakPtr<IPropertyHandle> bOverrideAttenuationPropertyWeakPtr = bOverrideAttenuationHandle;
+	TWeakPtr<IPropertyHandle> bIsAttenuatedPropertyWeakPtr = bIsAttenuatedHandle;
+	TWeakPtr<IPropertyHandle> DbAttenuationAtMaxHandleWeakPtr = DbAttenuationAtMaxHandle;
+
+	auto Lambda = [bOverrideAttenuationPropertyWeakPtr, bIsAttenuatedPropertyWeakPtr, DbAttenuationAtMaxHandleWeakPtr]()
+	{
+		TSharedPtr<IPropertyHandle> bOverrideAttenuationPropertyPtr = bOverrideAttenuationPropertyWeakPtr.Pin();
+		TSharedPtr<IPropertyHandle> bIsAttenuatedPropertyPtr = bIsAttenuatedPropertyWeakPtr.Pin();
+		TSharedPtr<IPropertyHandle> DbAttenuationAtMaxHandlePtr = DbAttenuationAtMaxHandleWeakPtr.Pin();
+
+		float AttenuationValue = -60.f;
+		if (DbAttenuationAtMaxHandlePtr.IsValid())
+		{
+			DbAttenuationAtMaxHandlePtr->GetValue(AttenuationValue);
+		}
+
+		bool Value = GetValue(bOverrideAttenuationPropertyPtr);
+		Value &= GetValue(bIsAttenuatedPropertyPtr);
+		Value &= AttenuationValue > -60.f;
+		return Value;
+	};
+
+	return TAttribute<bool>::Create(TAttribute<bool>::FGetter::CreateLambda(Lambda));
+}
+
 void FSoundAttenuationSettingsCustomization::CustomizeHeader(TSharedRef<IPropertyHandle> PropertyHandle, class FDetailWidgetRow& HeaderRow, IPropertyTypeCustomizationUtils& CustomizationUtils)
 {
 	// Property handle here is the base struct. We are going to hide it since we're showing it's properties directly.
@@ -496,7 +528,7 @@ void FSoundAttenuationSettingsCustomization::CustomizeChildren(TSharedRef<IPrope
 			.EditCondition(IsAttenuationOverriddenAttribute(), nullptr);
 	}
 
-	if (PropertyHandles.Num() != 60)
+	if (PropertyHandles.Num() != 61)
 	{
 		FString PropertyList;
 		for (auto It(PropertyHandles.CreateConstIterator()); It; ++It)
diff --git a/Engine/Source/Editor/DetailCustomizations/Private/AttenuationSettingsCustomizations.h b/Engine/Source/Editor/DetailCustomizations/Private/AttenuationSettingsCustomizations.h
--- a/Engine/Source/Editor/DetailCustomizations/Private/AttenuationSettingsCustomizations.h
+++ b/Engine/Source/Editor/DetailCustomizations/Private/AttenuationSettingsCustomizations.h
@@ -28,11 +28,13 @@ protected:
 	EVisibility IsCustomCurveSelected() const;
 
 	TAttribute<bool> GetIsAttenuationEnabledAttribute() const;
+	TAttribute<bool> GetIsFalloffModeEnabledAttribute() const;
 	TSharedPtr<IPropertyHandle> GetOverrideAttenuationHandle(TSharedRef<IPropertyHandle> StructPropertyHandle);
 
 	TSharedPtr<IPropertyHandle> bIsAttenuatedHandle;
 	TSharedPtr<IPropertyHandle> AttenuationShapeHandle;
 	TSharedPtr<IPropertyHandle> DistanceAlgorithmHandle;
+	TSharedPtr<IPropertyHandle> DbAttenuationAtMaxHandle;
 	TSharedPtr<IPropertyHandle> bOverrideAttenuationHandle;
 };
 
diff --git a/Engine/Source/Runtime/Engine/Classes/Engine/Attenuation.h b/Engine/Source/Runtime/Engine/Classes/Engine/Attenuation.h
--- a/Engine/Source/Runtime/Engine/Classes/Engine/Attenuation.h
+++ b/Engine/Source/Runtime/Engine/Classes/Engine/Attenuation.h
@@ -35,6 +35,20 @@ namespace EAttenuationShape
 	};
 }
 
+UENUM(BlueprintType)
+enum class ENaturalSoundFalloffMode : uint8
+{
+	// (Default) Continues attenuating pass falloff max using volume value
+	// specified at the max falloff distance's bounds
+	Continues,
+
+	// Sound goes silent upon leaving the shape
+	Silent,
+
+	// Holds the volume value specified at the shapes falloff bounds
+	Hold,
+};
+
 /*
 * Base class for attenuation settings.
 */
@@ -53,10 +67,16 @@ struct ENGINE_API FBaseAttenuationSettings
 	UPROPERTY(EditAnywhere, BlueprintReadWrite, Category= AttenuationDistance)
 	TEnumAsByte<enum EAttenuationShape::Type> AttenuationShape;
 
-	/* The attenuation volume at maximum distance in decibels, used for natural attenuation method. */
-	UPROPERTY(EditAnywhere, BlueprintReadWrite, Category= AttenuationDistance, meta=(DisplayName = "dB Attenuation At Max", ClampMax = "0" ))
+	/* The attenuation volume at the falloff distance in decibels (Only for 'Natural Sound' Distance Algorithm). */
+	UPROPERTY(EditAnywhere, BlueprintReadWrite, Category= AttenuationDistance, meta=(DisplayName = "Attenuation At Max (dB)", ClampMin = "-60", ClampMax = "0"))
 	float dBAttenuationAtMax;
 
+	// Whether to continue attenuating, go silent, or hold last volume value when beyond falloff bounds and 
+	// 'Attenuation At Max (dB)' is set to a value greater than -60dB.
+	// (Only for 'Natural Sound' Distance Algorithm). */
+	UPROPERTY(EditAnywhere, BlueprintReadWrite, Category= AttenuationDistance)
+	ENaturalSoundFalloffMode FalloffMode;
+
 	/* The dimensions to use for the attenuation shape. Interpretation of the values differ per shape.
 	   Sphere  - X is Sphere Radius. Y and Z are unused
 	   Capsule - X is Capsule Half Height, Y is Capsule Radius, Z is unused
@@ -82,6 +102,7 @@ struct ENGINE_API FBaseAttenuationSettings
 		: DistanceAlgorithm(EAttenuationDistanceModel::Linear)
 		, AttenuationShape(EAttenuationShape::Sphere)
 		, dBAttenuationAtMax(-60.f)
+		, FalloffMode(ENaturalSoundFalloffMode::Continues)
 		, AttenuationShapeExtents(400.f, 0.f, 0.f)
 		, ConeOffset(0.f)
 		, FalloffDistance(3600.f)
diff --git a/Engine/Source/Runtime/Engine/Private/Attenuation.cpp b/Engine/Source/Runtime/Engine/Private/Attenuation.cpp
--- a/Engine/Source/Runtime/Engine/Private/Attenuation.cpp
+++ b/Engine/Source/Runtime/Engine/Private/Attenuation.cpp
@@ -2,11 +2,37 @@
 
 #include "Engine/Attenuation.h"
 
+#include "DSP/Dsp.h"
+#include "EngineDefines.h"
+
+
 float FBaseAttenuationSettings::GetMaxDimension() const
 {
+	static const float WorldMax = static_cast<float>(WORLD_MAX);
+	if (FalloffDistance > WorldMax)
+	{
+		return WorldMax;
+	}
+
 	float MaxDimension = FalloffDistance;
+	if (DistanceAlgorithm == EAttenuationDistanceModel::NaturalSound)
+	{
+		switch (FalloffMode)
+		{
+		case ENaturalSoundFalloffMode::Hold:
+			return WorldMax;
+
+		case ENaturalSoundFalloffMode::Continues:
+			MaxDimension = (FalloffDistance / dBAttenuationAtMax) * -60.f;
+			break;
+
+		case ENaturalSoundFalloffMode::Silent:
+		default:
+			break;
+		}
+	}
 
-	switch(AttenuationShape)
+	switch (AttenuationShape)
 	{
 	case EAttenuationShape::Sphere:
 	case EAttenuationShape::Cone:
@@ -20,7 +46,7 @@ float FBaseAttenuationSettings::GetMaxDimension() const
 		break;
 
 	case EAttenuationShape::Capsule:
-		
+
 		MaxDimension += FMath::Max(AttenuationShapeExtents.X, AttenuationShapeExtents.Y);
 		break;
 
@@ -28,7 +54,7 @@ float FBaseAttenuationSettings::GetMaxDimension() const
 		check(false);
 	}
 
-	return MaxDimension;
+	return FMath::Clamp(MaxDimension, 0.0f, WorldMax);
 }
 
 float FBaseAttenuationSettings::Evaluate(const FTransform& Origin, const FVector Location, const float DistanceScale) const
@@ -113,7 +139,28 @@ float FBaseAttenuationSettings::AttenuationEval(const float Distance, const floa
 		case EAttenuationDistanceModel::NaturalSound:
 		{
 			check(dBAttenuationAtMax <= 0.0f);
-			Result = FMath::Pow(10.0f, ((DistanceCopy / FalloffCopy) * dBAttenuationAtMax) / 20.0f);
+			float Alpha = DistanceCopy / FalloffCopy;
+			if (FalloffMode == ENaturalSoundFalloffMode::Hold)
+			{
+				Alpha = FMath::Clamp(Alpha, 0.0f, 1.0f);
+				Result = Audio::ConvertToLinear(Alpha * dBAttenuationAtMax);
+			}
+			else if (FalloffMode == ENaturalSoundFalloffMode::Silent)
+			{
+				Alpha = FMath::Clamp(Alpha, 0.0f, 1.0f);
+				if (Alpha < 1.0f)
+				{
+					Result = Audio::ConvertToLinear(Alpha * dBAttenuationAtMax);
+				}
+				else
+				{
+					Result = 0.0f;
+				}
+			}
+			else
+			{
+				Result = Audio::ConvertToLinear(Alpha * dBAttenuationAtMax);
+			}
 			break;
 		}
 
@@ -123,8 +170,8 @@ float FBaseAttenuationSettings::AttenuationEval(const float Distance, const floa
 			break;
 
 		default:
-			checkf(false, TEXT("Uknown attenuation distance algorithm!"))
-				break;
+			checkf(false, TEXT("Unknown attenuation distance algorithm!"))
+			break;
 	}
 
 	// Make sure the output is clamped between 0.0 and 1.0f. Some of the algorithms above can
diff --git a/Engine/Source/Runtime/Engine/Private/SoundAttenuation.cpp b/Engine/Source/Runtime/Engine/Private/SoundAttenuation.cpp
--- a/Engine/Source/Runtime/Engine/Private/SoundAttenuation.cpp
+++ b/Engine/Source/Runtime/Engine/Private/SoundAttenuation.cpp
@@ -90,6 +90,7 @@ bool FSoundAttenuationSettings::operator==(const FSoundAttenuationSettings& Othe
 	return (   bAttenuate			    == Other.bAttenuate
 			&& bSpatialize			    == Other.bSpatialize
 			&& dBAttenuationAtMax	    == Other.dBAttenuationAtMax
+			&& FalloffMode				== Other.FalloffMode
 			&& OmniRadius				== Other.OmniRadius
 			&& bApplyNormalizationToStereoSounds == Other.bApplyNormalizationToStereoSounds
 			&& StereoSpread				== Other.StereoSpread
