Carrying over change potentially causing crash in 11.40

diff --git a/Engine/Source/Editor/DetailCustomizations/Private/WindowsTargetSettingsDetails.cpp b/Engine/Source/Editor/DetailCustomizations/Private/WindowsTargetSettingsDetails.cpp
--- a/Engine/Source/Editor/DetailCustomizations/Private/WindowsTargetSettingsDetails.cpp
+++ b/Engine/Source/Editor/DetailCustomizations/Private/WindowsTargetSettingsDetails.cpp
@@ -288,7 +288,7 @@ void FWindowsTargetSettingsDetails::CustomizeDetails( IDetailLayoutBuilder& Deta
 	];
 
 
-	AudioPluginWidgetManager.BuildAudioCategory(DetailBuilder, TEXT("Windows"));
+	AudioPluginWidgetManager.BuildAudioCategory(DetailBuilder, EAudioPlatform::Windows);
 	IDetailCategoryBuilder& AudioCategory = DetailBuilder.EditCategory("Audio");
 
 	// Here we add a callback when the 
