MoviePlayerSettingsDetails: Fixed copy movie requester showing up when it should not.

diff --git a/Engine/Source/Editor/DetailCustomizations/Private/MoviePlayerSettingsDetails.cpp b/Engine/Source/Editor/DetailCustomizations/Private/MoviePlayerSettingsDetails.cpp
--- a/Engine/Source/Editor/DetailCustomizations/Private/MoviePlayerSettingsDetails.cpp
+++ b/Engine/Source/Editor/DetailCustomizations/Private/MoviePlayerSettingsDetails.cpp
@@ -16,6 +16,7 @@
 #include "Widgets/Input/SFilePathPicker.h"
 #include "Framework/Notifications/NotificationManager.h"
 #include "Widgets/Notifications/SNotificationList.h"
+#include "HAL/FileManager.h"
 
 
 #define LOCTEXT_NAMESPACE "MoviePlayerSettingsDetails"
@@ -77,6 +78,16 @@ void FMoviePlayerSettingsDetails::GenerateArrayElementWidget(TSharedRef<IPropert
 
 void FMoviePlayerSettingsDetails::HandleFilePathPickerPathPicked( const FString& PickedPath, TSharedRef<IPropertyHandle> Property )
 {
+	// Ignore if value is unchanged.
+	FString OldValue;
+	if (Property->GetValueAsFormattedString(OldValue) == FPropertyAccess::Success)
+	{
+		if (PickedPath == OldValue)
+		{
+			return;
+		}
+	}
+
 	FEditorDirectories::Get().SetLastDirectory(ELastDirectory::GENERIC_OPEN, FPaths::GetPath(PickedPath));
 
 	// sanitize the location of the chosen movies to the content/movies directory
@@ -102,35 +113,75 @@ void FMoviePlayerSettingsDetails::HandleFilePathPickerPathPicked( const FString&
 	}
 	else if (!PickedPath.IsEmpty())
 	{
-		// ask the user if they want to import this movie
-		FSuppressableWarningDialog::FSetupInfo Info( 
-			LOCTEXT("ExternalMovieImportWarning", "This movie needs to be copied into your project, would you like to copy the file now?"), 
-			LOCTEXT("ExternalMovieImportTitle", "Copy Movie"), 
-			TEXT("ImportMovieIntoProject") );
-		Info.ConfirmText = LOCTEXT("ExternalMovieImport_Confirm", "Copy");
-		Info.CancelText = LOCTEXT("ExternalMovieImport_Cancel", "Don't Copy");
-
-		FSuppressableWarningDialog ImportWarningDialog( Info );
-
-		if(ImportWarningDialog.ShowModal() != FSuppressableWarningDialog::EResult::Cancel)
+		// If the path of PickedPath is empty then we are dealing with a movie already in the correct directory.
+		if (FPaths::GetPath(PickedPath).IsEmpty())
 		{
-			const FString FileName = FPaths::GetCleanFilename(PickedPath);
-			const FString DestPath = MoviesBaseDir / FileName;
+			FString FullPickedPath = FPaths::Combine(MoviesBaseDir, PickedPath);
+
+			// Look for this movie.
+			TArray<FString> ExistingMovieFiles;
+			IFileManager::Get().FindFiles(ExistingMovieFiles, *MoviesBaseDir);
+			bool bHasValidMovie = ExistingMovieFiles.ContainsByPredicate(
+				[&PickedPath](const FString& ExistingMovie)
+			{
+				return ExistingMovie.Contains(PickedPath);
+			});
 
-			FText FailReason;
-			
-			if (SourceControlHelpers::CopyFileUnderSourceControl(DestPath, PickedPath, LOCTEXT("MovieFileDescription", "movie"), FailReason))
+			// Did we find a movie?
+			if (bHasValidMovie)
 			{
-				// trim the path so we just have a partial path with no extension (the movie player expects this)
-				Property->SetValue(FPaths::GetBaseFilename(DestPath.RightChop(MoviesBaseDir.Len())));
+				// Yes.
+				Property->SetValue(FPaths::GetBaseFilename(PickedPath));
 			}
 			else
 			{
-				FNotificationInfo FailureInfo(FailReason);
-				FailureInfo.ExpireDuration = 3.0f;
-				FSlateNotificationManager::Get().AddNotification(FailureInfo);
+				// Nope. Bring up a dialog box informing the user.
+				FSuppressableWarningDialog::FSetupInfo Info(
+					LOCTEXT("ExternalMovieImportNotExistWarning", "This movie does not exist."),
+					LOCTEXT("ExternalMovieImportNotExistTitle", "Does not exist"),
+					TEXT("ImportMovieIntoProjectNotExist"));
+				Info.ConfirmText = LOCTEXT("ExternalMovieImportNotExist_Confirm", "OK");
+				Info.CancelText = LOCTEXT("ExternalMovieImportNotExist_Cancel", "Cancel");
+
+				FSuppressableWarningDialog ImportWarningDialog(Info);
+				if (ImportWarningDialog.ShowModal() == FSuppressableWarningDialog::EResult::Cancel)
+				{
+					Property->SetValue(FPaths::GetBaseFilename(PickedPath));
+				}
 			}
-		}		
+		}
+		else
+		{
+			// ask the user if they want to import this movie
+			FSuppressableWarningDialog::FSetupInfo Info(
+				LOCTEXT("ExternalMovieImportWarning", "This movie needs to be copied into your project, would you like to copy the file now?"),
+				LOCTEXT("ExternalMovieImportTitle", "Copy Movie"),
+				TEXT("ImportMovieIntoProject"));
+			Info.ConfirmText = LOCTEXT("ExternalMovieImport_Confirm", "Copy");
+			Info.CancelText = LOCTEXT("ExternalMovieImport_Cancel", "Don't Copy");
+
+			FSuppressableWarningDialog ImportWarningDialog(Info);
+
+			if (ImportWarningDialog.ShowModal() != FSuppressableWarningDialog::EResult::Cancel)
+			{
+				const FString FileName = FPaths::GetCleanFilename(PickedPath);
+				const FString DestPath = MoviesBaseDir / FileName;
+
+				FText FailReason;
+
+				if (SourceControlHelpers::CopyFileUnderSourceControl(DestPath, PickedPath, LOCTEXT("MovieFileDescription", "movie"), FailReason))
+				{
+					// trim the path so we just have a partial path with no extension (the movie player expects this)
+					Property->SetValue(FPaths::GetBaseFilename(DestPath.RightChop(MoviesBaseDir.Len())));
+				}
+				else
+				{
+					FNotificationInfo FailureInfo(FailReason);
+					FailureInfo.ExpireDuration = 3.0f;
+					FSlateNotificationManager::Get().AddNotification(FailureInfo);
+				}
+			}
+		}
 	}
 	else
 	{
