PR #5659: Support maths expressions in margin shorthand textbox (Contributed by [REDACTED_PERSON])

diff --git a/Engine/Source/Editor/DetailCustomizations/Private/MarginCustomization.cpp b/Engine/Source/Editor/DetailCustomizations/Private/MarginCustomization.cpp
--- a/Engine/Source/Editor/DetailCustomizations/Private/MarginCustomization.cpp
+++ b/Engine/Source/Editor/DetailCustomizations/Private/MarginCustomization.cpp
@@ -23,6 +23,8 @@ void FMarginStructCustomization::CustomizeHeader( TSharedRef<class IPropertyHand
 	const FString& UVSpaceString( StructPropertyHandle->GetProperty()->GetMetaData( TEXT( "UVSpace" ) ) );
 	bIsMarginUsingUVSpace = UVSpaceString.Len() > 0 && UVSpaceString == TEXT( "true" );
 
+	NumericInterface = MakeShareable(new TDefaultNumericTypeInterface<float>);
+
 	uint32 NumChildren;
 	StructPropertyHandle->GetNumChildren( NumChildren );
 
@@ -102,6 +104,7 @@ TSharedRef<SWidget> FMarginStructCustomization::MakeChildPropertyWidget( int32 P
 		.MaxValue( bIsMarginUsingUVSpace ? 1.0f : TNumericLimits<float>::Max() )
 		.MinSliderValue( bIsMarginUsingUVSpace ? 0.0f : TNumericLimits<float>::Lowest() )
 		.MaxSliderValue( bIsMarginUsingUVSpace ? 1.0f : TNumericLimits<float>::Max()  )
+		.TypeInterface(NumericInterface)
 		.Label()
 		[
 			SNew( STextBlock )
@@ -141,10 +144,11 @@ void FMarginStructCustomization::OnMarginTextCommitted( const FText& InText, ETe
 
 				LeftString.TrimStartAndEndInline();
 
-				if( LeftString.IsNumeric() )
+				float Value = 0.f;
+				TOptional<float> NumericValue = NumericInterface->FromString(LeftString, Value);
+				if (NumericValue.IsSet())
 				{
-					float Value;
-					TTypeFromString<float>::FromString( Value, *LeftString );
+					Value = NumericValue.GetValue();
 					PropertyValues.Add( bIsMarginUsingUVSpace ? FMath::Clamp( Value, 0.0f, 1.0f ) : FMath::Max( Value, 0.0f ) );
 				}
 				else
diff --git a/Engine/Source/Editor/DetailCustomizations/Private/MarginCustomization.h b/Engine/Source/Editor/DetailCustomizations/Private/MarginCustomization.h
--- a/Engine/Source/Editor/DetailCustomizations/Private/MarginCustomization.h
+++ b/Engine/Source/Editor/DetailCustomizations/Private/MarginCustomization.h
@@ -4,6 +4,7 @@
 
 #include "CoreMinimal.h"
 #include "Widgets/SWidget.h"
+#include "Widgets/Input/NumericTypeInterface.h"
 #include "UnrealClient.h"
 #include "IPropertyTypeCustomization.h"
 #include "PropertyHandle.h"
@@ -103,5 +104,8 @@ private:
 
 	/** Margin text editable text box */
 	TSharedPtr<SEditableTextBox> MarginEditableTextBox;
+
+	/** Used to evaluate margin strings as mathematical expressions **/
+	TSharedPtr< INumericTypeInterface<float> > NumericInterface;
 };
 
