Fix SBrushResourceObjectBox changing size on first frames after construction (because of missing initialization of resource error widget collapsed state).

diff --git a/Engine/Source/Editor/DetailCustomizations/Private/SlateBrushCustomization.cpp b/Engine/Source/Editor/DetailCustomizations/Private/SlateBrushCustomization.cpp
--- a/Engine/Source/Editor/DetailCustomizations/Private/SlateBrushCustomization.cpp
+++ b/Engine/Source/Editor/DetailCustomizations/Private/SlateBrushCustomization.cpp
@@ -1148,29 +1148,13 @@ class SBrushResourceObjectBox : public SCompoundWidget
 				]
 			]
 		];
+
+		UpdateResourceError();
 	}
 
 	void Tick( const FGeometry& AllottedGeometry, const double InCurrentTime, const float InDeltaTime )
 	{
-		UObject* Resource = nullptr;
-
-		if( ResourceObjectProperty->GetValue(Resource) == FPropertyAccess::Success && Resource && Resource->IsA<UMaterialInterface>() )
-		{
-			UMaterialInterface* MaterialInterface = Cast<UMaterialInterface>( Resource );
-			UMaterial* BaseMaterial = MaterialInterface->GetBaseMaterial();
-			if( BaseMaterial && !BaseMaterial->IsUIMaterial() )
-			{
-				ResourceError->SetVisibility( EVisibility::Visible );
-			}
-			else
-			{
-				ResourceError->SetVisibility( EVisibility::Collapsed );
-			}
-		}
-		else if( ResourceError->GetVisibility() != EVisibility::Collapsed )
-		{
-			ResourceError->SetVisibility( EVisibility::Collapsed );
-		}
+		UpdateResourceError();
 	}
 
 private:
@@ -1238,6 +1222,29 @@ private:
 		}
 	}
 
+	void UpdateResourceError()
+	{
+		UObject* Resource = nullptr;
+
+		if (ResourceObjectProperty->GetValue(Resource) == FPropertyAccess::Success && Resource && Resource->IsA<UMaterialInterface>())
+		{
+			UMaterialInterface* MaterialInterface = Cast<UMaterialInterface>(Resource);
+			UMaterial* BaseMaterial = MaterialInterface->GetBaseMaterial();
+			if (BaseMaterial && !BaseMaterial->IsUIMaterial())
+			{
+				ResourceError->SetVisibility(EVisibility::Visible);
+			}
+			else
+			{
+				ResourceError->SetVisibility(EVisibility::Collapsed);
+			}
+		}
+		else if (ResourceError->GetVisibility() != EVisibility::Collapsed)
+		{
+			ResourceError->SetVisibility(EVisibility::Collapsed);
+		}
+	}
+
 private:
 	TSharedPtr<IPropertyHandle> ResourceObjectProperty;
 	TSharedPtr<IPropertyHandle> ImageSizeProperty;
