Fix SSlateBrushStaticPreview changing size on first frames after construction (because of late FSlateBrush member initialization).

diff --git a/Engine/Source/Editor/DetailCustomizations/Private/SlateBrushCustomization.cpp b/Engine/Source/Editor/DetailCustomizations/Private/SlateBrushCustomization.cpp
--- a/Engine/Source/Editor/DetailCustomizations/Private/SlateBrushCustomization.cpp
+++ b/Engine/Source/Editor/DetailCustomizations/Private/SlateBrushCustomization.cpp
@@ -981,6 +981,8 @@ class SSlateBrushStaticPreview : public SCompoundWidget
 	{
 		ResourceObjectProperty = InResourceObjectProperty;
 
+		UpdateBrush();
+
 		ChildSlot
 		[
 			SNew(SHorizontalBox)
@@ -1014,19 +1016,7 @@ class SSlateBrushStaticPreview : public SCompoundWidget
 
 	void Tick( const FGeometry& AllottedGeometry, const double InCurrentTime, const float InDeltaTime )
 	{
-		TArray<void*> RawData;
-
-		if (ResourceObjectProperty.IsValid() && ResourceObjectProperty->GetProperty())
-		{
-			ResourceObjectProperty->AccessRawData(RawData);
-
-			// RawData will be empty when creating a new Data Table, an idiosyncrasy
-			// of the Data Table Editor...
-			if (RawData.Num() > 0)
-			{
-				TemporaryBrush = *static_cast<FSlateBrush*>(RawData[0]);
-			}
-		}
+		UpdateBrush();
 	}
 
 private:
@@ -1061,6 +1051,22 @@ private:
 		return TemporaryBrush.DrawAs == ESlateBrushDrawType::Image ? EVisibility::Visible : EVisibility::Collapsed;
 	}
 
+	void UpdateBrush()
+	{
+		if (ResourceObjectProperty.IsValid() && ResourceObjectProperty->GetProperty())
+		{
+			TArray<void*> RawData;
+			ResourceObjectProperty->AccessRawData(RawData);
+
+			// RawData will be empty when creating a new Data Table, an idiosyncrasy
+			// of the Data Table Editor...
+			if (RawData.Num() > 0)
+			{
+				TemporaryBrush = *static_cast<FSlateBrush*>(RawData[0]);
+			}
+		}
+	}
+
 private:
 
 	/**
