Update curve struct customization to disable editing on curves when the corresponding property isn't editable

diff --git a/Engine/Source/Editor/DetailCustomizations/Private/CurveStructCustomization.cpp b/Engine/Source/Editor/DetailCustomizations/Private/CurveStructCustomization.cpp
--- a/Engine/Source/Editor/DetailCustomizations/Private/CurveStructCustomization.cpp
+++ b/Engine/Source/Editor/DetailCustomizations/Private/CurveStructCustomization.cpp
@@ -108,7 +108,7 @@ void FCurveStructCustomization::CustomizeHeader( TSharedRef<IPropertyHandle> InS
 		}
 		else
 		{
-			CurveWidget->SetCurveOwner(this);
+			CurveWidget->SetCurveOwner(this, InStructPropertyHandle->IsEditable());
 		}
 	}
 	else
@@ -276,7 +276,7 @@ void FCurveStructCustomization::OnExternalCurveChanged(TSharedRef<IPropertyHandl
 		}
 		else
 		{
-			CurveWidget->SetCurveOwner(this);
+			CurveWidget->SetCurveOwner(this, CurvePropertyHandle->IsEditable());
 		}
 
 		CurvePropertyHandle->NotifyPostChange();
