Additionally, the dropdown now opens the combo box with a single click after adding a pass.

diff --git a/Engine/Source/Editor/DetailCustomizations/Private/RenderPassesCustomization.cpp b/Engine/Source/Editor/DetailCustomizations/Private/RenderPassesCustomization.cpp
--- a/Engine/Source/Editor/DetailCustomizations/Private/RenderPassesCustomization.cpp
+++ b/Engine/Source/Editor/DetailCustomizations/Private/RenderPassesCustomization.cpp
@@ -100,10 +100,15 @@ public:
 			}
 			else
 			{
-				TSharedRef<SWidget> RemoveButton = PropertyCustomizationHelpers::MakeRemoveButton(FSimpleDelegate::CreateLambda([=]{
-					Property->Value.RemoveAt(PropertyIndex);
-					Update();
-				}));
+				TSharedRef<SWidget> RemoveButton = PropertyCustomizationHelpers::MakeRemoveButton(
+					FSimpleDelegate::CreateLambda([this, PropertyIndex] 
+					{
+						if (PropertyIndex < Property->Value.Num()) 
+						{
+							Property->Value.RemoveAt(PropertyIndex);
+							Update();
+						}
+					}));
 
 				EnabledPassesContainer->AddSlot()
 				.AutoHeight()
@@ -133,6 +138,8 @@ public:
 		ComboEntries.Sort([](const TSharedPtr<FRenderPassInfo>& A, const TSharedPtr<FRenderPassInfo>& B){
 			return A->Text.CompareToCaseIgnored(B->Text) < 0;
 		});
+
+		ComboBox->ClearSelection();
 		ComboBox->SetVisibility(ComboEntries.Num() == 0 ? EVisibility::Collapsed : EVisibility::Visible);
 	}
 	
