PR #5772: Fix UE-28209 (Contributed by [REDACTED_PERSON])

diff --git a/Engine/Source/Editor/DetailCustomizations/Private/MathStructProxyCustomizations.cpp b/Engine/Source/Editor/DetailCustomizations/Private/MathStructProxyCustomizations.cpp
--- a/Engine/Source/Editor/DetailCustomizations/Private/MathStructProxyCustomizations.cpp
+++ b/Engine/Source/Editor/DetailCustomizations/Private/MathStructProxyCustomizations.cpp
@@ -29,6 +29,7 @@ TSharedRef<SWidget> FMathStructProxyCustomization::MakeNumericProxyWidget(TShare
 
 	return 
 		SNew( SNumericEntryBox<NumericType> )
+		.IsEnabled( this, &FMathStructProxyCustomization::IsValueEnabled, WeakHandlePtr )
 		.Value( this, &FMathStructProxyCustomization::OnGetValue<ProxyType, NumericType>, WeakHandlePtr, ProxyValue )
 		.Font( IDetailLayoutBuilder::GetDetailFont() )
 		.UndeterminedString( NSLOCTEXT("PropertyEditor", "MultipleValues", "Multiple Values") )
