Added spinners to the Location and Scale in the Details view Transform section.

diff --git a/Engine/Plugins/Compositing/Composure/Source/ComposureLayersEditor/Private/ComposureDetailCustomizations.cpp b/Engine/Plugins/Compositing/Composure/Source/ComposureLayersEditor/Private/ComposureDetailCustomizations.cpp
--- a/Engine/Plugins/Compositing/Composure/Source/ComposureLayersEditor/Private/ComposureDetailCustomizations.cpp
+++ b/Engine/Plugins/Compositing/Composure/Source/ComposureLayersEditor/Private/ComposureDetailCustomizations.cpp
@@ -25,7 +25,6 @@
 #include "Widgets/Text/STextBlock.h"
 #include "Widgets/Input/SNumericEntryBox.h"
 #include "Misc/Optional.h"
-#include "Widgets/Input/SVectorInputBox.h"
 #include "Widgets/SWidget.h"
 #include "Widgets/Colors/SColorPicker.h"
 #include "Widgets/Input/SComboButton.h"
diff --git a/Engine/Source/Editor/AnimGraph/Private/PoseDriverDetails.cpp b/Engine/Source/Editor/AnimGraph/Private/PoseDriverDetails.cpp
--- a/Engine/Source/Editor/AnimGraph/Private/PoseDriverDetails.cpp
+++ b/Engine/Source/Editor/AnimGraph/Private/PoseDriverDetails.cpp
@@ -360,6 +360,7 @@ TSharedRef< SWidget > SPDD_TargetRow::GenerateWidgetForColumn(const FName& Colum
 					+SWidgetSwitcher::Slot()
 					[
 						SNew(SVectorInputBox)
+						.AllowSpin(true)
 						.X(this, &SPDD_TargetRow::GetTranslation, BoneIndex, EAxis::X)
 						.OnXChanged(this, &SPDD_TargetRow::SetTranslation, BoneIndex, EAxis::X)
 						.Y(this, &SPDD_TargetRow::GetTranslation, BoneIndex, EAxis::Y)
diff --git a/Engine/Source/Editor/DetailCustomizations/Private/ComponentTransformDetails.cpp b/Engine/Source/Editor/DetailCustomizations/Private/ComponentTransformDetails.cpp
--- a/Engine/Source/Editor/DetailCustomizations/Private/ComponentTransformDetails.cpp
+++ b/Engine/Source/Editor/DetailCustomizations/Private/ComponentTransformDetails.cpp
@@ -341,12 +341,18 @@ void FComponentTransformDetails::GenerateChildContent( IDetailChildrenBuilder& C
 				.bColorAxisLabels( true )
 				.AllowResponsiveLayout( true )
 				.IsEnabled( this, &FComponentTransformDetails::GetIsEnabled )
+				.OnXChanged( this, &FComponentTransformDetails::OnSetTransformAxis, ETextCommit::Default, ETransformField::Location, EAxisList::X, false )
+				.OnYChanged( this, &FComponentTransformDetails::OnSetTransformAxis, ETextCommit::Default, ETransformField::Location, EAxisList::Y, false )
+				.OnZChanged( this, &FComponentTransformDetails::OnSetTransformAxis, ETextCommit::Default, ETransformField::Location, EAxisList::Z, false )
 				.OnXCommitted( this, &FComponentTransformDetails::OnSetTransformAxis, ETransformField::Location, EAxisList::X, true )
 				.OnYCommitted( this, &FComponentTransformDetails::OnSetTransformAxis, ETransformField::Location, EAxisList::Y, true )
 				.OnZCommitted( this, &FComponentTransformDetails::OnSetTransformAxis, ETransformField::Location, EAxisList::Z, true )
 				.Font( FontInfo )
 				.TypeInterface( TypeInterface )
-				.AllowSpin(false)
+				.AllowSpin( SelectedObjects.Num() == 1 )
+				.SpinDelta( 1 )
+				.OnBeginSliderMovement( this, &FComponentTransformDetails::OnBeginLocationSlider )
+				.OnEndSliderMovement( this, &FComponentTransformDetails::OnEndLocationSlider )
 			]
 			+SHorizontalBox::Slot()
 			.AutoWidth()
@@ -409,7 +415,7 @@ void FComponentTransformDetails::GenerateChildContent( IDetailChildrenBuilder& C
 				.AllowResponsiveLayout( true )
 				.bColorAxisLabels( true )
 				.IsEnabled( this, &FComponentTransformDetails::GetIsEnabled )
-				.OnBeginSliderMovement( this, &FComponentTransformDetails::OnBeginRotatonSlider )
+				.OnBeginSliderMovement( this, &FComponentTransformDetails::OnBeginRotationSlider )
 				.OnEndSliderMovement( this, &FComponentTransformDetails::OnEndRotationSlider )
 				.OnRollChanged( this, &FComponentTransformDetails::OnSetTransformAxis, ETextCommit::Default, ETransformField::Rotation, EAxisList::X, false )
 				.OnPitchChanged( this, &FComponentTransformDetails::OnSetTransformAxis, ETextCommit::Default, ETransformField::Rotation, EAxisList::Y, false )
@@ -417,7 +423,7 @@ void FComponentTransformDetails::GenerateChildContent( IDetailChildrenBuilder& C
 				.OnRollCommitted( this, &FComponentTransformDetails::OnSetTransformAxis, ETransformField::Rotation, EAxisList::X, true )
 				.OnPitchCommitted( this, &FComponentTransformDetails::OnSetTransformAxis, ETransformField::Rotation, EAxisList::Y, true )
 				.OnYawCommitted( this, &FComponentTransformDetails::OnSetTransformAxis, ETransformField::Rotation, EAxisList::Z, true )
-				.TypeInterface(TypeInterface)
+				.TypeInterface( TypeInterface )
 				.Font( FontInfo )
 			]
 			+SHorizontalBox::Slot()
@@ -474,6 +480,9 @@ void FComponentTransformDetails::GenerateChildContent( IDetailChildrenBuilder& C
 				.bColorAxisLabels( true )
 				.AllowResponsiveLayout( true )
 				.IsEnabled( this, &FComponentTransformDetails::GetIsEnabled )
+				.OnXChanged( this, &FComponentTransformDetails::OnSetTransformAxis, ETextCommit::Default, ETransformField::Scale, EAxisList::X, false )
+				.OnYChanged( this, &FComponentTransformDetails::OnSetTransformAxis, ETextCommit::Default, ETransformField::Scale, EAxisList::Y, false )
+				.OnZChanged( this, &FComponentTransformDetails::OnSetTransformAxis, ETextCommit::Default, ETransformField::Scale, EAxisList::Z, false )
 				.OnXCommitted( this, &FComponentTransformDetails::OnSetTransformAxis, ETransformField::Scale, EAxisList::X, true )
 				.OnYCommitted( this, &FComponentTransformDetails::OnSetTransformAxis, ETransformField::Scale, EAxisList::Y, true )
 				.OnZCommitted( this, &FComponentTransformDetails::OnSetTransformAxis, ETransformField::Scale, EAxisList::Z, true )
@@ -481,7 +490,10 @@ void FComponentTransformDetails::GenerateChildContent( IDetailChildrenBuilder& C
 				.ContextMenuExtenderY( this, &FComponentTransformDetails::ExtendYScaleContextMenu )
 				.ContextMenuExtenderZ( this, &FComponentTransformDetails::ExtendZScaleContextMenu )
 				.Font( FontInfo )
-				.AllowSpin(false)
+				.AllowSpin( SelectedObjects.Num() == 1 )
+				.SpinDelta( 0.0025f )
+				.OnBeginSliderMovement( this, &FComponentTransformDetails::OnBeginScaleSlider )
+				.OnEndSliderMovement( this, &FComponentTransformDetails::OnEndScaleSlider )
 			]
 			+SHorizontalBox::Slot()
 			.AutoWidth()
@@ -1300,40 +1312,37 @@ void FComponentTransformDetails::OnSetTransformAxis(float NewValue, ETextCommit:
 	OnSetTransform(TransformField, Axis, NewVector, false, bCommitted);
 }
 
-void FComponentTransformDetails::OnBeginRotatonSlider()
+void FComponentTransformDetails::BeginSliderTransaction(FText ActorTransaction, FText ComponentTransaction) const
 {
-	bEditingRotationInUI = true;
-
 	bool bBeganTransaction = false;
-	for( int32 ObjectIndex = 0; ObjectIndex < SelectedObjects.Num(); ++ObjectIndex )
+	for (TWeakObjectPtr<UObject> ObjectPtr : SelectedObjects)
 	{
-		TWeakObjectPtr<UObject> ObjectPtr = SelectedObjects[ObjectIndex];
-		if( ObjectPtr.IsValid() )
+		if (ObjectPtr.IsValid())
 		{
 			UObject* Object = ObjectPtr.Get();
 
-			// Start a new transation when a rotator slider begins to change
-			// We'll end it when the slider is release
+			// Start a new transaction when a slider begins to change
+			// We'll end it when the slider is released
 			// NOTE: One transaction per change, not per actor
-			if(!bBeganTransaction)
+			if (!bBeganTransaction)
 			{
-				if(Object->IsA<AActor>())
+				if (Object->IsA<AActor>())
 				{
-					GEditor->BeginTransaction( LOCTEXT( "OnSetRotation", "Set Rotation" ) );
+					GEditor->BeginTransaction(ActorTransaction);
 				}
 				else
 				{
-					GEditor->BeginTransaction( LOCTEXT( "OnSetRotation_ComponentDirect", "Modify Component(s)") );
+					GEditor->BeginTransaction(ComponentTransaction);
 				}
 
 				bBeganTransaction = true;
 			}
 
-			USceneComponent* SceneComponent = GetSceneComponentFromDetailsObject( Object );
-			if( SceneComponent )
+			USceneComponent* SceneComponent = GetSceneComponentFromDetailsObject(Object);
+			if (SceneComponent)
 			{
-				FScopedSwitchWorldForObject WorldSwitcher( Object );
-				
+				FScopedSwitchWorldForObject WorldSwitcher(Object);
+
 				if (SceneComponent->HasAnyFlags(RF_DefaultSubObject))
 				{
 					// Default subobjects must be included in any undo/redo operations
@@ -1342,18 +1351,41 @@ void FComponentTransformDetails::OnBeginRotatonSlider()
 
 				// Call modify but not PreEdit, we don't do the proper "Edit" until it's committed
 				SceneComponent->Modify();
-
-				// Add/update cached rotation value prior to slider interaction
-				ObjectToRelativeRotationMap.FindOrAdd(SceneComponent) = SceneComponent->RelativeRotation;
 			}
 		}
 	}
 
 	// Just in case we couldn't start a new transaction for some reason
-	if(!bBeganTransaction)
+	if (!bBeganTransaction)
 	{
-		GEditor->BeginTransaction( LOCTEXT( "OnSetRotation", "Set Rotation" ) );
-	}	
+		GEditor->BeginTransaction(ActorTransaction);
+	}
+}
+
+void FComponentTransformDetails::OnBeginRotationSlider()
+{
+	FText ActorTransaction = LOCTEXT("OnSetRotation", "Set Rotation");
+	FText ComponentTransaction = LOCTEXT("OnSetRotation_ComponentDirect", "Modify Component(s)");
+	BeginSliderTransaction(ActorTransaction, ComponentTransaction);
+
+	bEditingRotationInUI = true;
+
+	for (TWeakObjectPtr<UObject> ObjectPtr : SelectedObjects)
+	{
+		if (ObjectPtr.IsValid())
+		{
+			UObject* Object = ObjectPtr.Get();
+
+			USceneComponent* SceneComponent = GetSceneComponentFromDetailsObject(Object);
+			if (SceneComponent)
+			{
+				FScopedSwitchWorldForObject WorldSwitcher(Object);
+
+				// Add/update cached rotation value prior to slider interaction
+				ObjectToRelativeRotationMap.FindOrAdd(SceneComponent) = SceneComponent->RelativeRotation;
+			}
+		}
+	}
 }
 
 void FComponentTransformDetails::OnEndRotationSlider(float NewValue)
@@ -1363,4 +1395,28 @@ void FComponentTransformDetails::OnEndRotationSlider(float NewValue)
 	GEditor->EndTransaction();
 }
 
+void FComponentTransformDetails::OnBeginLocationSlider()
+{
+	FText ActorTransaction = LOCTEXT("OnSetLocation", "Set Location");
+	FText ComponentTransaction = LOCTEXT("OnSetLocation_ComponentDirect", "Modify Component Location");
+	BeginSliderTransaction(ActorTransaction, ComponentTransaction);
+}
+
+void FComponentTransformDetails::OnEndLocationSlider(float NewValue)
+{
+	GEditor->EndTransaction();
+}
+
+void FComponentTransformDetails::OnBeginScaleSlider()
+{
+	FText ActorTransaction = LOCTEXT("OnSetScale", "Set Scale");
+	FText ComponentTransaction = LOCTEXT("OnSetScale_ComponentDirect", "Modify Component Scale");
+	BeginSliderTransaction(ActorTransaction, ComponentTransaction);
+}
+
+void FComponentTransformDetails::OnEndScaleSlider(float NewValue)
+{
+	GEditor->EndTransaction();
+}
+
 #undef LOCTEXT_NAMESPACE
diff --git a/Engine/Source/Editor/DetailCustomizations/Private/ComponentTransformDetails.h b/Engine/Source/Editor/DetailCustomizations/Private/ComponentTransformDetails.h
--- a/Engine/Source/Editor/DetailCustomizations/Private/ComponentTransformDetails.h
+++ b/Engine/Source/Editor/DetailCustomizations/Private/ComponentTransformDetails.h
@@ -81,16 +81,32 @@ private:
 	 */
 	void OnSetTransformAxis(float NewValue, ETextCommit::Type CommitInfo, ETransformField::Type TransformField, EAxisList::Type Axis, bool bCommitted);
 
-	/**
-	 * Called when the one of the axis sliders for object rotation begins to change for the first time 
-	 */
-	void OnBeginRotatonSlider();
 
-	/**
-	 * Called when the one of the axis sliders for object rotation is released
+	/** 
+	 * Helper to begin a new transaction for a slider interaction. 
+	 * @param ActorTransaction		The name to give the transaction when changing an actor transform.
+	 * @param ComponentTransaction	The name to give the transaction when directly editing a component transform.
 	 */
+	void BeginSliderTransaction(FText ActorTransaction, FText ComponentTransaction) const;
+
+	/** Called when the one of the axis sliders for object rotation begins to change for the first time */
+	void OnBeginRotationSlider();
+
+	/** Called when the one of the axis sliders for object rotation is released */
 	void OnEndRotationSlider(float NewValue);
 
+	/** Called when one of the axis sliders for object location begins to change */
+	void OnBeginLocationSlider();
+
+	/** Called when one of the axis sliders for object location is released */
+	void OnEndLocationSlider(float NewValue);
+
+	/** Called when one of the axis sliders for object scale begins to change */
+	void OnBeginScaleSlider();
+
+	/** Called when one of the axis sliders for object scale is released */
+	void OnEndScaleSlider(float NewValue);
+
 	/** @return Icon to use in the preserve scale ratio check box */
 	const FSlateBrush* GetPreserveScaleRatioImage() const;
 
@@ -216,7 +232,7 @@ private:
 	/** @return The visibility of the "Reset to Default" button for the scale component */
 	EVisibility GetScaleResetVisibility() const;
 
-	/** Cache a single unit to display all location comonents in */
+	/** Cache a single unit to display all location components in */
 	void CacheCommonLocationUnits();
 
 private:
@@ -254,7 +270,7 @@ private:
 		TOptional<float> Y;
 		TOptional<float> Z;
 	};
-	
+
 	FSelectedActorInfo SelectedActorInfo;
 	/** Copy of selected actor references in the details view */
 	TArray< TWeakObjectPtr<UObject> > SelectedObjects;
diff --git a/Engine/Source/Editor/DetailCustomizations/Private/SplineComponentDetails.cpp b/Engine/Source/Editor/DetailCustomizations/Private/SplineComponentDetails.cpp
--- a/Engine/Source/Editor/DetailCustomizations/Private/SplineComponentDetails.cpp
+++ b/Engine/Source/Editor/DetailCustomizations/Private/SplineComponentDetails.cpp
@@ -272,6 +272,7 @@ void FSplinePointDetails::GenerateChildContent(IDetailChildrenBuilder& ChildrenB
 		.Y(this, &FSplinePointDetails::GetPositionY)
 		.Z(this, &FSplinePointDetails::GetPositionZ)
 		.AllowResponsiveLayout(true)
+		.AllowSpin(false)
 		.OnXCommitted(this, &FSplinePointDetails::OnSetPosition, 0)
 		.OnYCommitted(this, &FSplinePointDetails::OnSetPosition, 1)
 		.OnZCommitted(this, &FSplinePointDetails::OnSetPosition, 2)
@@ -298,6 +299,7 @@ void FSplinePointDetails::GenerateChildContent(IDetailChildrenBuilder& ChildrenB
 		.Y(this, &FSplinePointDetails::GetArriveTangentY)
 		.Z(this, &FSplinePointDetails::GetArriveTangentZ)
 		.AllowResponsiveLayout(true)
+		.AllowSpin(false)
 		.OnXCommitted(this, &FSplinePointDetails::OnSetArriveTangent, 0)
 		.OnYCommitted(this, &FSplinePointDetails::OnSetArriveTangent, 1)
 		.OnZCommitted(this, &FSplinePointDetails::OnSetArriveTangent, 2)
@@ -324,6 +326,7 @@ void FSplinePointDetails::GenerateChildContent(IDetailChildrenBuilder& ChildrenB
 		.Y(this, &FSplinePointDetails::GetLeaveTangentY)
 		.Z(this, &FSplinePointDetails::GetLeaveTangentZ)
 		.AllowResponsiveLayout(true)
+		.AllowSpin(false)
 		.OnXCommitted(this, &FSplinePointDetails::OnSetLeaveTangent, 0)
 		.OnYCommitted(this, &FSplinePointDetails::OnSetLeaveTangent, 1)
 		.OnZCommitted(this, &FSplinePointDetails::OnSetLeaveTangent, 2)
@@ -350,6 +353,7 @@ void FSplinePointDetails::GenerateChildContent(IDetailChildrenBuilder& ChildrenB
 		.Y(this, &FSplinePointDetails::GetRotationPitch)
 		.Z(this, &FSplinePointDetails::GetRotationYaw)
 		.AllowResponsiveLayout(true)
+		.AllowSpin(false)
 		.OnXCommitted(this, &FSplinePointDetails::OnSetRotation, 0)
 		.OnYCommitted(this, &FSplinePointDetails::OnSetRotation, 1)
 		.OnZCommitted(this, &FSplinePointDetails::OnSetRotation, 2)
@@ -375,6 +379,7 @@ void FSplinePointDetails::GenerateChildContent(IDetailChildrenBuilder& ChildrenB
 		.X(this, &FSplinePointDetails::GetScaleX)
 		.Y(this, &FSplinePointDetails::GetScaleY)
 		.Z(this, &FSplinePointDetails::GetScaleZ)
+		.AllowSpin(false)
 		.AllowResponsiveLayout(true)
 		.OnXCommitted(this, &FSplinePointDetails::OnSetScale, 0)
 		.OnYCommitted(this, &FSplinePointDetails::OnSetScale, 1)
diff --git a/Engine/Source/Editor/LevelEditor/Private/DlgDeltaTransform.cpp b/Engine/Source/Editor/LevelEditor/Private/DlgDeltaTransform.cpp
--- a/Engine/Source/Editor/LevelEditor/Private/DlgDeltaTransform.cpp
+++ b/Engine/Source/Editor/LevelEditor/Private/DlgDeltaTransform.cpp
@@ -72,6 +72,7 @@ public:
 							.Z( this, &SDlgDeltaTransform::GetDeltaZ )
 							.bColorAxisLabels( true )
 							.AllowResponsiveLayout( true )
+							.AllowSpin( false )
 							.OnXCommitted( this, &SDlgDeltaTransform::OnSetDelta, 0 )
 							.OnYCommitted( this, &SDlgDeltaTransform::OnSetDelta, 1 )
 							.OnZCommitted( this, &SDlgDeltaTransform::OnSetDelta, 2 )
diff --git a/Engine/Source/Editor/StaticMeshEditor/Private/StaticMeshEditorTools.cpp b/Engine/Source/Editor/StaticMeshEditor/Private/StaticMeshEditorTools.cpp
--- a/Engine/Source/Editor/StaticMeshEditor/Private/StaticMeshEditorTools.cpp
+++ b/Engine/Source/Editor/StaticMeshEditor/Private/StaticMeshEditorTools.cpp
@@ -627,6 +627,7 @@ void FMeshBuildSettingsLayout::GenerateChildContent( IDetailChildrenBuilder& Chi
 			.Z(this, &FMeshBuildSettingsLayout::GetBuildScaleZ)
 			.bColorAxisLabels(false)
 			.AllowResponsiveLayout(true)
+			.AllowSpin(false)
 			.OnXCommitted(this, &FMeshBuildSettingsLayout::OnBuildScaleXChanged)
 			.OnYCommitted(this, &FMeshBuildSettingsLayout::OnBuildScaleYChanged)
 			.OnZCommitted(this, &FMeshBuildSettingsLayout::OnBuildScaleZChanged)
diff --git a/Engine/Source/Editor/WorldBrowser/Private/StreamingLevels/StreamingLevelCustomization.cpp b/Engine/Source/Editor/WorldBrowser/Private/StreamingLevels/StreamingLevelCustomization.cpp
--- a/Engine/Source/Editor/WorldBrowser/Private/StreamingLevels/StreamingLevelCustomization.cpp
+++ b/Engine/Source/Editor/WorldBrowser/Private/StreamingLevels/StreamingLevelCustomization.cpp
@@ -67,6 +67,7 @@ void FStreamingLevelCustomization::CustomizeDetails(IDetailLayoutBuilder& Detail
 			.Font(FEditorStyle::GetFontStyle(TEXT("PropertyWindow.NormalFont")))
 			.bColorAxisLabels(true)
 			.AllowResponsiveLayout(true)
+			.AllowSpin(false)
 			.X(this, &FStreamingLevelCustomization::OnGetLevelPosition, 0)
 			.Y(this, &FStreamingLevelCustomization::OnGetLevelPosition, 1)
 			.Z(this, &FStreamingLevelCustomization::OnGetLevelPosition, 2)
diff --git a/Engine/Source/Editor/WorldBrowser/Private/Tiles/STiledLandscapeImportDlg.cpp b/Engine/Source/Editor/WorldBrowser/Private/Tiles/STiledLandscapeImportDlg.cpp
--- a/Engine/Source/Editor/WorldBrowser/Private/Tiles/STiledLandscapeImportDlg.cpp
+++ b/Engine/Source/Editor/WorldBrowser/Private/Tiles/STiledLandscapeImportDlg.cpp
@@ -177,6 +177,7 @@ void STiledLandcapeImportDlg::Construct(const FArguments& InArgs, TSharedPtr<SWi
 					SNew( SVectorInputBox )
 					.bColorAxisLabels( true )
 					.AllowResponsiveLayout( true )
+					.AllowSpin(false)
 					.X( this, &STiledLandcapeImportDlg::GetScaleX )
 					.Y( this, &STiledLandcapeImportDlg::GetScaleY )
 					.Z( this, &STiledLandcapeImportDlg::GetScaleZ )
diff --git a/Engine/Source/Runtime/Slate/Private/Widgets/Input/SVectorInputBox.cpp b/Engine/Source/Runtime/Slate/Private/Widgets/Input/SVectorInputBox.cpp
--- a/Engine/Source/Runtime/Slate/Private/Widgets/Input/SVectorInputBox.cpp
+++ b/Engine/Source/Runtime/Slate/Private/Widgets/Input/SVectorInputBox.cpp
@@ -56,6 +56,14 @@ void SVectorInputBox::ConstructX( const FArguments& InArgs, TSharedRef<SHorizont
 		.OverrideTextMargin(MarginAttribute)
 		.ContextMenuExtender( InArgs._ContextMenuExtenderX )
 		.TypeInterface(InArgs._TypeInterface)
+		.MinValue(TOptional<float>())
+		.MaxValue(TOptional<float>())
+		.MinSliderValue(TOptional<float>())
+		.MaxSliderValue(TOptional<float>())
+		.LinearDeltaSensitivity(1)
+		.Delta(InArgs._SpinDelta)
+		.OnBeginSliderMovement(InArgs._OnBeginSliderMovement)
+		.OnEndSliderMovement(InArgs._OnEndSliderMovement)
 		.Label()
 		[
 			LabelWidget
@@ -91,6 +99,14 @@ void SVectorInputBox::ConstructY( const FArguments& InArgs, TSharedRef<SHorizont
 		.OverrideTextMargin(MarginAttribute)
 		.ContextMenuExtender(InArgs._ContextMenuExtenderY)
 		.TypeInterface(InArgs._TypeInterface)
+		.MinValue(TOptional<float>())
+		.MaxValue(TOptional<float>())
+		.MinSliderValue(TOptional<float>())
+		.MaxSliderValue(TOptional<float>())
+		.LinearDeltaSensitivity(1)
+		.Delta(InArgs._SpinDelta)
+		.OnBeginSliderMovement(InArgs._OnBeginSliderMovement)
+		.OnEndSliderMovement(InArgs._OnEndSliderMovement)
 		.Label()
 		[
 			LabelWidget
@@ -126,6 +142,14 @@ void SVectorInputBox::ConstructZ( const FArguments& InArgs, TSharedRef<SHorizont
 		.OverrideTextMargin(MarginAttribute)
 		.ContextMenuExtender(InArgs._ContextMenuExtenderZ)
 		.TypeInterface(InArgs._TypeInterface)
+		.MinValue(TOptional<float>())
+		.MaxValue(TOptional<float>())
+		.MinSliderValue(TOptional<float>())
+		.MaxSliderValue(TOptional<float>())
+		.LinearDeltaSensitivity(1)
+		.Delta(InArgs._SpinDelta)
+		.OnBeginSliderMovement(InArgs._OnBeginSliderMovement)
+		.OnEndSliderMovement(InArgs._OnEndSliderMovement)
 		.Label()
 		[
 			LabelWidget
diff --git a/Engine/Source/Runtime/Slate/Public/Widgets/Input/SNumericEntryBox.h b/Engine/Source/Runtime/Slate/Public/Widgets/Input/SNumericEntryBox.h
--- a/Engine/Source/Runtime/Slate/Public/Widgets/Input/SNumericEntryBox.h
+++ b/Engine/Source/Runtime/Slate/Public/Widgets/Input/SNumericEntryBox.h
@@ -106,6 +106,8 @@ public:
 		SLATE_ARGUMENT( bool, AllowSpin )
 		/** How many pixel the mouse must move to change the value of the delta step (only use if there is a spinbox allow) */
 		SLATE_ARGUMENT( int32, ShiftMouseMovePixelPerDelta )
+		/** If we're an unbounded spinbox, what value do we divide mouse movement by before multiplying by Delta. Requires Delta to be set. */
+		SLATE_ATTRIBUTE( int32, LinearDeltaSensitivity)
 		/** Tell us if we want to support dynamically changing of the max value using ctrl  (only use if there is a spinbox allow) */
 		SLATE_ATTRIBUTE(bool, SupportDynamicSliderMaxValue)
 		/** Tell us if we want to support dynamically changing of the min value using ctrl  (only use if there is a spinbox allow) */
@@ -183,6 +185,7 @@ public:
 				.Value( this, &SNumericEntryBox<NumericType>::OnGetValueForSpinBox )
 				.Delta( InArgs._Delta )
 				.ShiftMouseMovePixelPerDelta(InArgs._ShiftMouseMovePixelPerDelta)
+				.LinearDeltaSensitivity(InArgs._LinearDeltaSensitivity)
 				.SupportDynamicSliderMaxValue(InArgs._SupportDynamicSliderMaxValue)
 				.SupportDynamicSliderMinValue(InArgs._SupportDynamicSliderMinValue)
 				.OnDynamicSliderMaxValueChanged(InArgs._OnDynamicSliderMaxValueChanged)
diff --git a/Engine/Source/Runtime/Slate/Public/Widgets/Input/SVectorInputBox.h b/Engine/Source/Runtime/Slate/Public/Widgets/Input/SVectorInputBox.h
--- a/Engine/Source/Runtime/Slate/Public/Widgets/Input/SVectorInputBox.h
+++ b/Engine/Source/Runtime/Slate/Public/Widgets/Input/SVectorInputBox.h
@@ -26,7 +26,8 @@ class SLATE_API SVectorInputBox : public SCompoundWidget
 public:
 	SLATE_BEGIN_ARGS( SVectorInputBox )
 		: _Font( FCoreStyle::Get().GetFontStyle("NormalFont") )
-		, _AllowSpin(true)
+		, _AllowSpin( false )
+		, _SpinDelta( 1 )
 		, _bColorAxisLabels( false )
 		, _AllowResponsiveLayout( false )
 		{}
@@ -44,13 +45,16 @@ public:
 		SLATE_ATTRIBUTE( FSlateFontInfo, Font )
 
 		/** Whether or not values can be spun or if they should be typed in */
-		SLATE_ARGUMENT(bool, AllowSpin)
+		SLATE_ARGUMENT( bool, AllowSpin )
+
+		/** The delta amount to apply, per pixel, when the spinner is dragged. */
+		SLATE_ATTRIBUTE( float, SpinDelta )
 
 		/** Should the axis labels be colored */
 		SLATE_ARGUMENT( bool, bColorAxisLabels )		
 
 		/** Allow responsive layout to crush the label and margins when there is not a lot of room */
-		SLATE_ARGUMENT(bool, AllowResponsiveLayout)
+		SLATE_ARGUMENT( bool, AllowResponsiveLayout )
 
 		/** Called when the x value of the vector is changed */
 		SLATE_EVENT( FOnFloatValueChanged, OnXChanged )
@@ -79,6 +83,12 @@ public:
 		/** Menu extender delegate for the Z value */
 		SLATE_EVENT( FMenuExtensionDelegate, ContextMenuExtenderZ )
 
+		/** Called right before the slider begins to move for any of the vector components */
+		SLATE_EVENT( FSimpleDelegate, OnBeginSliderMovement )
+		
+		/** Called right after the slider handle is released by the user for any of the vector components */
+		SLATE_EVENT( FOnFloatValueChanged, OnEndSliderMovement )
+
 		/** Provide custom type functionality for the vector */
 		SLATE_ARGUMENT( TSharedPtr< INumericTypeInterface<float> >, TypeInterface )
 
