Renamed public facing "Early Access" messaging to "Beta".  Internal metadata and other keywords are unchanged

diff --git a/Engine/Plugins/2D/Paper2D/Source/Paper2D/Classes/PaperTileMapComponent.h b/Engine/Plugins/2D/Paper2D/Source/Paper2D/Classes/PaperTileMapComponent.h
--- a/Engine/Plugins/2D/Paper2D/Source/Paper2D/Classes/PaperTileMapComponent.h
+++ b/Engine/Plugins/2D/Paper2D/Source/Paper2D/Classes/PaperTileMapComponent.h
@@ -22,13 +22,13 @@ struct FSpriteRenderSection;
  * This component is created when you drag a tile map asset from the content browser into a Blueprint, or
  * contained inside of the actor created when you drag one into the level.
  *
- * NOTE: This is an early access preview class.  While not considered production-ready, it is a step beyond
+ * NOTE: This is an beta preview class.  While not considered production-ready, it is a step beyond
  * 'experimental' and is being provided as a preview of things to come:
  *  - We will try to provide forward-compatibility for content you create.
  *  - The classes may change significantly in the future.
  *  - The code is in an early state and may not meet the desired polish / quality bar.
  *  - There is probably no documentation or example content yet.
- *  - They will be promoted out of 'Early Access' when they are production ready.
+ *  - They will be promoted out of 'beta' when they are production ready.
  *
  * @see UPrimitiveComponent, UPaperTileMap
  */
diff --git a/Engine/Source/Editor/DetailCustomizations/Private/ObjectDetails.cpp b/Engine/Source/Editor/DetailCustomizations/Private/ObjectDetails.cpp
--- a/Engine/Source/Editor/DetailCustomizations/Private/ObjectDetails.cpp
+++ b/Engine/Source/Editor/DetailCustomizations/Private/ObjectDetails.cpp
@@ -46,9 +46,9 @@ void FObjectDetails::AddExperimentalWarningCategory(IDetailLayoutBuilder& Detail
 		const FText CategoryDisplayName = LOCTEXT("WarningCategoryDisplayName", "Warning");
 		FString ClassUsed = DetailBuilder.GetTopLevelProperty().ToString();
 		const FText WarningText = bBaseClassIsExperimental ? FText::Format( LOCTEXT("ExperimentalClassWarning", "Uses experimental class: {0}") , FText::FromString(ClassUsed) )
-			: FText::Format( LOCTEXT("EarlyAccessClassWarning", "Uses early access class {0}"), FText::FromString(*ClassUsed) );
+			: FText::Format( LOCTEXT("EarlyAccessClassWarning", "Uses beta class {0}"), FText::FromString(*ClassUsed) );
 		const FText SearchString = WarningText;
-		const FText Tooltip = bBaseClassIsExperimental ? LOCTEXT("ExperimentalClassTooltip", "Here be dragons!  Uses one or more unsupported 'experimental' classes") : LOCTEXT("EarlyAccessClassTooltip", "Uses one or more 'early access' classes");
+		const FText Tooltip = bBaseClassIsExperimental ? LOCTEXT("ExperimentalClassTooltip", "Here be dragons!  Uses one or more unsupported 'experimental' classes") : LOCTEXT("EarlyAccessClassTooltip", "Uses one or more 'beta' classes");
 		const FString ExcerptName = bBaseClassIsExperimental ? TEXT("ObjectUsesExperimentalClass") : TEXT("ObjectUsesEarlyAccessClass");
 		const FSlateBrush* WarningIcon = FEditorStyle::GetBrush(bBaseClassIsExperimental ? "PropertyEditor.ExperimentalClass" : "PropertyEditor.EarlyAccessClass");
 
