Allow inline reduction of custom LOD

diff --git a/Engine/Source/Editor/DetailCustomizations/Private/SkeletalMeshReductionSettingsDetails.cpp b/Engine/Source/Editor/DetailCustomizations/Private/SkeletalMeshReductionSettingsDetails.cpp
--- a/Engine/Source/Editor/DetailCustomizations/Private/SkeletalMeshReductionSettingsDetails.cpp
+++ b/Engine/Source/Editor/DetailCustomizations/Private/SkeletalMeshReductionSettingsDetails.cpp
@@ -15,6 +15,8 @@
 #include "Widgets/Text/STextBlock.h"
 #include "SkeletalRenderPublic.h"
 
+#include "Rendering/SkeletalMeshModel.h"
+
 #define LOCTEXT_NAMESPACE "SkeletalMeshReductionSettingsDetails"
 
 TSharedRef<IPropertyTypeCustomization> FSkeletalMeshReductionSettingsDetails::MakeInstance()
@@ -54,6 +56,14 @@ void FSkeletalMeshReductionSettingsDetails::CustomizeChildren(TSharedRef<IProper
 		GET_MEMBER_NAME_CHECKED(FSkeletalMeshOptimizationSettings, MaxDeviationPercentage)
 	};
 
+	USkeletalMesh* SkeletalMesh = nullptr;
+	TArray<UObject*> Objects;
+	StructPropertyHandle->GetOuterObjects(Objects);
+	if (Objects.Num() > 0)
+	{
+		SkeletalMesh = Cast<USkeletalMesh>(Objects[0]);
+	}
+
 	ReductionMethodPropertyHandle = StructPropertyHandle->GetChildHandle(CustomizedProperties[0]);
 	NumTrianglesPercentagePropertyHandle = StructPropertyHandle->GetChildHandle(CustomizedProperties[1]);
 	MaxDeviationPercentagePropertyHandle = StructPropertyHandle->GetChildHandle(CustomizedProperties[2]);
@@ -73,11 +83,13 @@ void FSkeletalMeshReductionSettingsDetails::CustomizeChildren(TSharedRef<IProper
 		return INDEX_NONE;
 	}();
 
-	auto BaseLODCustomization = [LODIndex, &StructBuilder](TSharedPtr<IPropertyHandle>& BaseLODPropertyHandle)
+	auto BaseLODCustomization = [LODIndex, &StructBuilder, &SkeletalMesh](TSharedPtr<IPropertyHandle>& BaseLODPropertyHandle)
 	{
-		// Only able to do this for LOD2 and above, so only show the property if this is the case
-		if (LODIndex > 1)
+		// Only able to do this for LOD1 and above, so only show the property if this is the case
+		if (LODIndex > 0)
 		{
+			const FSkeletalMeshModel* ImportedModel = SkeletalMesh != nullptr ? SkeletalMesh->GetImportedModel() : nullptr;
+			bool AllowInline = ImportedModel != nullptr && ImportedModel->LODModels.IsValidIndex(LODIndex) && !ImportedModel->LODModels[LODIndex].RawSkeletalMeshBulkData.IsEmpty();
 			// Add and retrieve the default widgets
 			IDetailPropertyRow& Row = StructBuilder.AddProperty(BaseLODPropertyHandle->AsShared());
 
@@ -99,7 +111,7 @@ void FSkeletalMeshReductionSettingsDetails::CustomizeChildren(TSharedRef<IProper
 					SNew(SSpinBox<int32>)
 					.Font(IDetailLayoutBuilder::GetDetailFont())
 					.MinValue(0)
-					.MaxValue(FMath::Max(LODIndex - 1, 0))
+					.MaxValue(FMath::Max(LODIndex - (AllowInline ? 0 : 1), 0))
 					.Value_Lambda([BaseLODPropertyHandle]() -> int32
 					{
 						int32 Value = INDEX_NONE;
