PR #5861: UE-73008: Only include unique function names to avoid duplicates in the button list  (Contributed by [REDACTED_PERSON])

diff --git a/Engine/Source/Editor/DetailCustomizations/Private/ObjectDetails.cpp b/Engine/Source/Editor/DetailCustomizations/Private/ObjectDetails.cpp
--- a/Engine/Source/Editor/DetailCustomizations/Private/ObjectDetails.cpp
+++ b/Engine/Source/Editor/DetailCustomizations/Private/ObjectDetails.cpp
@@ -110,7 +110,11 @@ void FObjectDetails::AddCallInEditorMethods(IDetailLayoutBuilder& DetailBuilder)
 				}
 			}
 
-			CallInEditorFunctions.Add(*FunctionIter);
+			const FName FunctionName = TestFunction->GetFName();
+			if (!CallInEditorFunctions.FindByPredicate([&FunctionName](const UFunction* Func) { return Func->GetFName() == FunctionName; }))
+			{
+				CallInEditorFunctions.Add(*FunctionIter);
+			}
 		}
 	}
 
