Fixed text ovelapping key selection in input settings. Also increased the width of the input text boxes slightly.

#5990

diff --git a/Engine/Source/Editor/DetailCustomizations/Private/InputSettingsDetails.h b/Engine/Source/Editor/DetailCustomizations/Private/InputSettingsDetails.h
--- a/Engine/Source/Editor/DetailCustomizations/Private/InputSettingsDetails.h
+++ b/Engine/Source/Editor/DetailCustomizations/Private/InputSettingsDetails.h
@@ -16,7 +16,7 @@ class IDetailLayoutBuilder;
 namespace InputConstants
 {
 	const FMargin PropertyPadding(2.0f, 0.0f, 2.0f, 0.0f);
-	const float TextBoxWidth = 200.0f;
+	const float TextBoxWidth = 250.0f;
 	const float ScaleBoxWidth = 50.0f;
 }
 
diff --git a/Engine/Source/Editor/UnrealEd/Private/SKeySelector.cpp b/Engine/Source/Editor/UnrealEd/Private/SKeySelector.cpp
--- a/Engine/Source/Editor/UnrealEd/Private/SKeySelector.cpp
+++ b/Engine/Source/Editor/UnrealEd/Private/SKeySelector.cpp
@@ -138,7 +138,6 @@ void SKeySelector::Construct(const FArguments& InArgs)
 				.Image(this, &SKeySelector::GetKeyIconImage)
 			]
 			+ SHorizontalBox::Slot()
-			.AutoWidth()
 			.VAlign(VAlign_Center)
 			.HAlign(HAlign_Left)
 			[
