Fixed issue where spinbox min/max values on vector structs wouldn't propegate to the nested individual member fields

diff --git a/Engine/Source/Editor/DetailCustomizations/Private/MathStructCustomizations.cpp b/Engine/Source/Editor/DetailCustomizations/Private/MathStructCustomizations.cpp
--- a/Engine/Source/Editor/DetailCustomizations/Private/MathStructCustomizations.cpp
+++ b/Engine/Source/Editor/DetailCustomizations/Private/MathStructCustomizations.cpp
@@ -69,6 +69,17 @@ void FMathStructCustomization::MakeHeaderRow(TSharedRef<class IPropertyHandle>&
 	{
 		TSharedRef<IPropertyHandle> ChildHandle = SortedChildHandles[ChildIndex];
 
+		// Propagate metadata to child properties so that it's reflected in the nested, individual spin boxes
+		ChildHandle->SetInstanceMetaData(TEXT("UIMin"), StructPropertyHandle->GetMetaData(TEXT("UIMin")));
+		ChildHandle->SetInstanceMetaData(TEXT("UIMax"), StructPropertyHandle->GetMetaData(TEXT("UIMax")));
+		ChildHandle->SetInstanceMetaData(TEXT("SliderExponent"), StructPropertyHandle->GetMetaData(TEXT("SliderExponent")));
+		ChildHandle->SetInstanceMetaData(TEXT("Delta"), StructPropertyHandle->GetMetaData(TEXT("Delta")));
+		ChildHandle->SetInstanceMetaData(TEXT("ShiftMouseMovePixelPerDelta"), StructPropertyHandle->GetMetaData(TEXT("ShiftMouseMovePixelPerDelta")));
+		ChildHandle->SetInstanceMetaData(TEXT("SupportDynamicSliderMaxValue"), StructPropertyHandle->GetMetaData(TEXT("SupportDynamicSliderMaxValue")));
+		ChildHandle->SetInstanceMetaData(TEXT("SupportDynamicSliderMinValue"), StructPropertyHandle->GetMetaData(TEXT("SupportDynamicSliderMinValue")));
+		ChildHandle->SetInstanceMetaData(TEXT("ClampMin"), StructPropertyHandle->GetMetaData(TEXT("ClampMin")));
+		ChildHandle->SetInstanceMetaData(TEXT("ClampMax"), StructPropertyHandle->GetMetaData(TEXT("ClampMax")));
+
 		const bool bLastChild = SortedChildHandles.Num()-1 == ChildIndex;
 		// Make a widget for each property.  The vector component properties  will be displayed in the header
 
