Scale spin boxes no longer cause the other components of the Scale to get stuck at 0 when scaling through zero (eg. 1 -> -1).

diff --git a/Engine/Source/Editor/DetailCustomizations/Private/ComponentTransformDetails.cpp b/Engine/Source/Editor/DetailCustomizations/Private/ComponentTransformDetails.cpp
--- a/Engine/Source/Editor/DetailCustomizations/Private/ComponentTransformDetails.cpp
+++ b/Engine/Source/Editor/DetailCustomizations/Private/ComponentTransformDetails.cpp
@@ -1180,20 +1180,53 @@ void FComponentTransformDetails::OnSetTransform(ETransformField::Type TransformF
 								switch (Axis)
 								{
 								case EAxisList::X:
-									// Account for the previous scale being zero.  Just set to the new value in that case?
-									Ratio = OldComponentValue.X == 0.0f ? NewComponentValue.X : NewComponentValue.X / OldComponentValue.X;
-									NewComponentValue.Y *= Ratio;
-									NewComponentValue.Z *= Ratio;
+									if (bIsSliderTransaction)
+									{
+										Ratio = SliderScaleRatio.X == 0.0f ? SliderScaleRatio.Y : (SliderScaleRatio.Y / SliderScaleRatio.X);
+										NewComponentValue.Y = NewComponentValue.X * Ratio;
+
+										Ratio = SliderScaleRatio.X == 0.0f ? SliderScaleRatio.Z : (SliderScaleRatio.Z / SliderScaleRatio.X);
+										NewComponentValue.Z = NewComponentValue.X * Ratio;
+									}
+									else
+									{
+										Ratio = OldComponentValue.X == 0.0f ? NewComponentValue.Z : NewComponentValue.X / OldComponentValue.X;
+										NewComponentValue.Y *= Ratio;
+										NewComponentValue.Z *= Ratio;
+									}
 									break;
 								case EAxisList::Y:
-									Ratio = OldComponentValue.Y == 0.0f ? NewComponentValue.Y : NewComponentValue.Y / OldComponentValue.Y;
-									NewComponentValue.X *= Ratio;
-									NewComponentValue.Z *= Ratio;
+									if (bIsSliderTransaction)
+									{
+										Ratio = SliderScaleRatio.Y == 0.0f ? SliderScaleRatio.X : (SliderScaleRatio.X / SliderScaleRatio.Y);
+										NewComponentValue.X = NewComponentValue.Y * Ratio;
+
+										Ratio = SliderScaleRatio.Y == 0.0f ? SliderScaleRatio.Z : (SliderScaleRatio.Z / SliderScaleRatio.Y);
+										NewComponentValue.Z = NewComponentValue.Y * Ratio;
+									}
+									else
+									{
+										Ratio = OldComponentValue.Y == 0.0f ? NewComponentValue.Z : NewComponentValue.Y / OldComponentValue.Y;
+										NewComponentValue.X *= Ratio;
+										NewComponentValue.Z *= Ratio;
+									}
 									break;
 								case EAxisList::Z:
-									Ratio = OldComponentValue.Z == 0.0f ? NewComponentValue.Z : NewComponentValue.Z / OldComponentValue.Z;
-									NewComponentValue.X *= Ratio;
-									NewComponentValue.Y *= Ratio;
+									if (bIsSliderTransaction)
+									{
+										Ratio = SliderScaleRatio.Z == 0.0f ? SliderScaleRatio.X : (SliderScaleRatio.X / SliderScaleRatio.Z);
+										NewComponentValue.X = NewComponentValue.Z * Ratio;
+
+										Ratio = SliderScaleRatio.Z == 0.0f ? SliderScaleRatio.Y : (SliderScaleRatio.Y / SliderScaleRatio.Z);
+										NewComponentValue.Y = NewComponentValue.Z * Ratio;
+									}
+									else
+									{
+										Ratio = OldComponentValue.Z == 0.0f ? NewComponentValue.Z : NewComponentValue.Z / OldComponentValue.Z;
+										NewComponentValue.X *= Ratio;
+										NewComponentValue.Y *= Ratio;
+									}
+									break;
 								default:
 									// Do nothing, this set multiple axis at once
 									break;
@@ -1414,6 +1447,11 @@ void FComponentTransformDetails::OnEndLocationSlider(float NewValue)
 
 void FComponentTransformDetails::OnBeginScaleSlider()
 {
+	// Assumption: slider isn't usable if multiple objects are selected
+	SliderScaleRatio.X = CachedScale.X.GetValue();
+	SliderScaleRatio.Y = CachedScale.Y.GetValue();
+	SliderScaleRatio.Z = CachedScale.Z.GetValue();
+
 	bIsSliderTransaction = true;
 	FText ActorTransaction = LOCTEXT("OnSetScale", "Set Scale");
 	FText ComponentTransaction = LOCTEXT("OnSetScale_ComponentDirect", "Modify Component Scale");
diff --git a/Engine/Source/Editor/DetailCustomizations/Private/ComponentTransformDetails.h b/Engine/Source/Editor/DetailCustomizations/Private/ComponentTransformDetails.h
--- a/Engine/Source/Editor/DetailCustomizations/Private/ComponentTransformDetails.h
+++ b/Engine/Source/Editor/DetailCustomizations/Private/ComponentTransformDetails.h
@@ -292,6 +292,8 @@ private:
 	bool bAbsoluteScale;
 	/** Whether or not to preserve scale ratios */
 	bool bPreserveScaleRatio;
+	/** Scale ratio to use when we are using the sliders with bPreserveScaleRatio set. */
+	FVector SliderScaleRatio;
 	/** Flag to indicate we are currently editing the rotation in the UI, so we should rely on the cached value in objectToRelativeRotationMap, not the value from the object */
 	bool bEditingRotationInUI;
 	/** Flag to indicate we are currently performing a slider transaction */
