Add support for DisplayPriority metadata tag, to allow defining sort order of UFunction CallInEditor buttons in DetailsView

diff --git a/Engine/Source/Editor/DetailCustomizations/Private/ObjectDetails.cpp b/Engine/Source/Editor/DetailCustomizations/Private/ObjectDetails.cpp
--- a/Engine/Source/Editor/DetailCustomizations/Private/ObjectDetails.cpp
+++ b/Engine/Source/Editor/DetailCustomizations/Private/ObjectDetails.cpp
@@ -90,6 +90,9 @@ void FObjectDetails::AddExperimentalWarningCategory(IDetailLayoutBuilder& Detail
 
 void FObjectDetails::AddCallInEditorMethods(IDetailLayoutBuilder& DetailBuilder)
 {
+	// metadata tag for defining sort order of function buttons within a Category
+	static const FName NAME_DisplayPriority("DisplayPriority");
+
 	// Get all of the functions we need to display (done ahead of time so we can sort them)
 	TArray<UFunction*, TInlineAllocator<8>> CallInEditorFunctions;
 	for (TFieldIterator<UFunction> FunctionIter(DetailBuilder.GetBaseClass(), EFieldIteratorFlags::IncludeSuper); FunctionIter; ++FunctionIter)
@@ -128,11 +131,32 @@ void FObjectDetails::AddCallInEditorMethods(IDetailLayoutBuilder& DetailBuilder)
 			return;
 		}
 
-		// Sort the functions by category and then by name
+		// Sort the functions by category and then by DisplayPriority meta tag, and then by name
 		CallInEditorFunctions.Sort([](UFunction& A, UFunction& B)
 		{
 			const int32 CategorySort = A.GetMetaData(FBlueprintMetadata::MD_FunctionCategory).Compare(B.GetMetaData(FBlueprintMetadata::MD_FunctionCategory));
-			return (CategorySort == 0) ? (A.GetName() <= B.GetName()) : (CategorySort <= 0);
+			if (CategorySort != 0)
+			{
+				return (CategorySort <= 0);
+			}
+			else 
+			{
+				FString DisplayPriorityAStr = A.GetMetaData(NAME_DisplayPriority);
+				int32 DisplayPriorityA = (DisplayPriorityAStr.IsEmpty() ? MAX_int32 : FCString::Atoi(*DisplayPriorityAStr));
+				if (DisplayPriorityA == 0 && !FCString::IsNumeric(*DisplayPriorityAStr))
+				{
+					DisplayPriorityA = MAX_int32;
+				}
+
+				FString DisplayPriorityBStr = B.GetMetaData(NAME_DisplayPriority);
+				int32 DisplayPriorityB = (DisplayPriorityBStr.IsEmpty() ? MAX_int32 : FCString::Atoi(*DisplayPriorityBStr));
+				if (DisplayPriorityB == 0 && !FCString::IsNumeric(*DisplayPriorityBStr))
+				{
+					DisplayPriorityB = MAX_int32;
+				}
+
+				return (DisplayPriorityA == DisplayPriorityB) ? (A.GetName() <= B.GetName()) : (DisplayPriorityA <= DisplayPriorityB);
+			}
 		});
 
 		struct FCategoryEntry
