Fixed issue where GUIDs wouldn't be generated properly when they belong to a component attached to a child BP class

diff --git a/Engine/Source/Editor/DetailCustomizations/Private/GuidStructCustomization.cpp b/Engine/Source/Editor/DetailCustomizations/Private/GuidStructCustomization.cpp
--- a/Engine/Source/Editor/DetailCustomizations/Private/GuidStructCustomization.cpp
+++ b/Engine/Source/Editor/DetailCustomizations/Private/GuidStructCustomization.cpp
@@ -173,13 +173,14 @@ void FGuidStructCustomization::HandleTextBoxTextCommited( const FText& NewText,
 void WriteGuidToProperty(TSharedPtr<IPropertyHandle> GuidPropertyHandle, const FGuid& Guid)
 {
 	FScopedTransaction Transaction(FText::Format(LOCTEXT("EditGuidPropertyTransaction", "Edit {0}"), GuidPropertyHandle->GetPropertyDisplayName()));
-
-	// Do not want a transaction on each individual set call as our scope transaction will handle it all
-	EPropertyValueSetFlags::Type GuidComponentFlags = (EPropertyValueSetFlags::NotTransactable);
+	
 	for (uint32 ChildIndex = 0; ChildIndex < 4; ++ChildIndex)
 	{
+		// Do not want a transaction on each individual set call as our scope transaction will handle it all
+		// Need to mark first 3 as interactive so that post edit doesn't reinstance anything until we're done
+		EPropertyValueSetFlags::Type GuidComponentFlags = EPropertyValueSetFlags::NotTransactable;
 		TSharedRef<IPropertyHandle> ChildHandle = GuidPropertyHandle->GetChildHandle(ChildIndex).ToSharedRef();
-		ChildHandle->SetValue((int32)Guid[ChildIndex], GuidComponentFlags);
+		ChildHandle->SetValue((int32)Guid[ChildIndex], ChildIndex != 3 ? GuidComponentFlags | EPropertyValueSetFlags::InteractiveChange : GuidComponentFlags);
 	}
 }
 
