Crash fix when picking collections

#rnx

diff --git a/Engine/Source/Editor/DetailCustomizations/Private/CollectionReferenceStructCustomization.cpp b/Engine/Source/Editor/DetailCustomizations/Private/CollectionReferenceStructCustomization.cpp
--- a/Engine/Source/Editor/DetailCustomizations/Private/CollectionReferenceStructCustomization.cpp
+++ b/Engine/Source/Editor/DetailCustomizations/Private/CollectionReferenceStructCustomization.cpp
@@ -72,6 +72,8 @@ FReply FCollectionReferenceStructCustomization::OnPickContent(TSharedRef<IProper
 {
 	FContentBrowserModule& ContentBrowserModule = FModuleManager::LoadModuleChecked<FContentBrowserModule>("ContentBrowser");
 	FCollectionPickerConfig PickerConfig;
+	PickerConfig.AllowCollectionButtons = false;
+	PickerConfig.AllowRightClickMenu = false;
 	PickerConfig.OnCollectionSelected = FOnCollectionSelected::CreateSP(this, &FCollectionReferenceStructCustomization::OnCollectionPicked, PropertyHandle);
 	
 	FMenuBuilder MenuBuilder(true, NULL);
