Fixed issue where the color picker wouldn't update color properties on components after the mouse was release for the first time

diff --git a/Engine/Source/Editor/DetailCustomizations/Private/ColorStructCustomization.cpp b/Engine/Source/Editor/DetailCustomizations/Private/ColorStructCustomization.cpp
--- a/Engine/Source/Editor/DetailCustomizations/Private/ColorStructCustomization.cpp
+++ b/Engine/Source/Editor/DetailCustomizations/Private/ColorStructCustomization.cpp
@@ -173,6 +173,8 @@ void FColorStructCustomization::GetSortedChildren(TSharedRef<IPropertyHandle> In
 
 void FColorStructCustomization::CreateColorPicker(bool bUseAlpha)
 {
+	GEditor->BeginTransaction(FText::Format(LOCTEXT("SetColorProperty", "Edit {0}"), StructPropertyHandle->GetPropertyDisplayName()));
+
 	int32 NumObjects = StructPropertyHandle->GetNumOuterObjects();
 
 	SavedPreColorPickerColors.Empty();
@@ -210,6 +212,7 @@ void FColorStructCustomization::CreateColorPicker(bool bUseAlpha)
 		PickerArgs.DisplayGamma = TAttribute<float>::Create(TAttribute<float>::FGetter::CreateUObject(GEngine, &UEngine::GetDisplayGamma));
 		PickerArgs.OnColorCommitted = FOnLinearColorValueChanged::CreateSP(this, &FColorStructCustomization::OnSetColorFromColorPicker);
 		PickerArgs.OnColorPickerCancelled = FOnColorPickerCancelled::CreateSP(this, &FColorStructCustomization::OnColorPickerCancelled);
+		PickerArgs.OnColorPickerWindowClosed = FOnWindowClosed::CreateSP(this, &FColorStructCustomization::OnColorPickerWindowClosed);
 		PickerArgs.OnInteractivePickBegin = FSimpleDelegate::CreateSP(this, &FColorStructCustomization::OnColorPickerInteractiveBegin);
 		PickerArgs.OnInteractivePickEnd = FSimpleDelegate::CreateSP(this, &FColorStructCustomization::OnColorPickerInteractiveEnd);
 		PickerArgs.InitialColorOverride = InitialColor;
@@ -228,6 +231,8 @@ void FColorStructCustomization::CreateColorPicker(bool bUseAlpha)
 
 TSharedRef<SColorPicker> FColorStructCustomization::CreateInlineColorPicker(TWeakPtr<IPropertyHandle> StructWeakHandlePtr)
 {
+	GEditor->BeginTransaction(FText::Format(LOCTEXT("SetColorProperty", "Edit {0}"), StructPropertyHandle->GetPropertyDisplayName()));
+
 	int32 NumObjects = StructPropertyHandle->GetNumOuterObjects();
 
 	SavedPreColorPickerColors.Empty();
@@ -263,6 +268,7 @@ TSharedRef<SColorPicker> FColorStructCustomization::CreateInlineColorPicker(TWea
 		.DisplayGamma(TAttribute<float>::Create(TAttribute<float>::FGetter::CreateUObject(GEngine, &UEngine::GetDisplayGamma)))
 		.OnColorCommitted(FOnLinearColorValueChanged::CreateSP(this, &FColorStructCustomization::OnSetColorFromColorPicker))
 		.OnColorPickerCancelled(FOnColorPickerCancelled::CreateSP(this, &FColorStructCustomization::OnColorPickerCancelled))
+		.OnColorPickerWindowClosed(FOnWindowClosed::CreateSP(this, &FColorStructCustomization::OnColorPickerWindowClosed))
 		.OnInteractivePickBegin(FSimpleDelegate::CreateSP(this, &FColorStructCustomization::OnColorPickerInteractiveBegin))
 		.OnInteractivePickEnd(FSimpleDelegate::CreateSP(this, &FColorStructCustomization::OnColorPickerInteractiveEnd))
 		.sRGBOverride(sRGBOverride)
@@ -311,30 +317,29 @@ void FColorStructCustomization::OnColorPickerCancelled(FLinearColor OriginalColo
 	{
 		StructPropertyHandle->SetPerObjectValues(PerObjectColors);
 	}
+
+	GEditor->CancelTransaction(0);
+}
+
+void FColorStructCustomization::OnColorPickerWindowClosed(const TSharedRef<SWindow>& Window)
+{
+	// pushes the last value from the interactive change without the interactive flag
+	FString ColorString;
+	StructPropertyHandle->GetValueAsFormattedString(ColorString);
+	StructPropertyHandle->SetValueFromFormattedString(ColorString);
+	GEditor->EndTransaction();
 }
 
 
 void FColorStructCustomization::OnColorPickerInteractiveBegin()
 {
 	bIsInteractive = true;
-
-	GEditor->BeginTransaction(FText::Format(LOCTEXT("SetColorProperty", "Edit {0}"), StructPropertyHandle->GetPropertyDisplayName()));
 }
 
 
 void FColorStructCustomization::OnColorPickerInteractiveEnd()
 {
 	bIsInteractive = false;
-
-	if (!bDontUpdateWhileEditing)
-	{
-		// pushes the last value from the interactive change without the interactive flag
-		FString ColorString;
-		StructPropertyHandle->GetValueAsFormattedString(ColorString);
-		StructPropertyHandle->SetValueFromFormattedString(ColorString);
-	}
-
-	GEditor->EndTransaction();
 }
 
 
diff --git a/Engine/Source/Editor/DetailCustomizations/Public/Customizations/ColorStructCustomization.h b/Engine/Source/Editor/DetailCustomizations/Public/Customizations/ColorStructCustomization.h
--- a/Engine/Source/Editor/DetailCustomizations/Public/Customizations/ColorStructCustomization.h
+++ b/Engine/Source/Editor/DetailCustomizations/Public/Customizations/ColorStructCustomization.h
@@ -76,6 +76,11 @@ protected:
 	 */
 	void OnColorPickerCancelled(FLinearColor OriginalColor);
 
+	/**
+	 * Called when the color picker window is clsoed
+	 */
+	void OnColorPickerWindowClosed(const TSharedRef<SWindow>& Window);
+
 	/**
 	 * Called when the user enters an interactive color change (dragging something in the picker)
 	 */
