QoL: Add Total Component Count in Landscape UI Details

#rnx

diff --git a/Engine/Source/Editor/DetailCustomizations/Private/LandscapeProxyUIDetails.cpp b/Engine/Source/Editor/DetailCustomizations/Private/LandscapeProxyUIDetails.cpp
--- a/Engine/Source/Editor/DetailCustomizations/Private/LandscapeProxyUIDetails.cpp
+++ b/Engine/Source/Editor/DetailCustomizations/Private/LandscapeProxyUIDetails.cpp
@@ -104,6 +104,18 @@ void FLandscapeProxyUIDetails::CustomizeDetails( IDetailLayoutBuilder& DetailBui
 					GenerateTextWidget(FText::Format(LOCTEXT("LandscapeCountValue", "{0}"), LandscapeCount), true)
 				];
 
+				int32 TotalComponentCount = LandscapeInfo->XYtoComponentMap.Num();
+				RowDisplayText = LOCTEXT("TotalLandscapeComponentCount", "Total Component Count");
+				CategoryBuilder.AddCustomRow(RowDisplayText)
+				.NameContent()
+				[
+					GenerateTextWidget(RowDisplayText)
+				]
+				.ValueContent()
+				[
+					GenerateTextWidget(FText::Format(LOCTEXT("TotalLandscapeComponentCountValue", "{0}"), TotalComponentCount), true)
+				];
+
 				LandscapeInfo->GetLandscapeExtent(Rect.Min.X, Rect.Min.Y, Rect.Max.X, Rect.Max.Y);
 				Size = Rect.Size();
 				RowDisplayText = LOCTEXT("LandscapeOverallResolution", "Overall Resolution (Verts)");
