Fixed crash when unloading a level with a landscape actor being displayed in the details view and both the details and level tools on-screen

diff --git a/Engine/Source/Editor/DetailCustomizations/Private/LandscapeUIDetails.cpp b/Engine/Source/Editor/DetailCustomizations/Private/LandscapeUIDetails.cpp
--- a/Engine/Source/Editor/DetailCustomizations/Private/LandscapeUIDetails.cpp
+++ b/Engine/Source/Editor/DetailCustomizations/Private/LandscapeUIDetails.cpp
@@ -37,8 +37,8 @@ void FLandscapeUIDetails::CustomizeDetails( IDetailLayoutBuilder& DetailBuilder
 
 	if (EditingObjects.Num() == 1)
 	{
-		ALandscape* Landscape = Cast<ALandscape>(EditingObjects[0].Get());
-		if (Landscape == nullptr)
+		TWeakObjectPtr<ALandscape> Landscape(Cast<ALandscape>(EditingObjects[0].Get()));
+		if (!Landscape.IsValid())
 		{
 			return;
 		}
@@ -65,14 +65,14 @@ void FLandscapeUIDetails::CustomizeDetails( IDetailLayoutBuilder& DetailBuilder
 			.Type(ESlateCheckBoxType::CheckBox)
 			.IsChecked_Lambda([=]()
 			{
-				return Landscape->CanHaveLayersContent() ? ECheckBoxState::Checked : ECheckBoxState::Unchecked;
+				return Landscape.IsValid() && Landscape->CanHaveLayersContent() ? ECheckBoxState::Checked : ECheckBoxState::Unchecked;
 			})
 			.OnCheckStateChanged_Lambda([=](ECheckBoxState NewState)
 			{
 				bool bChecked = (NewState == ECheckBoxState::Checked);
-				if (Landscape->CanHaveLayersContent() != bChecked)
+				if (Landscape.IsValid() && (Landscape->CanHaveLayersContent() != bChecked))
 				{
-					ToggleCanHaveLayersContent(Landscape);
+					ToggleCanHaveLayersContent(Landscape.Get());
 				}
 			})
 		];
@@ -127,4 +127,4 @@ void FLandscapeUIDetails::ToggleCanHaveLayersContent(ALandscape* Landscape)
 	}
 }
 
-#undef LOCTEXT_NAMESPACE
\ No newline at end of file
+#undef LOCTEXT_NAMESPACE
