Add opacity texture to FBX material instance import. When importing an FBX, you have the option to import materials as instances off of a base material, as opposed to creating new material graphs for each FBX material. This change adds the option to import anything connected to the TransparencyColor.

diff --git a/Engine/Config/BaseEditorPerProjectUserSettings.ini b/Engine/Config/BaseEditorPerProjectUserSettings.ini
--- a/Engine/Config/BaseEditorPerProjectUserSettings.ini
+++ b/Engine/Config/BaseEditorPerProjectUserSettings.ini
@@ -713,6 +713,7 @@ BaseNormalTextureName=""
 BaseEmissiveColorName=""
 BaseEmmisiveTextureName=""
 BaseSpecularTextureName=""
+BaseOpacityTextureName=""
 
 [SoundSettings]
 ChirpSoundClasses=Dialog DialogMedium DialogLoud DialogDeafening
diff --git a/Engine/Source/Editor/DetailCustomizations/Private/FbxImportUIDetails.cpp b/Engine/Source/Editor/DetailCustomizations/Private/FbxImportUIDetails.cpp
--- a/Engine/Source/Editor/DetailCustomizations/Private/FbxImportUIDetails.cpp
+++ b/Engine/Source/Editor/DetailCustomizations/Private/FbxImportUIDetails.cpp
@@ -950,7 +950,35 @@ void FFbxImportUIDetails::ConstructBaseMaterialUI(TSharedPtr<IPropertyHandle> Ha
 				]
 			]
 		];
+
+		// base opacity properties
+		InitialSelect = FindString(BaseTextureNames, ImportUI->TextureImportData->BaseOpacityTextureName);
+		InitialSelect = InitialSelect == INDEX_NONE ? 0 : InitialSelect; // default to the empty string located at index 0
+		MaterialCategory.AddCustomRow(LOCTEXT("BaseOpacityTextureProperty", "Base Opacity Texture Property")).NameContent()
+		[
+			SNew(STextBlock)
+			.Text(LOCTEXT("BaseOpacityTextureProperty", "Base Opacity Texture Property"))
+			.Font(IDetailLayoutBuilder::GetDetailFont())
+		]
+		.ValueContent()
+		.MaxDesiredWidth(Row.ValueWidget.MaxWidth)
+		[
+			SNew(SHorizontalBox)
+			+ SHorizontalBox::Slot()
+			.AutoWidth()
+			[
+				SNew(SBox)
+				.MinDesiredWidth(MinDesiredWidth)
+				[
+					SNew(STextComboBox)
+					.OptionsSource(&BaseTextureNames)
+					.OnSelectionChanged(this, &FFbxImportUIDetails::OnOpacityTextureColor)
+					.InitiallySelectedItem(BaseTextureNames[InitialSelect])
+				]
+			]
+		];
 	}
+
 	if (BaseTextureNames.Num() > 1 || BaseColorNames.Num() > 1)
 	{
 		MaterialCategory.AddCustomRow(LOCTEXT("BaseParamPropertyClearAll", "Clear All Properties"))
@@ -1174,6 +1202,11 @@ void FFbxImportUIDetails::OnSpecularTextureColor(TSharedPtr<FString> Selection,
 	GetSelectionParameterString(Selection, ImportUI->TextureImportData->BaseSpecularTextureName);
 }
 
+void FFbxImportUIDetails::OnOpacityTextureColor(TSharedPtr<FString> Selection, ESelectInfo::Type SelectInfo)
+{
+	GetSelectionParameterString(Selection, ImportUI->TextureImportData->BaseOpacityTextureName);
+}
+
 FReply FFbxImportUIDetails::MaterialBaseParamClearAllProperties()
 {
 	ImportUI->TextureImportData->BaseColorName.Empty();
@@ -1182,6 +1215,7 @@ FReply FFbxImportUIDetails::MaterialBaseParamClearAllProperties()
 	ImportUI->TextureImportData->BaseEmmisiveTextureName.Empty();
 	ImportUI->TextureImportData->BaseEmissiveColorName.Empty();
 	ImportUI->TextureImportData->BaseSpecularTextureName.Empty();
+	ImportUI->TextureImportData->BaseOpacityTextureName.Empty();
 	//Need to refresh the custom detail since we do not have any pointer on the combo box
 	RefreshCustomDetail();
 	return FReply::Handled();
diff --git a/Engine/Source/Editor/DetailCustomizations/Private/FbxImportUIDetails.h b/Engine/Source/Editor/DetailCustomizations/Private/FbxImportUIDetails.h
--- a/Engine/Source/Editor/DetailCustomizations/Private/FbxImportUIDetails.h
+++ b/Engine/Source/Editor/DetailCustomizations/Private/FbxImportUIDetails.h
@@ -78,6 +78,7 @@ public:
 	void OnEmmisiveTextureColor(TSharedPtr<FString> Selection, ESelectInfo::Type SelectInfo);
 	void OnEmissiveColor(TSharedPtr<FString> Selection, ESelectInfo::Type SelectInfo);
 	void OnSpecularTextureColor(TSharedPtr<FString> Selection, ESelectInfo::Type SelectInfo);
+	void OnOpacityTextureColor(TSharedPtr<FString> Selection, ESelectInfo::Type SelectInfo);
 	FReply MaterialBaseParamClearAllProperties();
 
 	void OnMaterialImportMethodChanged(TSharedPtr<FString> Selection, ESelectInfo::Type SelectInfo);
diff --git a/Engine/Source/Editor/UnrealEd/Classes/Factories/FbxTextureImportData.h b/Engine/Source/Editor/UnrealEd/Classes/Factories/FbxTextureImportData.h
--- a/Engine/Source/Editor/UnrealEd/Classes/Factories/FbxTextureImportData.h
+++ b/Engine/Source/Editor/UnrealEd/Classes/Factories/FbxTextureImportData.h
@@ -50,5 +50,8 @@ class UNREALED_API UFbxTextureImportData : public UFbxAssetImportData
 	UPROPERTY(config, meta = (ImportType = "Mesh"))
 	FString BaseSpecularTextureName;
 
+	UPROPERTY(config, meta = (ImportType = "Mesh"))
+	FString BaseOpacityTextureName;
+
 	bool CanEditChange( const FProperty* InProperty ) const override;
 };
diff --git a/Engine/Source/Editor/UnrealEd/Private/Fbx/FbxMainImport.cpp b/Engine/Source/Editor/UnrealEd/Private/Fbx/FbxMainImport.cpp
--- a/Engine/Source/Editor/UnrealEd/Private/Fbx/FbxMainImport.cpp
+++ b/Engine/Source/Editor/UnrealEd/Private/Fbx/FbxMainImport.cpp
@@ -394,6 +394,7 @@ void ApplyImportUIToImportOptions(UFbxImportUI* ImportUI, FBXImportOptions& InOu
 			InOutImportOptions.BaseNormalTextureName = ImportUI->TextureImportData->BaseNormalTextureName;
 			InOutImportOptions.BaseEmmisiveTextureName = ImportUI->TextureImportData->BaseEmmisiveTextureName;
 			InOutImportOptions.BaseSpecularTextureName = ImportUI->TextureImportData->BaseSpecularTextureName;
+			InOutImportOptions.BaseOpacityTextureName = ImportUI->TextureImportData->BaseOpacityTextureName;
 			InOutImportOptions.BaseEmissiveColorName = ImportUI->TextureImportData->BaseEmissiveColorName;
 		}
 		InOutImportOptions.bImportTextures			= ImportUI->bImportTextures;
diff --git a/Engine/Source/Editor/UnrealEd/Private/Fbx/FbxMaterialImport.cpp b/Engine/Source/Editor/UnrealEd/Private/Fbx/FbxMaterialImport.cpp
--- a/Engine/Source/Editor/UnrealEd/Private/Fbx/FbxMaterialImport.cpp
+++ b/Engine/Source/Editor/UnrealEd/Private/Fbx/FbxMaterialImport.cpp
@@ -628,6 +628,7 @@ void UnFbx::FFbxImporter::CreateUnrealMaterial(FbxSurfaceMaterial& FbxMaterial,
 		}
 		bCanInstance &= CanUseMaterialWithInstance(FbxMaterial, FbxSurfaceMaterial::sSpecular, FbxImportOptions->BaseSpecularTextureName, FbxImportOptions->BaseMaterial, UVSets);
 		bCanInstance &= CanUseMaterialWithInstance(FbxMaterial, FbxSurfaceMaterial::sNormalMap, FbxImportOptions->BaseNormalTextureName, FbxImportOptions->BaseMaterial, UVSets);
+		bCanInstance &= CanUseMaterialWithInstance(FbxMaterial, FbxSurfaceMaterial::sTransparentColor, FbxImportOptions->BaseOpacityTextureName, FbxImportOptions->BaseMaterial, UVSets);
 	}
 
 	UMaterialInterface* UnrealMaterialFinal = nullptr;
@@ -650,6 +651,7 @@ void UnFbx::FFbxImporter::CreateUnrealMaterial(FbxSurfaceMaterial& FbxMaterial,
 			// textures and properties
 			bool bDiffuseTextureCreated = LinkMaterialProperty(FbxMaterial, UnrealMaterialConstant, FbxSurfaceMaterial::sDiffuse, FName(*FbxImportOptions->BaseDiffuseTextureName), false);
 			bool bEmissiveTextureCreated = LinkMaterialProperty(FbxMaterial, UnrealMaterialConstant, FbxSurfaceMaterial::sEmissive, FName(*FbxImportOptions->BaseEmmisiveTextureName), false);
+			bool bOpacityTextureCreated = LinkMaterialProperty(FbxMaterial, UnrealMaterialConstant, FbxSurfaceMaterial::sTransparentColor, FName(*FbxImportOptions->BaseOpacityTextureName), false);
 			LinkMaterialProperty(FbxMaterial, UnrealMaterialConstant, FbxSurfaceMaterial::sSpecular, FName(*FbxImportOptions->BaseSpecularTextureName), false);
 			if (!LinkMaterialProperty(FbxMaterial, UnrealMaterialConstant, FbxSurfaceMaterial::sNormalMap, FName(*FbxImportOptions->BaseNormalTextureName), true))
 			{
@@ -717,6 +719,14 @@ void UnFbx::FFbxImporter::CreateUnrealMaterial(FbxSurfaceMaterial& FbxMaterial,
 					}
 				}
 			}
+			if (bOpacityTextureCreated)
+			{
+				FMaterialInstanceBasePropertyOverrides Overrides;
+				Overrides.bOverride_BlendMode = true;
+				Overrides.BlendMode = BLEND_Translucent;
+
+				UnrealMaterialConstant->BasePropertyOverrides = Overrides;
+			}
 		}
 	}
 	else
diff --git a/Engine/Source/Editor/UnrealEd/Public/FbxImporter.h b/Engine/Source/Editor/UnrealEd/Public/FbxImporter.h
--- a/Engine/Source/Editor/UnrealEd/Public/FbxImporter.h
+++ b/Engine/Source/Editor/UnrealEd/Public/FbxImporter.h
@@ -166,6 +166,7 @@ struct FBXImportOptions
 	FString BaseNormalTextureName;
 	FString BaseEmmisiveTextureName;
 	FString BaseSpecularTextureName;
+	FString BaseOpacityTextureName;
 	EMaterialSearchLocation MaterialSearchLocation;
 	//If true the materials will be reorder to follow the fbx order
 	bool bReorderMaterialToFbxOrder;
