CoreDelegates::FCrashOverrideParameters now includes values for SendUsageData and SendUnattendedBugReports. These values are set when the values change in UAnalyticsPrivacySettings and UCrashReportsPrivacySettings respectively.

Fixed typo in IImportantToggleSettingInterface's method GetToogleCategoryAndPropertyNames, from "Toogle" to "Toggle".

diff --git a/Engine/Source/Editor/DetailCustomizations/Private/ImportantToggleSettingCustomization.cpp b/Engine/Source/Editor/DetailCustomizations/Private/ImportantToggleSettingCustomization.cpp
--- a/Engine/Source/Editor/DetailCustomizations/Private/ImportantToggleSettingCustomization.cpp
+++ b/Engine/Source/Editor/DetailCustomizations/Private/ImportantToggleSettingCustomization.cpp
@@ -97,13 +97,14 @@ void FImportantToggleSettingCustomization::CustomizeDetails(IDetailLayoutBuilder
 	if (Objects.Num() == 1)
 	{
 		ToggleSettingObject = Objects[0];
-		IImportantToggleSettingInterface* ToogleSettingInterface = Cast<IImportantToggleSettingInterface>(ToggleSettingObject.Get());
+		IImportantToggleSettingInterface* ToggleSettingInterface = Cast<IImportantToggleSettingInterface>(ToggleSettingObject.Get());
 
-		if (ToogleSettingInterface != nullptr)
+		if (ToggleSettingInterface != nullptr)
 		{
 			FName CategoryName;
 			FName PropertyName;
-			ToogleSettingInterface->GetToogleCategoryAndPropertyNames(CategoryName, PropertyName);
+			ToggleSettingInterface->GetToggleCategoryAndPropertyNames(CategoryName, PropertyName);
+
 			IDetailCategoryBuilder& Category = DetailBuilder.EditCategory(CategoryName);
 			TogglePropertyHandle = DetailBuilder.GetProperty(PropertyName);
 
@@ -127,8 +128,8 @@ void FImportantToggleSettingCustomization::CustomizeDetails(IDetailLayoutBuilder
 					[
 						SNew(SImportantToggleButton)
 						.CheckBoxStyle(FEditorStyle::Get(), "Property.ToggleButton.Start")
-						.Text(ToogleSettingInterface->GetFalseStateLabel())
-						.ToolTipText(ToogleSettingInterface->GetFalseStateTooltip())
+						.Text(ToggleSettingInterface->GetFalseStateLabel())
+						.ToolTipText(ToggleSettingInterface->GetFalseStateTooltip())
 						.IsSet(this, &FImportantToggleSettingCustomization::IsToggleValue, false)
 						.OnToggled(this, &FImportantToggleSettingCustomization::OnToggledTo, false)
 					]
@@ -137,8 +138,8 @@ void FImportantToggleSettingCustomization::CustomizeDetails(IDetailLayoutBuilder
 					[
 						SNew(SImportantToggleButton)
 						.CheckBoxStyle(FEditorStyle::Get(), "Property.ToggleButton.End")
-						.Text(ToogleSettingInterface->GetTrueStateLabel())
-						.ToolTipText(ToogleSettingInterface->GetTrueStateTooltip())
+						.Text(ToggleSettingInterface->GetTrueStateLabel())
+						.ToolTipText(ToggleSettingInterface->GetTrueStateTooltip())
 						.IsSet(this, &FImportantToggleSettingCustomization::IsToggleValue, true)
 						.OnToggled(this, &FImportantToggleSettingCustomization::OnToggledTo, true)
 					]
@@ -151,8 +152,8 @@ void FImportantToggleSettingCustomization::CustomizeDetails(IDetailLayoutBuilder
 						.VAlign(VAlign_Center)
 						[
 							SNew(SHyperlink)
-							.Text(ToogleSettingInterface->GetAdditionalInfoUrlLabel())
-							.OnNavigate(this, &FImportantToggleSettingCustomization::OnNavigateHyperlink, ToogleSettingInterface->GetAdditionalInfoUrl())
+							.Text(ToggleSettingInterface->GetAdditionalInfoUrlLabel())
+							.OnNavigate(this, &FImportantToggleSettingCustomization::OnNavigateHyperlink, ToggleSettingInterface->GetAdditionalInfoUrl())
 						]
 					]
 				]
diff --git a/Engine/Source/Editor/UnrealEd/Classes/Analytics/AnalyticsPrivacySettings.h b/Engine/Source/Editor/UnrealEd/Classes/Analytics/AnalyticsPrivacySettings.h
--- a/Engine/Source/Editor/UnrealEd/Classes/Analytics/AnalyticsPrivacySettings.h
+++ b/Engine/Source/Editor/UnrealEd/Classes/Analytics/AnalyticsPrivacySettings.h
@@ -27,7 +27,7 @@ class UAnalyticsPrivacySettings : public UObject, public IImportantToggleSetting
 public:
 
 	// BEGIN IImportantToggleSettingInterface
-	virtual void GetToogleCategoryAndPropertyNames(FName& OutCategory, FName& OutProperty) const override;
+	virtual void GetToggleCategoryAndPropertyNames(FName& OutCategory, FName& OutProperty) const override;
 	virtual FText GetFalseStateLabel() const override;
 	virtual FText GetFalseStateTooltip() const override;
 	virtual FText GetFalseStateDescription() const override;
diff --git a/Engine/Source/Editor/UnrealEd/Classes/Analytics/CrashReportsPrivacySettings.h b/Engine/Source/Editor/UnrealEd/Classes/Analytics/CrashReportsPrivacySettings.h
--- a/Engine/Source/Editor/UnrealEd/Classes/Analytics/CrashReportsPrivacySettings.h
+++ b/Engine/Source/Editor/UnrealEd/Classes/Analytics/CrashReportsPrivacySettings.h
@@ -27,7 +27,7 @@ class UCrashReportsPrivacySettings : public UObject, public IImportantToggleSett
 public:
 
 	// BEGIN IImportantToggleSettingInterface
-	virtual void GetToogleCategoryAndPropertyNames(FName& OutCategory, FName& OutProperty) const override;
+	virtual void GetToggleCategoryAndPropertyNames(FName& OutCategory, FName& OutProperty) const override;
 	virtual FText GetFalseStateLabel() const override;
 	virtual FText GetFalseStateTooltip() const override;
 	virtual FText GetFalseStateDescription() const override;
@@ -37,4 +37,10 @@ public:
 	virtual FString GetAdditionalInfoUrl() const override;
 	virtual FText GetAdditionalInfoUrlLabel() const override;
 	// END IImportantToggleSettingInterface
+
+#if WITH_EDITOR
+	//~ Begin UObject Interface
+	virtual void PostEditChangeProperty(struct FPropertyChangedEvent& PropertyChangedEvent) override;
+	//~ End UObject Interface
+#endif
 };
diff --git a/Engine/Source/Editor/UnrealEd/Private/Analytics/AnalyticsPrivacySettings.cpp b/Engine/Source/Editor/UnrealEd/Private/Analytics/AnalyticsPrivacySettings.cpp
--- a/Engine/Source/Editor/UnrealEd/Private/Analytics/AnalyticsPrivacySettings.cpp
+++ b/Engine/Source/Editor/UnrealEd/Private/Analytics/AnalyticsPrivacySettings.cpp
@@ -4,6 +4,7 @@
 #include "UObject/UnrealType.h"
 #include "Interfaces/IAnalyticsProvider.h"
 #include "EngineAnalytics.h"
+#include "Misc/CoreDelegates.h"
 
 #define LOCTEXT_NAMESPACE "AnalyticsPrivacySettings"
 
@@ -13,7 +14,7 @@ UAnalyticsPrivacySettings::UAnalyticsPrivacySettings(const FObjectInitializer& O
 {
 }
 
-void UAnalyticsPrivacySettings::GetToogleCategoryAndPropertyNames(FName& OutCategory, FName& OutProperty) const
+void UAnalyticsPrivacySettings::GetToggleCategoryAndPropertyNames(FName& OutCategory, FName& OutProperty) const
 {
 	OutCategory = FName("Options");
 	OutProperty = FName("bSendUsageData");
@@ -64,10 +65,17 @@ void UAnalyticsPrivacySettings::PostEditChangeProperty(struct FPropertyChangedEv
 {
 	Super::PostEditChangeProperty(PropertyChangedEvent);
 
-	const FName PropertyName = (PropertyChangedEvent.Property != NULL) ? PropertyChangedEvent.Property->GetFName() : NAME_None;
+	const FName PropertyName = (PropertyChangedEvent.Property != nullptr) ? PropertyChangedEvent.Property->GetFName() : NAME_None;
 	if (PropertyName == GET_MEMBER_NAME_CHECKED(UAnalyticsPrivacySettings, bSendUsageData))
 	{
 		OnSendFullUsageDataChanged();
+
+		FCrashOverrideParameters Params;
+		Params.bSetCrashReportClientMessageText = false;
+		Params.bSetGameNameSuffix = false;
+		Params.SendUsageData = bSendUsageData;
+
+		FCoreDelegates::CrashOverrideParamsChanged.Broadcast(Params);
 	}
 }
 
diff --git a/Engine/Source/Editor/UnrealEd/Private/Analytics/CrashReportsPrivacySettings.cpp b/Engine/Source/Editor/UnrealEd/Private/Analytics/CrashReportsPrivacySettings.cpp
--- a/Engine/Source/Editor/UnrealEd/Private/Analytics/CrashReportsPrivacySettings.cpp
+++ b/Engine/Source/Editor/UnrealEd/Private/Analytics/CrashReportsPrivacySettings.cpp
@@ -5,6 +5,7 @@
 #include "Interfaces/IAnalyticsProvider.h"
 #include "EngineAnalytics.h"
 #include "Misc/ConfigCacheIni.h"
+#include "Misc/CoreDelegates.h"
 
 #define LOCTEXT_NAMESPACE "CrashReportsPrivacySettings"
 
@@ -18,7 +19,7 @@ UCrashReportsPrivacySettings::UCrashReportsPrivacySettings(const FObjectInitiali
 	}
 }
 
-void UCrashReportsPrivacySettings::GetToogleCategoryAndPropertyNames(FName& OutCategory, FName& OutProperty) const
+void UCrashReportsPrivacySettings::GetToggleCategoryAndPropertyNames(FName& OutCategory, FName& OutProperty) const
 {
 	OutCategory = FName("Options");
 	OutProperty = FName("bSendUnattendedBugReports");
@@ -64,4 +65,20 @@ FText UCrashReportsPrivacySettings::GetAdditionalInfoUrlLabel() const
 	return LOCTEXT("HyperlinkLabel", "Epic Games Privacy Policy");
 };
 
+#if WITH_EDITOR
+void UCrashReportsPrivacySettings::PostEditChangeProperty(struct FPropertyChangedEvent& PropertyChangedEvent)
+{
+	const FName PropertyName = (PropertyChangedEvent.Property != nullptr) ? PropertyChangedEvent.Property->GetFName() : NAME_None;
+	if (PropertyName == GET_MEMBER_NAME_CHECKED(UCrashReportsPrivacySettings, bSendUnattendedBugReports))
+	{
+		FCrashOverrideParameters Params;
+		Params.bSetCrashReportClientMessageText = false;
+		Params.bSetGameNameSuffix = false;
+		Params.SendUnattendedBugReports = bSendUnattendedBugReports;
+
+		FCoreDelegates::CrashOverrideParamsChanged.Broadcast(Params);
+	}
+}
+#endif
+
 #undef LOCTEXT_NAMESPACE
diff --git a/Engine/Source/Runtime/Core/Public/Misc/CoreDelegates.h b/Engine/Source/Runtime/Core/Public/Misc/CoreDelegates.h
--- a/Engine/Source/Runtime/Core/Public/Misc/CoreDelegates.h
+++ b/Engine/Source/Runtime/Core/Public/Misc/CoreDelegates.h
@@ -42,6 +42,8 @@ struct FCrashOverrideParameters
 	/** Default this to true for backward compatibility before these bools were added. */
 	bool bSetCrashReportClientMessageText = true;
 	bool bSetGameNameSuffix = false;
+	TOptional<bool> SendUnattendedBugReports;
+	TOptional<bool> SendUsageData;
 
 	CORE_API ~FCrashOverrideParameters();
 };
diff --git a/Engine/Source/Runtime/Engine/Classes/Engine/ImportantToggleSettingInterface.h b/Engine/Source/Runtime/Engine/Classes/Engine/ImportantToggleSettingInterface.h
--- a/Engine/Source/Runtime/Engine/Classes/Engine/ImportantToggleSettingInterface.h
+++ b/Engine/Source/Runtime/Engine/Classes/Engine/ImportantToggleSettingInterface.h
@@ -25,7 +25,7 @@ class IImportantToggleSettingInterface
 {
 	GENERATED_IINTERFACE_BODY()
 
-	virtual void GetToogleCategoryAndPropertyNames(FName& OutCategory, FName& OutProperty) const { OutCategory = NAME_None; OutProperty = NAME_None; }
+	virtual void GetToggleCategoryAndPropertyNames(FName& OutCategory, FName& OutProperty) const { OutCategory = NAME_None; OutProperty = NAME_None; }
 
 	virtual FText GetFalseStateLabel() const { return FText::GetEmpty(); }
 	virtual FText GetFalseStateTooltip() const { return FText::GetEmpty(); }
