Merging CL #11611449 from Dev-VirtualProduction

- Fix for FFrameRate customization not showing EditInlineToggle (now using the header row)

diff --git a/Engine/Source/Editor/DetailCustomizations/Private/FrameRateCustomization.cpp b/Engine/Source/Editor/DetailCustomizations/Private/FrameRateCustomization.cpp
--- a/Engine/Source/Editor/DetailCustomizations/Private/FrameRateCustomization.cpp
+++ b/Engine/Source/Editor/DetailCustomizations/Private/FrameRateCustomization.cpp
@@ -25,28 +25,27 @@ TSharedRef<IPropertyTypeCustomization> FFrameRateCustomization::MakeInstance()
 void FFrameRateCustomization::CustomizeHeader(TSharedRef<IPropertyHandle> InPropertyHandle, FDetailWidgetRow& HeaderRow, IPropertyTypeCustomizationUtils& CustomizationUtils)
 {
 	StructPropertyHandle = InPropertyHandle;
-}
-
 
-void FFrameRateCustomization::CustomizeChildren(TSharedRef<IPropertyHandle> InPropertyHandle, IDetailChildrenBuilder& ChildBuilder, IPropertyTypeCustomizationUtils& CustomizationUtils)
-{
 	TSharedPtr<IPropertyUtilities> PropertyUtils = CustomizationUtils.GetPropertyUtilities();
 
-	FDetailWidgetRow& CustomRow = ChildBuilder.AddCustomRow(StructPropertyHandle->GetPropertyDisplayName());
-
-	CustomRow.NameContent()
-	[
-		StructPropertyHandle->CreatePropertyNameWidget()
-	];
+	HeaderRow.NameContent()
+		[
+			StructPropertyHandle->CreatePropertyNameWidget()
+		]
 
-	CustomRow.ValueContent()
-	[
-		SNew(SFrameRatePicker)
-		.Font(CustomizationUtils.GetRegularFont())
+	.ValueContent()
+		[
+			SNew(SFrameRatePicker)
+			.Font(CustomizationUtils.GetRegularFont())
 		.HasMultipleValues(this, &FFrameRateCustomization::HasMultipleValues)
 		.Value(this, &FFrameRateCustomization::GetFirstFrameRate)
 		.OnValueChanged(this, &FFrameRateCustomization::SetFrameRate)
-	].IsEnabled(MakeAttributeLambda([=] { return !InPropertyHandle->IsEditConst() && PropertyUtils->IsPropertyEditingEnabled(); }));
+		].IsEnabled(MakeAttributeLambda([=] { return !InPropertyHandle->IsEditConst() && PropertyUtils->IsPropertyEditingEnabled(); }));
+}
+
+
+void FFrameRateCustomization::CustomizeChildren(TSharedRef<IPropertyHandle> InPropertyHandle, IDetailChildrenBuilder& ChildBuilder, IPropertyTypeCustomizationUtils& CustomizationUtils)
+{
 }
 
 
