FCollisionProfileNameCustomization reset to defaults fix;

* Fixed FCollisionProfileName detail customization to have a reset to defaults button.  Also fixed widget width and margins.

diff --git a/Engine/Source/Editor/DetailCustomizations/Private/CollisionProfileNameCustomization.cpp b/Engine/Source/Editor/DetailCustomizations/Private/CollisionProfileNameCustomization.cpp
--- a/Engine/Source/Editor/DetailCustomizations/Private/CollisionProfileNameCustomization.cpp
+++ b/Engine/Source/Editor/DetailCustomizations/Private/CollisionProfileNameCustomization.cpp
@@ -49,20 +49,34 @@ void FCollisionProfileNameCustomization::CustomizeChildren(TSharedRef<class IPro
 	]
 	.ValueContent()
 	[
-		SAssignNew(NameComboBox, SComboBox<TSharedPtr<FName>>)
-		.OptionsSource(&NameList)
-		.OnGenerateWidget(this, &FCollisionProfileNameCustomization::OnGenerateWidget)
-		.OnSelectionChanged(this, &FCollisionProfileNameCustomization::OnSelectionChanged, &CollisionGroup)
-		.OnComboBoxOpening(this, &FCollisionProfileNameCustomization::OnComboBoxOpening)
-		.InitiallySelectedItem(InitialSelectedName)
-		.IsEnabled(FSlateApplication::Get().GetNormalExecutionAttribute())
-		.ContentPadding(FMargin(2.0f, 2.0f))
-		.Content()
+		SNew(SHorizontalBox)
+		+SHorizontalBox::Slot()
+		.HAlign(HAlign_Center)
+		.VAlign(VAlign_Center)
+		.AutoWidth()
 		[
-			SNew(STextBlock)
-			.Text(this, &FCollisionProfileNameCustomization::GetProfileComboBoxContent)
-			.Font(IDetailLayoutBuilder::GetDetailFont())
-			.ToolTipText(this, &FCollisionProfileNameCustomization::GetProfileComboBoxToolTip)
+			SAssignNew(NameComboBox, SComboBox<TSharedPtr<FName>>)
+			.OptionsSource(&NameList)
+			.OnGenerateWidget(this, &FCollisionProfileNameCustomization::OnGenerateWidget)
+			.OnSelectionChanged(this, &FCollisionProfileNameCustomization::OnSelectionChanged, &CollisionGroup)
+			.OnComboBoxOpening(this, &FCollisionProfileNameCustomization::OnComboBoxOpening)
+			.InitiallySelectedItem(InitialSelectedName)
+			.IsEnabled(FSlateApplication::Get().GetNormalExecutionAttribute())
+			.ContentPadding(FMargin(2.0f, 2.0f))
+			.Content()
+			[
+				SNew(STextBlock)
+				.Text(this, &FCollisionProfileNameCustomization::GetProfileComboBoxContent)
+				.Margin(FMargin(2.0f, 0.0f, 2.0f, 0.0f))
+				.Font(IDetailLayoutBuilder::GetDetailFont())
+				.ToolTipText(this, &FCollisionProfileNameCustomization::GetProfileComboBoxToolTip)
+			]
+		]
+		+SHorizontalBox::Slot()
+		.AutoWidth()
+		.VAlign(VAlign_Center)
+		[
+			StructPropertyHandle->CreateDefaultPropertyButtonWidgets()
 		]
 	];
 }
