Timecode: Fix crash when data is empty

diff --git a/Engine/Source/Editor/DetailCustomizations/Private/TimecodeDetailsCustomization.cpp b/Engine/Source/Editor/DetailCustomizations/Private/TimecodeDetailsCustomization.cpp
--- a/Engine/Source/Editor/DetailCustomizations/Private/TimecodeDetailsCustomization.cpp
+++ b/Engine/Source/Editor/DetailCustomizations/Private/TimecodeDetailsCustomization.cpp
@@ -39,8 +39,13 @@ FText FTimecodeDetailsCustomization::OnGetTimecodeText() const
 	TArray<void*> RawData;
 	TimecodeProperty->AccessRawData(RawData);
 
-	FString CurrentValue = ((FTimecode*)RawData[0])->ToString();
-	return FText::FromString(CurrentValue);
+	if (RawData.Num())
+	{
+		FString CurrentValue = ((FTimecode*)RawData[0])->ToString();
+		return FText::FromString(CurrentValue);
+	}
+
+	return FText::GetEmpty();
 }
 
 void FTimecodeDetailsCustomization::OnTimecodeTextCommitted(const FText& InText, ETextCommit::Type CommitInfo)
@@ -48,19 +53,22 @@ void FTimecodeDetailsCustomization::OnTimecodeTextCommitted(const FText& InText,
 	TArray<void*> RawData;
 	TimecodeProperty->AccessRawData(RawData);
 
-	TArray<FString> Splits;
-	InText.ToString().ParseIntoArray(Splits, TEXT(":"));
-
-	if (Splits.Num() == 4)
+	if (RawData.Num())
 	{
-		((FTimecode*)RawData[0])->Hours = FCString::Atoi(*Splits[0]);
-		((FTimecode*)RawData[0])->Minutes = FCString::Atoi(*Splits[1]);
-		((FTimecode*)RawData[0])->Seconds = FCString::Atoi(*Splits[2]);
-		((FTimecode*)RawData[0])->Frames = FCString::Atoi(*Splits[3]);
-	}
-	else
-	{
-		UE_LOG(LogTemp, Error, TEXT("Unexpected timecode format. Expected 4 values, got %d"), Splits.Num());
+		TArray<FString> Splits;
+		InText.ToString().ParseIntoArray(Splits, TEXT(":"));
+
+		if (Splits.Num() == 4)
+		{
+			((FTimecode*)RawData[0])->Hours = FCString::Atoi(*Splits[0]);
+			((FTimecode*)RawData[0])->Minutes = FCString::Atoi(*Splits[1]);
+			((FTimecode*)RawData[0])->Seconds = FCString::Atoi(*Splits[2]);
+			((FTimecode*)RawData[0])->Frames = FCString::Atoi(*Splits[3]);
+		}
+		else
+		{
+			UE_LOG(LogTemp, Error, TEXT("Unexpected timecode format. Expected 4 values, got %d"), Splits.Num());
+		}
 	}
 }
 
