To avoid long UI refresh, don't display instances on InstancedStaticMeshComponent when there are more than 512 (controlled by CVar r.Editor.MaxNumInstancesDetails).
A performance warning is displayed with a button to let the user force show/hide the instances member if he wants to.

diff --git a/Engine/Source/Editor/DetailCustomizations/Private/DetailCustomizations.cpp b/Engine/Source/Editor/DetailCustomizations/Private/DetailCustomizations.cpp
--- a/Engine/Source/Editor/DetailCustomizations/Private/DetailCustomizations.cpp
+++ b/Engine/Source/Editor/DetailCustomizations/Private/DetailCustomizations.cpp
@@ -5,6 +5,7 @@
 #include "Modules/ModuleManager.h"
 #include "PropertyEditorModule.h"
 #include "StaticMeshComponentDetails.h"
+#include "InstancedStaticMeshComponentDetails.h"
 #include "LightComponentDetails.h"
 #include "LocalLightComponentDetails.h"
 #include "DirectionalLightComponentDetails.h"
@@ -287,6 +288,7 @@ void FDetailCustomizationsModule::RegisterObjectCustomizations()
 	RegisterCustomClassLayout("SceneComponent", FOnGetDetailCustomizationInstance::CreateStatic(&FSceneComponentDetails::MakeInstance));
 	RegisterCustomClassLayout("PrimitiveComponent", FOnGetDetailCustomizationInstance::CreateStatic(&FPrimitiveComponentDetails::MakeInstance));
 	RegisterCustomClassLayout("StaticMeshComponent", FOnGetDetailCustomizationInstance::CreateStatic(&FStaticMeshComponentDetails::MakeInstance));
+	RegisterCustomClassLayout("InstancedStaticMeshComponent", FOnGetDetailCustomizationInstance::CreateStatic(&FInstancedStaticMeshComponentDetails::MakeInstance));
 	RegisterCustomClassLayout("SkeletalMeshComponent", FOnGetDetailCustomizationInstance::CreateStatic(&FSkeletalMeshComponentDetails::MakeInstance));
 	RegisterCustomClassLayout("SkinnedMeshComponent", FOnGetDetailCustomizationInstance::CreateStatic(&FSkinnedMeshComponentDetails::MakeInstance));
 	RegisterCustomClassLayout("SplineComponent", FOnGetDetailCustomizationInstance::CreateStatic(&FSplineComponentDetails::MakeInstance));
diff --git a/Engine/Source/Editor/DetailCustomizations/Private/InstancedStaticMeshComponentDetails.cpp b/Engine/Source/Editor/DetailCustomizations/Private/InstancedStaticMeshComponentDetails.cpp
new file mode 100644
--- /dev/null
+++ b/Engine/Source/Editor/DetailCustomizations/Private/InstancedStaticMeshComponentDetails.cpp
@@ -0,0 +1,142 @@
+// Copyright Epic Games, Inc. All Rights Reserved.
+
+#include "InstancedStaticMeshComponentDetails.h"
+#include "Components/InstancedStaticMeshComponent.h"
+#include "PropertyHandle.h"
+#include "DetailLayoutBuilder.h"
+#include "IDetailPropertyRow.h"
+#include "DetailCategoryBuilder.h"
+#include "DetailWidgetRow.h"
+#include "Engine/StaticMesh.h"
+
+#define LOCTEXT_NAMESPACE "InstancedStaticMeshComponentDetails"
+
+TAutoConsoleVariable<float> CVarMaxNumInstancesDetails(
+	TEXT("r.Editor.MaxNumInstancesDetails"),
+	512,
+	TEXT("Maximum number of instances shown in the details panel. Above this value, instances are hidden by default. \n")
+	TEXT("< 0 : No maximum\n"),
+	ECVF_RenderThreadSafe);
+
+TSharedRef<IDetailCustomization> FInstancedStaticMeshComponentDetails::MakeInstance()
+{
+	return MakeShareable(new FInstancedStaticMeshComponentDetails);
+}
+
+void FInstancedStaticMeshComponentDetails::CustomizeDetails(IDetailLayoutBuilder& DetailBuilder)
+{
+	TArray<TWeakObjectPtr<UObject>> ObjectsBeingCustomized;
+	DetailBuilder.GetObjectsBeingCustomized(ObjectsBeingCustomized);
+
+	TSharedPtr<IPropertyHandle> PerInstanceSMDataProperty = DetailBuilder.GetProperty(GET_MEMBER_NAME_CHECKED(UInstancedStaticMeshComponent, PerInstanceSMData));
+
+	// Doesn't support multi-selection : 
+	if (ObjectsBeingCustomized.Num() != 1)
+	{
+		PerInstanceSMDataProperty->MarkHiddenByCustomization();
+	}
+	else
+	{
+		ComponentBeingCustomized = Cast<UInstancedStaticMeshComponent>(ObjectsBeingCustomized[0]);
+		if (ensure(ComponentBeingCustomized.IsValid()))
+		{
+			const FName InstancesCategoryName("Instances");
+			IDetailCategoryBuilder& InstancesCategory = DetailBuilder.EditCategory(InstancesCategoryName);
+
+			NumInstances = ComponentBeingCustomized->PerInstanceSMData.Num();
+			bForceShowAllInstances = ComponentBeingCustomized->bForceShowAllInstancesDetails;
+
+			int32 MaxNumInstances = CVarMaxNumInstancesDetails.GetValueOnGameThread();
+			bool bTooManyInstances = (MaxNumInstances >= 0) && (NumInstances > MaxNumInstances);
+			// If there are too many instances to display, hide the property as it will slow down the UI : 
+			if (bTooManyInstances && !bForceShowAllInstances)
+			{
+				PerInstanceSMDataProperty->MarkHiddenByCustomization();
+			}
+
+			// If there are too many instances to display, display an additional button to show/hide instances : 
+			if (bTooManyInstances)
+			{
+				const FText ShowHideAllInstancesText = LOCTEXT("ShowHideAllInstancesText", "Show/Hide All Instances");
+				const FText SlatePerformanceWarningText = LOCTEXT("SlatePerformanceWarningText", "Slate Performance Warning");
+				const FText TooltipText = LOCTEXT("ShowHideAllInstancesTooltip", "Displays per-instance data in the details. Can be very slow when there are lots of instances.");
+				FDetailWidgetRow& CustomRow = InstancesCategory.AddCustomRow(ShowHideAllInstancesText)
+					.NameContent()
+					[
+						SNew(SHorizontalBox)
+						+ SHorizontalBox::Slot()
+					.VAlign(VAlign_Center)
+					.HAlign(HAlign_Left)
+					.AutoWidth()
+					.Padding(0.0f, 0.0f, 4.0f, 0.0f)
+					[
+						SNew(SImage)
+						.Image(FCoreStyle::Get().GetBrush("Icons.Warning"))
+					.ToolTipText(TooltipText)
+					]
+				+ SHorizontalBox::Slot()
+					.VAlign(VAlign_Center)
+					.HAlign(HAlign_Left)
+					.AutoWidth()
+					[
+						SNew(STextBlock)
+						.Font(IDetailLayoutBuilder::GetDetailFontItalic())
+					.Text(SlatePerformanceWarningText)
+					.ToolTipText(TooltipText)
+					]
+					]
+				.ValueContent()
+					.MaxDesiredWidth(120.f)
+					[
+						SNew(SButton)
+						.OnClicked(this, &FInstancedStaticMeshComponentDetails::OnShowHideAllInstancesClicked)
+					.ToolTipText(TooltipText)
+					.IsEnabled(this, &FInstancedStaticMeshComponentDetails::IsShowHideAllInstancesEnabled)
+					[
+						SNew(STextBlock)
+						.Font(IDetailLayoutBuilder::GetDetailFontItalic())
+					.Text(this, &FInstancedStaticMeshComponentDetails::GetShowHideAllInstancesText)
+					]
+					];
+			}
+		}
+	}
+}
+
+void FInstancedStaticMeshComponentDetails::CustomizeDetails(const TSharedPtr<IDetailLayoutBuilder>& DetailBuilder)
+{
+	CachedDetailBuilder = DetailBuilder;
+	CustomizeDetails(*DetailBuilder);
+}
+
+FReply FInstancedStaticMeshComponentDetails::OnShowHideAllInstancesClicked()
+{
+	bForceShowAllInstances = !bForceShowAllInstances;
+
+	// pass the new value the component, as the following refresh will instanciate a new FInstancedStaticMeshComponentDetails : 
+	if (ComponentBeingCustomized.IsValid())
+	{
+		ComponentBeingCustomized->bForceShowAllInstancesDetails = bForceShowAllInstances;
+	}
+
+	// Here we can only take the ptr as ForceRefreshDetails() checks that the reference is unique.
+	if (IDetailLayoutBuilder* DetailBuilder = CachedDetailBuilder.Pin().Get())
+	{
+		DetailBuilder->ForceRefreshDetails();
+	}
+	return FReply::Handled();
+}
+
+bool FInstancedStaticMeshComponentDetails::IsShowHideAllInstancesEnabled() const
+{
+	return (NumInstances > 0);
+}
+
+FText FInstancedStaticMeshComponentDetails::GetShowHideAllInstancesText() const
+{
+	return bForceShowAllInstances ?
+		FText::Format(LOCTEXT("HideAllInstances", "Hide All {0} Instances"), NumInstances)
+		: FText::Format(LOCTEXT("ShowAllInstances", "Show All {0} Instances"), NumInstances);
+}
+
+#undef LOCTEXT_NAMESPACE
diff --git a/Engine/Source/Editor/DetailCustomizations/Private/InstancedStaticMeshComponentDetails.h b/Engine/Source/Editor/DetailCustomizations/Private/InstancedStaticMeshComponentDetails.h
new file mode 100644
--- /dev/null
+++ b/Engine/Source/Editor/DetailCustomizations/Private/InstancedStaticMeshComponentDetails.h
@@ -0,0 +1,37 @@
+// Copyright Epic Games, Inc. All Rights Reserved.
+
+#pragma once
+
+#include "CoreMinimal.h"
+#include "IDetailCustomization.h"
+
+class IDetailLayoutBuilder;
+class IPropertyHandle;
+class IDetailLayoutBuilder;
+class UInstancedStaticMeshComponent;
+
+class FInstancedStaticMeshComponentDetails : public IDetailCustomization
+{
+public:
+	/** Makes a new instance of this detail layout class for a specific detail view requesting it */
+	static TSharedRef<IDetailCustomization> MakeInstance();
+
+	/** IDetailCustomization interface */
+	virtual void CustomizeDetails(IDetailLayoutBuilder& DetailBuilder) override;
+	virtual void CustomizeDetails(const TSharedPtr<IDetailLayoutBuilder>& DetailBuilder) override;
+
+private:
+	FReply OnShowHideAllInstancesClicked();
+	bool IsShowHideAllInstancesEnabled() const;
+	FText GetShowHideAllInstancesText() const;
+
+private:
+	/** Indicate whether we're currently forcing the display of the instance list or not */
+	bool bForceShowAllInstances = false;
+	/** Component being edited */
+	TWeakObjectPtr<UInstancedStaticMeshComponent> ComponentBeingCustomized;
+	/** Current number of instances to show in the display panel */
+	int32 NumInstances = 0;
+	/** The detail builder for this customization */
+	TWeakPtr<IDetailLayoutBuilder> CachedDetailBuilder;
+};
diff --git a/Engine/Source/Runtime/Engine/Classes/Components/InstancedStaticMeshComponent.h b/Engine/Source/Runtime/Engine/Classes/Components/InstancedStaticMeshComponent.h
--- a/Engine/Source/Runtime/Engine/Classes/Components/InstancedStaticMeshComponent.h
+++ b/Engine/Source/Runtime/Engine/Classes/Components/InstancedStaticMeshComponent.h
@@ -289,6 +289,9 @@ public:
 #if WITH_EDITOR
 	/** One bit per instance if the instance is selected. */
 	TBitArray<> SelectedInstances;
+
+	/** Indicates that the user has purposedly chosen to show the instance list in the details panel, despite the performance warning. */
+	bool bForceShowAllInstancesDetails = false;
 #endif
 	/** Physics representation of the instance bodies. */
 	TArray<FBodyInstance*> InstanceBodies;
