Tentative fix for incremental build error

diff --git a/Engine/Source/Editor/DetailCustomizations/Private/InstancedStaticMeshComponentDetails.cpp b/Engine/Source/Editor/DetailCustomizations/Private/InstancedStaticMeshComponentDetails.cpp
--- a/Engine/Source/Editor/DetailCustomizations/Private/InstancedStaticMeshComponentDetails.cpp
+++ b/Engine/Source/Editor/DetailCustomizations/Private/InstancedStaticMeshComponentDetails.cpp
@@ -8,6 +8,9 @@
 #include "DetailCategoryBuilder.h"
 #include "DetailWidgetRow.h"
 #include "Engine/StaticMesh.h"
+#include "Widgets/DeclarativeSyntaxSupport.h"
+#include "Widgets/Images/SImage.h"
+#include "Widgets/Input/SButton.h"
 
 #define LOCTEXT_NAMESPACE "InstancedStaticMeshComponentDetails"
 
@@ -65,38 +68,38 @@ void FInstancedStaticMeshComponentDetails::CustomizeDetails(IDetailLayoutBuilder
 					[
 						SNew(SHorizontalBox)
 						+ SHorizontalBox::Slot()
-					.VAlign(VAlign_Center)
-					.HAlign(HAlign_Left)
-					.AutoWidth()
-					.Padding(0.0f, 0.0f, 4.0f, 0.0f)
-					[
-						SNew(SImage)
-						.Image(FCoreStyle::Get().GetBrush("Icons.Warning"))
-					.ToolTipText(TooltipText)
-					]
-				+ SHorizontalBox::Slot()
-					.VAlign(VAlign_Center)
-					.HAlign(HAlign_Left)
-					.AutoWidth()
-					[
-						SNew(STextBlock)
-						.Font(IDetailLayoutBuilder::GetDetailFontItalic())
-					.Text(SlatePerformanceWarningText)
-					.ToolTipText(TooltipText)
-					]
+						.VAlign(VAlign_Center)
+						.HAlign(HAlign_Left)
+						.AutoWidth()
+						.Padding(0.0f, 0.0f, 4.0f, 0.0f)
+						[
+							SNew(SImage)
+							.Image(FCoreStyle::Get().GetBrush("Icons.Warning"))
+						.ToolTipText(TooltipText)
+						]
+						+ SHorizontalBox::Slot()
+						.VAlign(VAlign_Center)
+						.HAlign(HAlign_Left)
+						.AutoWidth()
+						[
+							SNew(STextBlock)
+							.Font(IDetailLayoutBuilder::GetDetailFontItalic())
+						.Text(SlatePerformanceWarningText)
+						.ToolTipText(TooltipText)
+						]
 					]
 				.ValueContent()
 					.MaxDesiredWidth(120.f)
 					[
 						SNew(SButton)
 						.OnClicked(this, &FInstancedStaticMeshComponentDetails::OnShowHideAllInstancesClicked)
-					.ToolTipText(TooltipText)
-					.IsEnabled(this, &FInstancedStaticMeshComponentDetails::IsShowHideAllInstancesEnabled)
-					[
-						SNew(STextBlock)
-						.Font(IDetailLayoutBuilder::GetDetailFontItalic())
-					.Text(this, &FInstancedStaticMeshComponentDetails::GetShowHideAllInstancesText)
-					]
+						.ToolTipText(TooltipText)
+						.IsEnabled(this, &FInstancedStaticMeshComponentDetails::IsShowHideAllInstancesEnabled)
+						[
+							SNew(STextBlock)
+							.Font(IDetailLayoutBuilder::GetDetailFontItalic())
+							.Text(this, &FInstancedStaticMeshComponentDetails::GetShowHideAllInstancesText)
+						]
 					];
 			}
 		}
diff --git a/Engine/Source/Editor/DetailCustomizations/Private/InstancedStaticMeshComponentDetails.h b/Engine/Source/Editor/DetailCustomizations/Private/InstancedStaticMeshComponentDetails.h
--- a/Engine/Source/Editor/DetailCustomizations/Private/InstancedStaticMeshComponentDetails.h
+++ b/Engine/Source/Editor/DetailCustomizations/Private/InstancedStaticMeshComponentDetails.h
@@ -4,6 +4,8 @@
 
 #include "CoreMinimal.h"
 #include "IDetailCustomization.h"
+#include "Templates/SharedPointer.h"
+#include "UObject/WeakObjectPtr.h"
 
 class IDetailLayoutBuilder;
 class IPropertyHandle;
@@ -35,3 +37,4 @@ private:
 	/** The detail builder for this customization */
 	TWeakPtr<IDetailLayoutBuilder> CachedDetailBuilder;
 };
+
