Fix missing symbol
#rnx

diff --git a/Engine/Source/Editor/DetailCustomizations/Private/InstancedStaticMeshComponentDetails.h b/Engine/Source/Editor/DetailCustomizations/Private/InstancedStaticMeshComponentDetails.h
--- a/Engine/Source/Editor/DetailCustomizations/Private/InstancedStaticMeshComponentDetails.h
+++ b/Engine/Source/Editor/DetailCustomizations/Private/InstancedStaticMeshComponentDetails.h
@@ -7,6 +7,7 @@
 #include "Templates/SharedPointer.h"
 #include "UObject/WeakObjectPtr.h"
 
+class FReply; 
 class IDetailLayoutBuilder;
 class IPropertyHandle;
 class IDetailLayoutBuilder;
