* Added options to make BoundsCopyComponent usable outside of RuntimeVirtualTextureVolume
* Made it editor-only

jeremy.moore

diff --git a/Engine/Source/Editor/DetailCustomizations/Private/BoundsCopyComponentDetails.cpp b/Engine/Source/Editor/DetailCustomizations/Private/BoundsCopyComponentDetails.cpp
--- a/Engine/Source/Editor/DetailCustomizations/Private/BoundsCopyComponentDetails.cpp
+++ b/Engine/Source/Editor/DetailCustomizations/Private/BoundsCopyComponentDetails.cpp
@@ -32,7 +32,7 @@ void FBoundsCopyComponentDetailsCustomization::CustomizeDetails(IDetailLayoutBui
 		return;
 	}
 	BoundsCopyComponent = Cast<UBoundsCopyComponent>(ObjectsBeingCustomized[0].Get());
-	if (BoundsCopyComponent == nullptr)
+	if (!BoundsCopyComponent.IsValid())
 	{
 		return;
 	}
@@ -50,6 +50,7 @@ void FBoundsCopyComponentDetailsCustomization::CustomizeDetails(IDetailLayoutBui
 		SourceActorValue->CreatePropertyNameWidget()
 	]
 	.ValueContent()
+	.MinDesiredWidth(250.f)
 	[
 		SourceActorValue->CreatePropertyValueWidget()
 	];
@@ -64,17 +65,22 @@ void FBoundsCopyComponentDetailsCustomization::CustomizeDetails(IDetailLayoutBui
 		.Text(LOCTEXT("Button_CopyRotation", "Copy Rotation"))
 	]
 	.ValueContent()
-	.MaxDesiredWidth(125.f)
+	.MinDesiredWidth(125.f)
 	[
 		SNew(SButton)
 		.VAlign(VAlign_Center)
 		.HAlign(HAlign_Center)
 		.ContentPadding(2)
 		.Text(LOCTEXT("Button_Copy", "Copy"))
-		.ToolTipText(LOCTEXT("Button_CopyRotation_Tooltip", "Set the virtual texture rotation to match the source actor"))
+		.ToolTipText(LOCTEXT("Button_CopyRotation_Tooltip", "Set the actor rotation to match the source actor"))
 		.OnClicked(this, &FBoundsCopyComponentDetailsCustomization::SetRotation)
+		.IsEnabled(this, &FBoundsCopyComponentDetailsCustomization::IsCopyEnabled)
 	];
 
+	const FText CheckboxCopy_X_Tooltip = LOCTEXT("CheckboxCopy_X_Tooltip", "Limit the change to the X component of the bounds");
+	const FText CheckboxCopy_Y_Tooltip = LOCTEXT("CheckboxCopy_Y_Tooltip", "Limit the change to the Y component of the bounds");
+	const FText CheckboxCopy_Z_Tooltip = LOCTEXT("CheckboxCopy_Z_Tooltip", "Limit the change to the Z component of the bounds");
+
 	BoundsCategory
 	.AddCustomRow(LOCTEXT("Button_CopyBounds", "Copy Bounds"))
 	.NameContent()
@@ -84,18 +90,119 @@ void FBoundsCopyComponentDetailsCustomization::CustomizeDetails(IDetailLayoutBui
 		.Text(LOCTEXT("Button_CopyBounds", "Copy Bounds"))
 	]
 	.ValueContent()
-	.MaxDesiredWidth(125.f)
+	.MinDesiredWidth(250.f)
 	[
-		SNew(SButton)
-		.VAlign(VAlign_Center)
-		.HAlign(HAlign_Center)
-		.ContentPadding(2)
-		.Text(LOCTEXT("Button_Copy", "Copy"))
-		.ToolTipText(LOCTEXT("Button_CopyBounds_Tooltip", "Set the virtual texture transform so that it includes the full bounds of the source actor"))
-		.OnClicked(this, &FBoundsCopyComponentDetailsCustomization::SetTransformToBounds)
+		SNew(SHorizontalBox)
+		+ SHorizontalBox::Slot()
+		[
+			SNew(SHorizontalBox)
+			+ SHorizontalBox::Slot()
+			[
+				SNew(SButton)
+				.VAlign(VAlign_Center)
+				.HAlign(HAlign_Center)
+				.ContentPadding(2)
+				.Text(LOCTEXT("Button_Copy", "Copy"))
+				.ToolTipText(LOCTEXT("Button_CopyBounds_Tooltip", "Set the actor transform so that it includes the full bounds of the source actor"))
+				.OnClicked(this, &FBoundsCopyComponentDetailsCustomization::SetTransformToBounds)
+				.IsEnabled(this, &FBoundsCopyComponentDetailsCustomization::IsCopyEnabled)
+			]
+			+ SHorizontalBox::Slot()
+			.VAlign(VAlign_Center)
+			.HAlign(HAlign_Center)
+			[
+				SNew(SHorizontalBox)
+				+ SHorizontalBox::Slot()
+				[
+					SNew(SCheckBox)
+					.OnCheckStateChanged(this, &FBoundsCopyComponentDetailsCustomization::OnBoundsComponentsXChanged)
+					.IsChecked(this, &FBoundsCopyComponentDetailsCustomization::IsBoundsComponentsXChecked)
+					.IsEnabled(this, &FBoundsCopyComponentDetailsCustomization::IsCopyEnabled)
+					.ToolTipText(CheckboxCopy_X_Tooltip)
+					[
+						SNew(STextBlock)
+						.Text(LOCTEXT("CheckboxCopy_X", "X"))
+						.ToolTipText(CheckboxCopy_X_Tooltip)
+						.MinDesiredWidth(10.0f)
+					]
+				]
+				+ SHorizontalBox::Slot()
+				[
+					SNew(SCheckBox)
+					.OnCheckStateChanged(this, &FBoundsCopyComponentDetailsCustomization::OnBoundsComponentsYChanged)
+					.IsChecked(this, &FBoundsCopyComponentDetailsCustomization::IsBoundsComponentsYChecked)
+					.IsEnabled(this, &FBoundsCopyComponentDetailsCustomization::IsCopyEnabled)
+					.ToolTipText(CheckboxCopy_Y_Tooltip)
+					[
+						SNew(STextBlock)
+						.Text(LOCTEXT("CheckboxCopy_Y", "Y"))
+						.ToolTipText(CheckboxCopy_Y_Tooltip)
+						.MinDesiredWidth(10.0f)
+					]
+				]
+				+ SHorizontalBox::Slot()
+				[
+					SNew(SCheckBox)
+					.OnCheckStateChanged(this, &FBoundsCopyComponentDetailsCustomization::OnBoundsComponentsZChanged)
+					.IsChecked(this, &FBoundsCopyComponentDetailsCustomization::IsBoundsComponentsZChecked)
+					.IsEnabled(this, &FBoundsCopyComponentDetailsCustomization::IsCopyEnabled)
+					.ToolTipText(CheckboxCopy_Z_Tooltip)
+					[
+						SNew(STextBlock)
+						.Text(LOCTEXT("CheckboxCopy_Z", "Z"))
+						.ToolTipText(CheckboxCopy_Z_Tooltip)
+						.MinDesiredWidth(10.0f)
+					]
+				]
+			]
+		]
 	];
 }
 
+void FBoundsCopyComponentDetailsCustomization::OnBoundsComponentsXChanged(ECheckBoxState NewState)
+{
+	if (BoundsCopyComponent.IsValid())
+	{
+		BoundsCopyComponent->bCopyXBounds = !BoundsCopyComponent->bCopyXBounds;
+	}
+}
+
+void FBoundsCopyComponentDetailsCustomization::OnBoundsComponentsYChanged(ECheckBoxState NewState)
+{
+	if (BoundsCopyComponent.IsValid())
+	{
+		BoundsCopyComponent->bCopyYBounds = !BoundsCopyComponent->bCopyYBounds;
+	}
+}
+
+void FBoundsCopyComponentDetailsCustomization::OnBoundsComponentsZChanged(ECheckBoxState NewState)
+{
+	if (BoundsCopyComponent.IsValid())
+	{
+		BoundsCopyComponent->bCopyZBounds = !BoundsCopyComponent->bCopyZBounds;
+	}
+}
+
+ECheckBoxState FBoundsCopyComponentDetailsCustomization::IsBoundsComponentsXChecked() const
+{
+	return BoundsCopyComponent.IsValid() ? (BoundsCopyComponent->bCopyXBounds ? ECheckBoxState::Checked : ECheckBoxState::Unchecked) : ECheckBoxState::Undetermined;
+}
+
+ECheckBoxState FBoundsCopyComponentDetailsCustomization::IsBoundsComponentsYChecked() const
+{
+	return BoundsCopyComponent.IsValid() ? (BoundsCopyComponent->bCopyYBounds ? ECheckBoxState::Checked : ECheckBoxState::Unchecked) : ECheckBoxState::Undetermined;
+}
+
+ECheckBoxState FBoundsCopyComponentDetailsCustomization::IsBoundsComponentsZChecked() const
+{
+	return BoundsCopyComponent.IsValid() ? (BoundsCopyComponent->bCopyZBounds ? ECheckBoxState::Checked : ECheckBoxState::Unchecked) : ECheckBoxState::Undetermined;
+}
+
+bool FBoundsCopyComponentDetailsCustomization::IsCopyEnabled() const
+{
+	return BoundsCopyComponent.IsValid() && BoundsCopyComponent->BoundsSourceActor.IsValid();
+}
+
 FReply FBoundsCopyComponentDetailsCustomization::SetRotation()
 {
 	FScopedTransaction BakeTransaction(LOCTEXT("Transaction_CopyRotation", "Copy Rotation"));
diff --git a/Engine/Source/Editor/DetailCustomizations/Private/BoundsCopyComponentDetails.h b/Engine/Source/Editor/DetailCustomizations/Private/BoundsCopyComponentDetails.h
--- a/Engine/Source/Editor/DetailCustomizations/Private/BoundsCopyComponentDetails.h
+++ b/Engine/Source/Editor/DetailCustomizations/Private/BoundsCopyComponentDetails.h
@@ -6,6 +6,10 @@
 #include "IDetailCustomization.h"
 #include "Widgets/Text/STextBlock.h"
 #include "Input/Reply.h"
+#include "UObject/WeakObjectPtr.h"
+enum class ECheckBoxState : uint8;
+
+class UBoundsCopyComponent;
 
 /** UI customization for UBoundsCopyComponent */
 class FBoundsCopyComponentDetailsCustomization : public IDetailCustomization
@@ -26,5 +30,14 @@ protected:
 	//~ End IDetailCustomization Interface.
 
 private:
-	class UBoundsCopyComponent* BoundsCopyComponent;
+	void OnBoundsComponentsXChanged(ECheckBoxState NewState);
+	void OnBoundsComponentsYChanged(ECheckBoxState NewState);
+	void OnBoundsComponentsZChanged(ECheckBoxState NewState);
+	ECheckBoxState IsBoundsComponentsXChecked() const;
+	ECheckBoxState IsBoundsComponentsYChecked() const;
+	ECheckBoxState IsBoundsComponentsZChecked() const;
+	bool IsCopyEnabled() const;
+
+private:
+	TWeakObjectPtr<UBoundsCopyComponent> BoundsCopyComponent;
 };
diff --git a/Engine/Source/Runtime/Engine/Classes/Components/BoundsCopyComponent.h b/Engine/Source/Runtime/Engine/Classes/Components/BoundsCopyComponent.h
--- a/Engine/Source/Runtime/Engine/Classes/Components/BoundsCopyComponent.h
+++ b/Engine/Source/Runtime/Engine/Classes/Components/BoundsCopyComponent.h
@@ -12,19 +12,46 @@ class ENGINE_API UBoundsCopyComponent : public UActorComponent
 {
 	GENERATED_UCLASS_BODY()
 
-private:
+public:
 	/** Actor to copy the bounds from to set up the transform. */
 	UPROPERTY(EditAnywhere, Category = TransformFromBounds, meta = (DisplayName = "Source Actor", DisplayPriority = 3))
 	TSoftObjectPtr<AActor> BoundsSourceActor = nullptr;
 
+	/** If true, the source actor's bounds will include its colliding components bounds. */
+	UPROPERTY(EditAnywhere, Category = TransformFromBounds, AdvancedDisplay, meta = (EditCondition = "BoundsSourceActor != nullptr"))
+	bool bUseCollidingComponentsForSourceBounds = false;
+
+	/** If true, the actor's scale will be changed so that after adjustment, its own bounds match the source bounds.*/
+	UPROPERTY(EditAnywhere, Category = TransformFromBounds, AdvancedDisplay, meta = (EditCondition = "BoundsSourceActor != nullptr"))
+	bool bKeepOwnBoundsScale = false;
+
+	/** If true, the actor's own bounds will include its colliding components bounds. */
+	UPROPERTY(EditAnywhere, Category = TransformFromBounds, AdvancedDisplay, meta = (EditCondition = "bKeepOwnBoundsScale && BoundsSourceActor != nullptr"))
+	bool bUseCollidingComponentsForOwnBounds = false;
+
+	/** If true, the actor's scale will be changed so that after adjustment, its own bounds match the source bounds.*/
+	UPROPERTY(EditAnywhere, Category = TransformFromBounds, AdvancedDisplay, meta = (EditCondition = "BoundsSourceActor != nullptr"))
+	FVector BoundsScaleFactor = FVector::OneVector;
+
+	UPROPERTY()
+	bool bCopyXBounds = true;
+
+	UPROPERTY()
+	bool bCopyYBounds = true;
+
+	UPROPERTY()
+	bool bCopyZBounds = true;
+
+	virtual bool IsEditorOnly() const override { return true; }
+
 public:
 #if WITH_EDITOR
 	/** Copy the rotation from BoundsSourceActor to this component */
-	UFUNCTION(CallInEditor, Category = TransformFromBounds, meta = (DisplayName = "Copy Rotation", DisplayPriority = 4))
+	UFUNCTION(BlueprintCallable, Category = TransformFromBounds, meta = (DisplayName = "Copy Rotation"))
 	void SetRotation();
 
 	/** Set this component transform to include the BoundsSourceActor bounds */
-	UFUNCTION(CallInEditor, Category = TransformFromBounds, meta = (DisplayName = "Copy Bounds", DisplayPriority = 4))
+	UFUNCTION(BlueprintCallable, Category = TransformFromBounds, meta = (DisplayName = "Copy Bounds"))
 	void SetTransformToBounds();
 #endif
 };
diff --git a/Engine/Source/Runtime/Engine/Private/Components/BoundsCopyComponent.cpp b/Engine/Source/Runtime/Engine/Private/Components/BoundsCopyComponent.cpp
--- a/Engine/Source/Runtime/Engine/Private/Components/BoundsCopyComponent.cpp
+++ b/Engine/Source/Runtime/Engine/Private/Components/BoundsCopyComponent.cpp
@@ -31,10 +31,10 @@ void UBoundsCopyComponent::SetTransformToBounds()
 	{
 		// Calculate the bounds in our local rotation space translated to the BoundsSourceActor center
 		const FQuat TargetRotation = GetOwner()->ActorToWorld().GetRotation();
-		const FVector InitialPosition = BoundsSourceActor->GetComponentsBoundingBox().GetCenter();
+		const FBox InitialBoundingBox = BoundsSourceActor->GetComponentsBoundingBox(bUseCollidingComponentsForSourceBounds);
 
 		FTransform LocalTransform;
-		LocalTransform.SetComponents(TargetRotation, InitialPosition, FVector::OneVector);
+		LocalTransform.SetComponents(TargetRotation, InitialBoundingBox.GetCenter(), FVector::OneVector);
 		FTransform WorldToLocal = LocalTransform.Inverse();
 
 		FBox BoundBox(ForceInit);
@@ -59,10 +59,25 @@ void UBoundsCopyComponent::SetTransformToBounds()
 		BoundBox.GetCenterAndExtents(Origin, Extent);
 
 		Origin = LocalTransform.TransformPosition(Origin);
-		Extent *= FVector(2.f, 2.f, 1.f); // Account for ARuntimeVirtualTextureVolume:Box offset which centers it on origin
+		Extent *= BoundsScaleFactor; 
 
+		FVector OwnExtent = bKeepOwnBoundsScale ? GetOwner()->CalculateComponentsBoundingBoxInLocalSpace(bUseCollidingComponentsForOwnBounds).GetExtent() : FVector::OneVector;
+
+		FVector NewScale = GetOwner()->ActorToWorld().GetScale3D();
+		if (bCopyXBounds)
+		{
+			NewScale.X = Extent.X / FMath::Max(OwnExtent.X, SMALL_NUMBER);
+		}
+		if (bCopyYBounds)
+		{
+			NewScale.Y = Extent.Y / FMath::Max(OwnExtent.Y, SMALL_NUMBER);
+		}
+		if (bCopyZBounds)
+		{
+			NewScale.Z = Extent.Z / FMath::Max(OwnExtent.Z, SMALL_NUMBER);
+		}
 		FTransform Transform;
-		Transform.SetComponents(TargetRotation, Origin, Extent);
+		Transform.SetComponents(TargetRotation, Origin, NewScale);
 
 		// Apply final result and notify the parent actor
 		GetOwner()->Modify();
diff --git a/Engine/Source/Runtime/Engine/Private/VT/RuntimeVirtualTextureVolume.cpp b/Engine/Source/Runtime/Engine/Private/VT/RuntimeVirtualTextureVolume.cpp
--- a/Engine/Source/Runtime/Engine/Private/VT/RuntimeVirtualTextureVolume.cpp
+++ b/Engine/Source/Runtime/Engine/Private/VT/RuntimeVirtualTextureVolume.cpp
@@ -14,6 +14,7 @@ ARuntimeVirtualTextureVolume::ARuntimeVirtualTextureVolume(const FObjectInitiali
 #if WITH_EDITORONLY_DATA
 	// Add bounds copier
 	BoundsCopyComponent = CreateDefaultSubobject<UBoundsCopyComponent>(TEXT("BoundsCopyComponent"));
+	BoundsCopyComponent->BoundsScaleFactor = FVector(2.f, 2.f, 1.f); // Account for Box offset which centers it on origin
 
 	// Add box for visualization of bounds
 	Box = CreateDefaultSubobject<UBoxComponent>(TEXT("Bounds"));
