Build fix

diff --git a/Engine/Source/Editor/DetailCustomizations/Private/BoundsCopyComponentDetails.cpp b/Engine/Source/Editor/DetailCustomizations/Private/BoundsCopyComponentDetails.cpp
--- a/Engine/Source/Editor/DetailCustomizations/Private/BoundsCopyComponentDetails.cpp
+++ b/Engine/Source/Editor/DetailCustomizations/Private/BoundsCopyComponentDetails.cpp
@@ -8,8 +8,10 @@
 #include "DetailWidgetRow.h"
 #include "ScopedTransaction.h"
 #include "Widgets/Input/SButton.h"
+#include "Widgets/Input/SCheckBox.h"
 #include "Widgets/Layout/SBox.h"
 #include "Widgets/Layout/SWrapBox.h"
+#include "Widgets/Text/STextBlock.h"
 
 #define LOCTEXT_NAMESPACE "BoundsCopyComponentDetails"
 
