Fix crash in the color picker if color picker is canceled after switching properties

diff --git a/Engine/Source/Editor/DetailCustomizations/Private/ColorStructCustomization.cpp b/Engine/Source/Editor/DetailCustomizations/Private/ColorStructCustomization.cpp
--- a/Engine/Source/Editor/DetailCustomizations/Private/ColorStructCustomization.cpp
+++ b/Engine/Source/Editor/DetailCustomizations/Private/ColorStructCustomization.cpp
@@ -181,7 +181,10 @@ void FColorStructCustomization::GetSortedChildren(TSharedRef<IPropertyHandle> In
 
 void FColorStructCustomization::CreateColorPicker(bool bUseAlpha)
 {
-	ActiveTransaction = GEditor->BeginTransaction(FText::Format(LOCTEXT("SetColorProperty", "Edit {0}"), StructPropertyHandle->GetPropertyDisplayName()));
+	if (!ActiveTransaction)
+	{
+		ActiveTransaction = GEditor->BeginTransaction(FText::Format(LOCTEXT("SetColorProperty", "Edit {0}"), StructPropertyHandle->GetPropertyDisplayName()));
+	}
 
 	int32 NumObjects = StructPropertyHandle->GetNumOuterObjects();
 
