Removed default change type value for IPropertyHandle::NotifyPostChange, and fixed all usages of it.

diff --git a/Engine/Plugins/Editor/GameplayTagsEditor/Source/GameplayTagsEditor/Private/GameplayTagQueryCustomization.cpp b/Engine/Plugins/Editor/GameplayTagsEditor/Source/GameplayTagsEditor/Private/GameplayTagQueryCustomization.cpp
--- a/Engine/Plugins/Editor/GameplayTagsEditor/Source/GameplayTagsEditor/Private/GameplayTagQueryCustomization.cpp
+++ b/Engine/Plugins/Editor/GameplayTagsEditor/Source/GameplayTagsEditor/Private/GameplayTagQueryCustomization.cpp
@@ -116,7 +116,7 @@ FReply FGameplayTagQueryCustomization::OnClearAllButtonClicked()
 
 	if (StructPropertyHandle.IsValid())
 	{
-		StructPropertyHandle->NotifyPostChange();
+		StructPropertyHandle->NotifyPostChange(EPropertyChangeType::ValueSet);
 	}
 
 	return FReply::Handled();
@@ -263,7 +263,7 @@ void FGameplayTagQueryCustomization::CloseWidgetWindow(bool WasCancelled)
 	// Notify change.
 	if (!WasCancelled && StructPropertyHandle.IsValid())
 	{
-		StructPropertyHandle->NotifyPostChange();
+		StructPropertyHandle->NotifyPostChange(EPropertyChangeType::ValueSet);
 	}
 
 	if (GameplayTagQueryWidgetWindow.IsValid())
diff --git a/Engine/Plugins/FX/Niagara/Source/NiagaraEditor/Private/Customizations/NiagaraPlatformSetCustomization.cpp b/Engine/Plugins/FX/Niagara/Source/NiagaraEditor/Private/Customizations/NiagaraPlatformSetCustomization.cpp
--- a/Engine/Plugins/FX/Niagara/Source/NiagaraEditor/Private/Customizations/NiagaraPlatformSetCustomization.cpp
+++ b/Engine/Plugins/FX/Niagara/Source/NiagaraEditor/Private/Customizations/NiagaraPlatformSetCustomization.cpp
@@ -313,7 +313,7 @@ void FNiagaraPlatformSetCustomization::InvalidateSiblingConflicts() const
 		TSharedPtr<IPropertyHandle> SiblingPlatformSet = FindChildPlatformSet(Sibling);
 		if (SiblingPlatformSet.IsValid())
 		{
-			SiblingPlatformSet->NotifyPostChange();
+			SiblingPlatformSet->NotifyPostChange(EPropertyChangeType::ValueSet);
 		}
 	}
 }
diff --git a/Engine/Plugins/FX/Niagara/Source/NiagaraEditor/Private/Customizations/NiagaraTypeCustomizations.cpp b/Engine/Plugins/FX/Niagara/Source/NiagaraEditor/Private/Customizations/NiagaraTypeCustomizations.cpp
--- a/Engine/Plugins/FX/Niagara/Source/NiagaraEditor/Private/Customizations/NiagaraTypeCustomizations.cpp
+++ b/Engine/Plugins/FX/Niagara/Source/NiagaraEditor/Private/Customizations/NiagaraTypeCustomizations.cpp
@@ -248,7 +248,7 @@ void FNiagaraVariableAttributeBindingCustomization::ChangeSource(FName InVarName
 	PropertyHandle->NotifyPreChange();
 	TargetVariableBinding->BoundVariable.SetName(InVarName);
 	TargetVariableBinding->DataSetVariable = FNiagaraConstants::GetAttributeAsDataSetKey(TargetVariableBinding->BoundVariable);
-	PropertyHandle->NotifyPostChange();
+	PropertyHandle->NotifyPostChange(EPropertyChangeType::ValueSet);
 	PropertyHandle->NotifyFinishedChangingProperties();
 }
 
@@ -431,7 +431,7 @@ void FNiagaraUserParameterBindingCustomization::ChangeSource(FName InVarName)
 	TargetUserParameterBinding->Parameter.SetName(InVarName);
 	//TargetUserParameterBinding->Parameter.SetType(FNiagaraTypeDefinition::GetUObjectDef()); Do not override the type here!
 	//TargetVariableBinding->DataSetVariable = FNiagaraConstants::GetAttributeAsDataSetKey(TargetVariableBinding->BoundVariable);
-	PropertyHandle->NotifyPostChange();
+	PropertyHandle->NotifyPostChange(EPropertyChangeType::ValueSet);
 	PropertyHandle->NotifyFinishedChangingProperties();
 }
 
@@ -638,7 +638,7 @@ void FNiagaraDataInterfaceBindingCustomization::ChangeSource(FName InVarName)
 
 	PropertyHandle->NotifyPreChange();
 	TargetDataInterfaceBinding->BoundVariable.SetName(InVarName);
-	PropertyHandle->NotifyPostChange();
+	PropertyHandle->NotifyPostChange(EPropertyChangeType::ValueSet);
 	PropertyHandle->NotifyFinishedChangingProperties();
 }
 
@@ -849,7 +849,7 @@ void FNiagaraScriptVariableBindingCustomization::ChangeSource(FName InVarName)
 
 	PropertyHandle->NotifyPreChange();
 	TargetVariableBinding->Name = InVarName;
-	PropertyHandle->NotifyPostChange();
+	PropertyHandle->NotifyPostChange(EPropertyChangeType::ValueSet);
 	PropertyHandle->NotifyFinishedChangingProperties();
 }
 
diff --git a/Engine/Plugins/FX/Niagara/Source/NiagaraEditorWidgets/Private/DetailCustomizations/NiagaraDataInterfaceCurveDetails.cpp b/Engine/Plugins/FX/Niagara/Source/NiagaraEditorWidgets/Private/DetailCustomizations/NiagaraDataInterfaceCurveDetails.cpp
--- a/Engine/Plugins/FX/Niagara/Source/NiagaraEditorWidgets/Private/DetailCustomizations/NiagaraDataInterfaceCurveDetails.cpp
+++ b/Engine/Plugins/FX/Niagara/Source/NiagaraEditorWidgets/Private/DetailCustomizations/NiagaraDataInterfaceCurveDetails.cpp
@@ -197,7 +197,7 @@ private:
 		{
 			if (GetCurveFromPropertyHandle(CurveProperty) == ChangedCurve)
 			{
-				CurveProperty->NotifyPostChange();
+				CurveProperty->NotifyPostChange(EPropertyChangeType::ValueSet);
 				break;
 			}
 		}
@@ -347,7 +347,7 @@ void FNiagaraDataInterfaceCurveDetailsBase::ImportSelectedAsset(UObject* Selecte
 		CustomizedCurveInterface->UpdateLUT(); // we need this done before notify change because of the internal copy methods
 		for (auto CurveProperty : CurveProperties)
 		{
-			CurveProperty->NotifyPostChange();
+			CurveProperty->NotifyPostChange(EPropertyChangeType::ValueSet);
 		}
 	}
 }
diff --git a/Engine/Plugins/Runtime/AR/Microsoft/HoloLensAR/Source/HoloLensPlatformEditor/Private/HoloLensLocalizedResourcesCustomization.cpp b/Engine/Plugins/Runtime/AR/Microsoft/HoloLensAR/Source/HoloLensPlatformEditor/Private/HoloLensLocalizedResourcesCustomization.cpp
--- a/Engine/Plugins/Runtime/AR/Microsoft/HoloLensAR/Source/HoloLensPlatformEditor/Private/HoloLensLocalizedResourcesCustomization.cpp
+++ b/Engine/Plugins/Runtime/AR/Microsoft/HoloLensAR/Source/HoloLensPlatformEditor/Private/HoloLensLocalizedResourcesCustomization.cpp
@@ -87,7 +87,7 @@ void FHoloLensLocalizedResourcesNodeBuilder::AddLocalizedEntry(const FString& Cu
 		DlcPluginProperty->SetValue(OptionalPluginName);
 	}
 	NewlyAddedItem->GetChildHandle("CultureId")->SetValue(CultureId);
-	NewlyAddedItem->NotifyPostChange();
+	NewlyAddedItem->NotifyPostChange(EPropertyChangeType::ArrayAdd);
 	OnRebuildChildren.ExecuteIfBound();
 }
 
diff --git a/Engine/Plugins/Runtime/AR/Microsoft/HoloLensAR/Source/HoloLensPlatformEditor/Private/HoloLensTargetSettingsCustomization.cpp b/Engine/Plugins/Runtime/AR/Microsoft/HoloLensAR/Source/HoloLensPlatformEditor/Private/HoloLensTargetSettingsCustomization.cpp
--- a/Engine/Plugins/Runtime/AR/Microsoft/HoloLensAR/Source/HoloLensPlatformEditor/Private/HoloLensTargetSettingsCustomization.cpp
+++ b/Engine/Plugins/Runtime/AR/Microsoft/HoloLensAR/Source/HoloLensPlatformEditor/Private/HoloLensTargetSettingsCustomization.cpp
@@ -278,7 +278,7 @@ void FHoloLensTargetSettingsCustomization::CustomizeDetails(IDetailLayoutBuilder
 
 		SetDefaultCapabilitiesProperty->NotifyPreChange();
 		SetDefaultCapabilitiesProperty->SetValue(false);
-		SetDefaultCapabilitiesProperty->NotifyPostChange();
+		SetDefaultCapabilitiesProperty->NotifyPostChange(EPropertyChangeType::ValueSet);
 	}
 
 	AudioPluginManager.BuildAudioCategory(DetailBuilder, TEXT("HoloLens"));
@@ -502,15 +502,14 @@ void FHoloLensTargetSettingsCustomization::OnCapabilityStateChanged(ECheckBoxSta
 	{
 		// Remove existing capability from the list
 		RawCapabilityStringArray->RemoveAt(Index);
+		CapabilityList->NotifyPostChange(EPropertyChangeType::ArrayRemove);
 	}
 	else if (!Found && IsEnabled)
 	{
 		//Add new capability to the list
 		RawCapabilityStringArray->AddUnique(*CapabilityName);
+		CapabilityList->NotifyPostChange(EPropertyChangeType::ArrayAdd);
 	}
-
-	// Save settings to Ini
-	CapabilityList->NotifyPostChange();
 }
 
 FReply FHoloLensTargetSettingsCustomization::GenerateSigningCertificate()
@@ -763,7 +762,7 @@ void FHoloLensTargetSettingsCustomization::OnAutoDetectWin10SDKChanged(ECheckBox
 			WindowsSDKSelector->SetEnabled(true);
 		}
 	}
-	Win10SDKVersionPropertyHandle->NotifyPostChange();
+	Win10SDKVersionPropertyHandle->NotifyPostChange(EPropertyChangeType::ValueSet);
 }
 
 ECheckBoxState FHoloLensTargetSettingsCustomization::IsAutoDetectWin10SDKChecked() const
diff --git a/Engine/Plugins/Runtime/GameplayAbilities/Source/GameplayAbilitiesEditor/Private/GameplayEffectExecutionScopedModifierInfoDetails.cpp b/Engine/Plugins/Runtime/GameplayAbilities/Source/GameplayAbilitiesEditor/Private/GameplayEffectExecutionScopedModifierInfoDetails.cpp
--- a/Engine/Plugins/Runtime/GameplayAbilities/Source/GameplayAbilitiesEditor/Private/GameplayEffectExecutionScopedModifierInfoDetails.cpp
+++ b/Engine/Plugins/Runtime/GameplayAbilities/Source/GameplayAbilitiesEditor/Private/GameplayEffectExecutionScopedModifierInfoDetails.cpp
@@ -443,7 +443,7 @@ void FGameplayEffectExecutionScopedModifierInfoDetails::SetCurrentBackingData(TS
 				ScopedModifierInfo.CapturedAttribute = FGameplayEffectAttributeCaptureDefinition();
 				ScopedModifierInfo.TransientAggregatorIdentifier = InBackingData->TransientAggregatorIdentifier;
 			}
-			ScopedModifierStructPropertyHandle->NotifyPostChange();
+			ScopedModifierStructPropertyHandle->NotifyPostChange(EPropertyChangeType::ValueSet);
 		}
 	}
 }
diff --git a/Engine/Source/Developer/SharedSettingsWidgets/Private/ShaderFormatsPropertyDetails.cpp b/Engine/Source/Developer/SharedSettingsWidgets/Private/ShaderFormatsPropertyDetails.cpp
--- a/Engine/Source/Developer/SharedSettingsWidgets/Private/ShaderFormatsPropertyDetails.cpp
+++ b/Engine/Source/Developer/SharedSettingsWidgets/Private/ShaderFormatsPropertyDetails.cpp
@@ -175,7 +175,8 @@ void FShaderFormatsPropertyDetails::OnTargetedRHIChanged(ECheckBoxState InNewVal
 				Array.Remove(InRHIName.ToString());
 			}
 		}
-		ShaderFormatsPropertyHandle->NotifyPostChange();
+
+		ShaderFormatsPropertyHandle->NotifyPostChange(EPropertyChangeType::Unspecified);
 	}
 }
 
diff --git a/Engine/Source/Editor/AnimGraph/Private/PoseDriverDetails.cpp b/Engine/Source/Editor/AnimGraph/Private/PoseDriverDetails.cpp
--- a/Engine/Source/Editor/AnimGraph/Private/PoseDriverDetails.cpp
+++ b/Engine/Source/Editor/AnimGraph/Private/PoseDriverDetails.cpp
@@ -431,7 +431,7 @@ void SPDD_TargetRow::NotifyTargetChanged()
 	TSharedPtr<FPoseDriverDetails> PoseDriverDetails = PoseDriverDetailsPtr.Pin();
 	if (PoseDriverDetails.IsValid())
 	{
-		PoseDriverDetails->NodePropHandle->NotifyPostChange(); // Will push change to preview node instance
+		PoseDriverDetails->NodePropHandle->NotifyPostChange(EPropertyChangeType::ValueSet); // Will push change to preview node instance
 	}
 }
 
@@ -1119,7 +1119,7 @@ void FPoseDriverDetails::ClickedOnCopyFromPoseAsset()
 		PoseDriver->Node.RBFParams.Radius = 0.5f * MaxDist; // reasonable default radius
 
 		UpdateTargetInfosList();
-		NodePropHandle->NotifyPostChange();
+		NodePropHandle->NotifyPostChange(EPropertyChangeType::ValueSet);
 	}
 }
 
@@ -1138,7 +1138,7 @@ void FPoseDriverDetails::ClickedOnAutoScaleFactors()
 		float MaxDist;
 		PoseDriver->AutoSetTargetScales(MaxDist);
 		PoseDriver->Node.RBFParams.Radius = 0.5f * MaxDist; // reasonable default radius
-		NodePropHandle->NotifyPostChange();
+		NodePropHandle->NotifyPostChange(EPropertyChangeType::ValueSet);
 	}
 }
 
@@ -1155,7 +1155,7 @@ FReply FPoseDriverDetails::ClickedAddTarget()
 	{
 		PoseDriver->AddNewTarget();
 		UpdateTargetInfosList();
-		NodePropHandle->NotifyPostChange(); // will push changes to preview node instance
+		NodePropHandle->NotifyPostChange(EPropertyChangeType::ArrayAdd); // will push changes to preview node instance
 	}
 	return FReply::Handled();
 }
@@ -1167,7 +1167,7 @@ void FPoseDriverDetails::RemoveTarget(int32 TargetIndex)
 	{
 		PoseDriver->Node.PoseTargets.RemoveAt(TargetIndex);
 		UpdateTargetInfosList();
-		NodePropHandle->NotifyPostChange(); // will push changes to preview node instance
+		NodePropHandle->NotifyPostChange(EPropertyChangeType::ArrayRemove); // will push changes to preview node instance
 	}
 }
 
diff --git a/Engine/Source/Editor/DetailCustomizations/Private/AnimTrailNodeDetails.cpp b/Engine/Source/Editor/DetailCustomizations/Private/AnimTrailNodeDetails.cpp
--- a/Engine/Source/Editor/DetailCustomizations/Private/AnimTrailNodeDetails.cpp
+++ b/Engine/Source/Editor/DetailCustomizations/Private/AnimTrailNodeDetails.cpp
@@ -93,7 +93,7 @@ void FAnimTrailNodeDetails::FTrailRelaxCurveEditor::ModifyOwner()
 		GraphNodeOwner->Modify();
 		if (TrailRelaxCurveHandle.IsValid())
 		{
-			TrailRelaxCurveHandle->NotifyPostChange();
+			TrailRelaxCurveHandle->NotifyPostChange(EPropertyChangeType::ValueSet);
 		}
 	}
 }
diff --git a/Engine/Source/Editor/DetailCustomizations/Private/BodyInstanceCustomization.cpp b/Engine/Source/Editor/DetailCustomizations/Private/BodyInstanceCustomization.cpp
--- a/Engine/Source/Editor/DetailCustomizations/Private/BodyInstanceCustomization.cpp
+++ b/Engine/Source/Editor/DetailCustomizations/Private/BodyInstanceCustomization.cpp
@@ -123,10 +123,15 @@ void FBodyInstanceCustomization::AddCollisionCategory(TSharedRef<class IProperty
 	CollisionGroup.HeaderRow()
 	.NameContent()
 	[
-		SNew(STextBlock)
-		.Text(LOCTEXT("CollisionPresetsLabel", "Collision Presets"))
-		.ToolTip(ProfileTooltip)
-		.Font( IDetailLayoutBuilder::GetDetailFont() )
+		SNew(SHorizontalBox)
+		+ SHorizontalBox::Slot()
+		.HAlign(HAlign_Right)
+		[ 
+			SNew(STextBlock)
+			.Text(LOCTEXT("CollisionPresetsLabel", "Collision Presets"))
+			.ToolTip(ProfileTooltip)
+			.Font(IDetailLayoutBuilder::GetDetailFont())
+		]
 	]
 	.ValueContent()
 	.MinDesiredWidth(131.0f)
@@ -853,7 +858,7 @@ void FBodyInstanceCustomization::MarkAllBodiesDefaultCollision(bool bUseDefaultC
 		}
 	}
 
-	CollisionResponsesHandle->NotifyPostChange();
+	CollisionResponsesHandle->NotifyPostChange(EPropertyChangeType::ValueSet);
 }
 
 void FBodyInstanceCustomization::OnCollisionProfileChanged( TSharedPtr<FString> NewSelection, ESelectInfo::Type SelectInfo, IDetailGroup* CollisionGroup )
@@ -1149,7 +1154,7 @@ void FBodyInstanceCustomization::SetResponse(int32 ValidIndex, ECollisionRespons
 		}
 	}
 
-	CollisionResponsesHandle->NotifyPostChange();
+	CollisionResponsesHandle->NotifyPostChange(EPropertyChangeType::ValueSet);
 }
 
 ECheckBoxState FBodyInstanceCustomization::IsCollisionChannelChecked( int32 ValidIndex, ECollisionResponse InCollisionResponse) const
@@ -1263,7 +1268,7 @@ void FBodyInstanceCustomization::SetCollisionResponseContainer(const FCollisionR
 			}
 		}
 
-		CollisionResponsesHandle->NotifyPostChange();
+		CollisionResponsesHandle->NotifyPostChange(EPropertyChangeType::ValueSet);
 	}
 	
 }
@@ -1671,6 +1676,7 @@ void FBodyInstanceCustomizationHelper::AddBodyConstraint(IDetailCategoryBuilder&
 			.Text(LOCTEXT("LockPositionLabel", "Lock Position"))
 			.ToolTipText(LOCTEXT("LockPositionTooltip", "Locks movement along the specified axis"))
 			.Font(IDetailLayoutBuilder::GetDetailFont())
+			.Justification(ETextJustify::Right)
 		]
 		.ValueContent()
 		[
diff --git a/Engine/Source/Editor/DetailCustomizations/Private/ComponentTransformDetails.cpp b/Engine/Source/Editor/DetailCustomizations/Private/ComponentTransformDetails.cpp
--- a/Engine/Source/Editor/DetailCustomizations/Private/ComponentTransformDetails.cpp
+++ b/Engine/Source/Editor/DetailCustomizations/Private/ComponentTransformDetails.cpp
@@ -321,7 +321,6 @@ void FComponentTransformDetails::GenerateChildContent( IDetailChildrenBuilder& C
 		.CopyAction( CreateCopyAction( ETransformField::Location ) )
 		.PasteAction( CreatePasteAction( ETransformField::Location ) )
 		.NameContent()
-		.HAlign(HAlign_Left)
 		.VAlign(VAlign_Center)
 		[
 			BuildTransformFieldLabel( ETransformField::Location )
@@ -394,7 +393,6 @@ void FComponentTransformDetails::GenerateChildContent( IDetailChildrenBuilder& C
 		.CopyAction( CreateCopyAction(ETransformField::Rotation) )
 		.PasteAction( CreatePasteAction(ETransformField::Rotation) )
 		.NameContent()
-		.HAlign(HAlign_Left)
 		.VAlign(VAlign_Center)
 		[
 			BuildTransformFieldLabel(ETransformField::Rotation)
@@ -460,7 +458,6 @@ void FComponentTransformDetails::GenerateChildContent( IDetailChildrenBuilder& C
 		.CopyAction( CreateCopyAction(ETransformField::Scale) )
 		.PasteAction( CreatePasteAction(ETransformField::Scale) )
 		.NameContent()
-		.HAlign(HAlign_Left)
 		.VAlign(VAlign_Center)
 		[
 			BuildTransformFieldLabel(ETransformField::Scale)
diff --git a/Engine/Source/Editor/DetailCustomizations/Private/CurveColorCustomization.cpp b/Engine/Source/Editor/DetailCustomizations/Private/CurveColorCustomization.cpp
--- a/Engine/Source/Editor/DetailCustomizations/Private/CurveColorCustomization.cpp
+++ b/Engine/Source/Editor/DetailCustomizations/Private/CurveColorCustomization.cpp
@@ -254,7 +254,7 @@ void FCurveColorCustomization::MakeTransactional()
 
 void FCurveColorCustomization::OnCurveChanged(const TArray<FRichCurveEditInfo>& ChangedCurveEditInfos)
 {
-	StructPropertyHandle->NotifyPostChange();
+	StructPropertyHandle->NotifyPostChange(EPropertyChangeType::Unspecified);
 }
 
 FLinearColor FCurveColorCustomization::GetLinearColorValue(float InTime) const
@@ -308,7 +308,7 @@ void FCurveColorCustomization::OnExternalCurveChanged(TSharedRef<IPropertyHandle
 			CurveWidget->SetCurveOwner(this);
 		}
 
-		CurvePropertyHandle->NotifyPostChange();
+		CurvePropertyHandle->NotifyPostChange(EPropertyChangeType::ValueSet);
 	}
 }
 
diff --git a/Engine/Source/Editor/DetailCustomizations/Private/CurveStructCustomization.cpp b/Engine/Source/Editor/DetailCustomizations/Private/CurveStructCustomization.cpp
--- a/Engine/Source/Editor/DetailCustomizations/Private/CurveStructCustomization.cpp
+++ b/Engine/Source/Editor/DetailCustomizations/Private/CurveStructCustomization.cpp
@@ -248,7 +248,7 @@ void FCurveStructCustomization::MakeTransactional()
 
 void FCurveStructCustomization::OnCurveChanged(const TArray<FRichCurveEditInfo>& ChangedCurveEditInfos)
 {
-	StructPropertyHandle->NotifyPostChange();
+	StructPropertyHandle->NotifyPostChange(EPropertyChangeType::Unspecified);
 }
 
 bool FCurveStructCustomization::IsValidCurve( FRichCurveEditInfo CurveInfo )
@@ -280,7 +280,7 @@ void FCurveStructCustomization::OnExternalCurveChanged(TSharedRef<IPropertyHandl
 			CurveWidget->SetCurveOwner(this, CurvePropertyHandle->IsEditable());
 		}
 
-		CurvePropertyHandle->NotifyPostChange();
+		CurvePropertyHandle->NotifyPostChange(EPropertyChangeType::ValueSet);
 	}
 }
 
diff --git a/Engine/Source/Editor/DetailCustomizations/Private/DateTimeStructCustomization.cpp b/Engine/Source/Editor/DetailCustomizations/Private/DateTimeStructCustomization.cpp
--- a/Engine/Source/Editor/DetailCustomizations/Private/DateTimeStructCustomization.cpp
+++ b/Engine/Source/Editor/DetailCustomizations/Private/DateTimeStructCustomization.cpp
@@ -109,7 +109,7 @@ void FDateTimeStructCustomization::HandleTextBoxTextCommited(const FText& NewTex
 		{
 			*(FDateTime*)RawDataInstance = ParsedDateTime;
 		}
-		PropertyHandle->NotifyPostChange();
+		PropertyHandle->NotifyPostChange(EPropertyChangeType::ValueSet);
 		PropertyHandle->NotifyFinishedChangingProperties();
 	}
 }
diff --git a/Engine/Source/Editor/DetailCustomizations/Private/DeviceProfileDetails.cpp b/Engine/Source/Editor/DetailCustomizations/Private/DeviceProfileDetails.cpp
--- a/Engine/Source/Editor/DetailCustomizations/Private/DeviceProfileDetails.cpp
+++ b/Engine/Source/Editor/DetailCustomizations/Private/DeviceProfileDetails.cpp
@@ -786,7 +786,7 @@ void FDeviceProfileConsoleVariablesPropertyDetails::HandleCVarAdded(const FStrin
 
 				Array.Add(CompleteCVarString);
 			}
-			CVarsHandle->NotifyPostChange();
+			CVarsHandle->NotifyPostChange(EPropertyChangeType::ArrayAdd);
 		}
 
 		// Update the UI with the selection
diff --git a/Engine/Source/Editor/DetailCustomizations/Private/DialogueStructsCustomizations.cpp b/Engine/Source/Editor/DetailCustomizations/Private/DialogueStructsCustomizations.cpp
--- a/Engine/Source/Editor/DetailCustomizations/Private/DialogueStructsCustomizations.cpp
+++ b/Engine/Source/Editor/DetailCustomizations/Private/DialogueStructsCustomizations.cpp
@@ -236,7 +236,7 @@ void SSpeakerDropDown::OnSelectionChanged(TSharedPtr<UDialogueVoice*> Speaker, E
 	{
 		SpeakerPropertyHandle->NotifyPreChange();
 		*SpeakerToChange = *Speaker;
-		SpeakerPropertyHandle->NotifyPostChange();
+		SpeakerPropertyHandle->NotifyPostChange(EPropertyChangeType::ValueSet);
 	}
 }
 
@@ -464,7 +464,7 @@ void STargetSetDropDown::OnSelectionChanged(TSharedPtr<FTargetSet> TargetSet, ES
 	{
 		TargetsPropertyHandle->NotifyPreChange();
 		*TargetSetToChange = *TargetSet;
-		TargetsPropertyHandle->NotifyPostChange();
+		TargetsPropertyHandle->NotifyPostChange(EPropertyChangeType::ValueSet);
 	}
 }
 
@@ -877,7 +877,7 @@ void SValidatedDialogueContextHeaderWidget::Tick( const FGeometry& AllottedGeome
 					{
 						TargetsPropertyHandle->NotifyPreChange();
 						*TargetSet = DialogueWave->ContextMappings[0].Context.Targets;
-						TargetsPropertyHandle->NotifyPostChange();
+						TargetsPropertyHandle->NotifyPostChange(EPropertyChangeType::ValueSet);
 					}
 				}
 			}
@@ -886,39 +886,39 @@ void SValidatedDialogueContextHeaderWidget::Tick( const FGeometry& AllottedGeome
 
 	if( !IsDialogueWaveValid() )
 	{
-		if( ContextErrorHint.IsValid() )	{ ContextErrorHint->SetError(	LOCTEXT("InvalidDialogueWaveError", "Invalid dialogue wave.") ); }
-
-		if( SpeakerErrorHint.IsValid() )	{ SpeakerErrorHint->SetError(	FText::GetEmpty() ); }
-		if( SpeakerErrorText.IsValid() )		{ SpeakerErrorText->SetError(		LOCTEXT("SelectDialogueWaveError", "Select a valid dialogue wave.") ); }
-		if( TargetsErrorHint.IsValid() )	{ TargetsErrorHint->SetError(	FText::GetEmpty() ); }
-		if( TargetsErrorText.IsValid() )		{ TargetsErrorText->SetError(		LOCTEXT("SelectDialogueWaveError", "Select a valid dialogue wave.") ); }
-	}
-	else if( !IsSpeakerValid() )
-	{
-		if( ContextErrorHint.IsValid() )	{ ContextErrorHint->SetError(	FText::GetEmpty() ); }
-
-		if( SpeakerErrorHint.IsValid() )	{ SpeakerErrorHint->SetError(	LOCTEXT("InvalidSpeakerError", "Invalid speaker for dialogue wave.") ); }
-		if( SpeakerErrorText.IsValid() )		{ SpeakerErrorText->SetError(		FText::GetEmpty() ); }
-		if( TargetsErrorHint.IsValid() )	{ TargetsErrorHint->SetError(	FText::GetEmpty() ); }
-		if( TargetsErrorText.IsValid() )		{ TargetsErrorText->SetError(		LOCTEXT("SelectSpeakerError", "Select a valid speaker.") ); }
-	}
-	else if( !IsTargetSetValid() )
-	{
-		if( ContextErrorHint.IsValid() )	{ ContextErrorHint->SetError(	FText::GetEmpty() ); }
-
-		if( SpeakerErrorHint.IsValid() )	{ SpeakerErrorHint->SetError(	FText::GetEmpty() ); }
-		if( SpeakerErrorText.IsValid() )		{ SpeakerErrorText->SetError(		FText::GetEmpty() ); }
-		if( TargetsErrorHint.IsValid() )	{ TargetsErrorHint->SetError(	LOCTEXT("SelectTargetsError", "Select a valid target set.") ); }
-		if( TargetsErrorText.IsValid() )		{ TargetsErrorText->SetError(		FText::GetEmpty() ); }
-	}
-	else
-	{
-		if( ContextErrorHint.IsValid() )	{ ContextErrorHint->SetError(	FText::GetEmpty() ); }
-
-		if( SpeakerErrorHint.IsValid() )	{ SpeakerErrorHint->SetError(	FText::GetEmpty() ); }
-		if( SpeakerErrorText.IsValid() )		{ SpeakerErrorText->SetError(		FText::GetEmpty() ); }
-		if( TargetsErrorHint.IsValid() )	{ TargetsErrorHint->SetError(	FText::GetEmpty() ); }
-		if( TargetsErrorText.IsValid() )		{ TargetsErrorText->SetError(		FText::GetEmpty() ); }
+		if( ContextErrorHint.IsValid() )	{ ContextErrorHint->SetError( LOCTEXT("InvalidDialogueWaveError", "Invalid dialogue wave.") ); }
+																		  
+		if( SpeakerErrorHint.IsValid() )	{ SpeakerErrorHint->SetError( FText::GetEmpty() ); }
+		if( SpeakerErrorText.IsValid() )	{ SpeakerErrorText->SetError( LOCTEXT("SelectDialogueWaveError", "Select a valid dialogue wave.") ); }
+		if( TargetsErrorHint.IsValid() )	{ TargetsErrorHint->SetError( FText::GetEmpty() ); }
+		if( TargetsErrorText.IsValid() )	{ TargetsErrorText->SetError( LOCTEXT("SelectDialogueWaveError", "Select a valid dialogue wave.") ); }
+	}																	  
+	else if( !IsSpeakerValid() )										  
+	{																	  
+		if( ContextErrorHint.IsValid() )	{ ContextErrorHint->SetError( FText::GetEmpty() ); }
+																		  
+		if( SpeakerErrorHint.IsValid() )	{ SpeakerErrorHint->SetError( LOCTEXT("InvalidSpeakerError", "Invalid speaker for dialogue wave.") ); }
+		if( SpeakerErrorText.IsValid() )	{ SpeakerErrorText->SetError( FText::GetEmpty() ); }
+		if( TargetsErrorHint.IsValid() )	{ TargetsErrorHint->SetError( FText::GetEmpty() ); }
+		if( TargetsErrorText.IsValid() )	{ TargetsErrorText->SetError( LOCTEXT("SelectSpeakerError", "Select a valid speaker.") ); }
+	}																	  
+	else if( !IsTargetSetValid() )										  
+	{																	  
+		if( ContextErrorHint.IsValid() )	{ ContextErrorHint->SetError( FText::GetEmpty() ); }
+																		  
+		if( SpeakerErrorHint.IsValid() )	{ SpeakerErrorHint->SetError( FText::GetEmpty() ); }
+		if( SpeakerErrorText.IsValid() )	{ SpeakerErrorText->SetError( FText::GetEmpty() ); }
+		if( TargetsErrorHint.IsValid() )	{ TargetsErrorHint->SetError( LOCTEXT("SelectTargetsError", "Select a valid target set.") ); }
+		if( TargetsErrorText.IsValid() )	{ TargetsErrorText->SetError( FText::GetEmpty() ); }
+	}																	  
+	else																  
+	{																	  
+		if( ContextErrorHint.IsValid() )	{ ContextErrorHint->SetError( FText::GetEmpty() ); }
+																		  
+		if( SpeakerErrorHint.IsValid() )	{ SpeakerErrorHint->SetError( FText::GetEmpty() ); }
+		if( SpeakerErrorText.IsValid() )	{ SpeakerErrorText->SetError( FText::GetEmpty() ); }
+		if( TargetsErrorHint.IsValid() )	{ TargetsErrorHint->SetError( FText::GetEmpty() ); }
+		if( TargetsErrorText.IsValid() )	{ TargetsErrorText->SetError( FText::GetEmpty() ); }
 	}
 }
 
diff --git a/Engine/Source/Editor/DetailCustomizations/Private/InputSettingsDetails.cpp b/Engine/Source/Editor/DetailCustomizations/Private/InputSettingsDetails.cpp
--- a/Engine/Source/Editor/DetailCustomizations/Private/InputSettingsDetails.cpp
+++ b/Engine/Source/Editor/DetailCustomizations/Private/InputSettingsDetails.cpp
@@ -159,7 +159,7 @@ void FActionMappingsNodeBuilder::AddActionMappingButton_OnClick()
 		FInputActionKeyMapping NewMapping(NewActionMappingName);
 		InputSettings->AddActionMapping(NewMapping);
 
-		ActionMappingsPropertyHandle->NotifyPostChange();
+		ActionMappingsPropertyHandle->NotifyPostChange(EPropertyChangeType::ArrayAdd);
 	}
 }
 
@@ -214,7 +214,7 @@ void FActionMappingsNodeBuilder::AddActionMappingToGroupButton_OnClick(const FMa
 		FInputActionKeyMapping NewMapping(MappingSet.SharedName);
 		InputSettings->AddActionMapping(NewMapping);
 
-		ActionMappingsPropertyHandle->NotifyPostChange();
+		ActionMappingsPropertyHandle->NotifyPostChange(EPropertyChangeType::ArrayAdd);
 	}
 }
 
@@ -450,7 +450,7 @@ void FAxisMappingsNodeBuilder::AddAxisMappingButton_OnClick()
 		FInputAxisKeyMapping NewMapping(NewAxisMappingName);
 		InputSettings->AddAxisMapping(NewMapping);
 
-		AxisMappingsPropertyHandle->NotifyPostChange();
+		AxisMappingsPropertyHandle->NotifyPostChange(EPropertyChangeType::ArrayAdd);
 	}
 }
 
@@ -505,7 +505,7 @@ void FAxisMappingsNodeBuilder::AddAxisMappingToGroupButton_OnClick(const FMappin
 		FInputAxisKeyMapping NewMapping(MappingSet.SharedName);
 		InputSettings->AddAxisMapping(NewMapping);
 
-		AxisMappingsPropertyHandle->NotifyPostChange();
+		AxisMappingsPropertyHandle->NotifyPostChange(EPropertyChangeType::ArrayAdd);
 	}
 }
 
diff --git a/Engine/Source/Editor/DetailCustomizations/Private/PhysicsSettingsDetails.cpp b/Engine/Source/Editor/DetailCustomizations/Private/PhysicsSettingsDetails.cpp
--- a/Engine/Source/Editor/DetailCustomizations/Private/PhysicsSettingsDetails.cpp
+++ b/Engine/Source/Editor/DetailCustomizations/Private/PhysicsSettingsDetails.cpp
@@ -256,7 +256,7 @@ private:
 			PhysicsSettings->UpdateDefaultConfigFile();
 			PhysicsSettings->LoadSurfaceType();
 
-			PhysicalSurfacesProperty->NotifyPostChange();
+			PhysicalSurfacesProperty->NotifyPostChange(EPropertyChangeType::ValueSet);
 		}
 	}
 private:
diff --git a/Engine/Source/Editor/DetailCustomizations/Private/SceneCaptureDetails.cpp b/Engine/Source/Editor/DetailCustomizations/Private/SceneCaptureDetails.cpp
--- a/Engine/Source/Editor/DetailCustomizations/Private/SceneCaptureDetails.cpp
+++ b/Engine/Source/Editor/DetailCustomizations/Private/SceneCaptureDetails.cpp
@@ -321,7 +321,7 @@ void FSceneCaptureDetails::OnShowFlagCheckStateChanged(ECheckBoxState InNewRadio
 		}
 	}
 
-	ShowFlagSettingsProperty->NotifyPostChange();
+	ShowFlagSettingsProperty->NotifyPostChange(EPropertyChangeType::ValueSet);
 	ShowFlagSettingsProperty->NotifyFinishedChangingProperties();
 }
 
diff --git a/Engine/Source/Editor/DetailCustomizations/Private/TimespanStructCustomization.cpp b/Engine/Source/Editor/DetailCustomizations/Private/TimespanStructCustomization.cpp
--- a/Engine/Source/Editor/DetailCustomizations/Private/TimespanStructCustomization.cpp
+++ b/Engine/Source/Editor/DetailCustomizations/Private/TimespanStructCustomization.cpp
@@ -112,7 +112,7 @@ void FTimespanStructCustomization::HandleTextBoxTextCommited(const FText& NewTex
 			*(FTimespan*)RawDataInstance = ParsedTimespan;
 		}
 
-		PropertyHandle->NotifyPostChange();
+		PropertyHandle->NotifyPostChange(EPropertyChangeType::ValueSet);
 		PropertyHandle->NotifyFinishedChangingProperties();
 	}
 }
diff --git a/Engine/Source/Editor/LocalizationDashboard/Private/LocalizationTargetDetailCustomization.cpp b/Engine/Source/Editor/LocalizationDashboard/Private/LocalizationTargetDetailCustomization.cpp
--- a/Engine/Source/Editor/LocalizationDashboard/Private/LocalizationTargetDetailCustomization.cpp
+++ b/Engine/Source/Editor/LocalizationDashboard/Private/LocalizationTargetDetailCustomization.cpp
@@ -541,7 +541,7 @@ void FLocalizationTargetDetailCustomization::OnTargetNameCommitted(const FText&
 
 			if (TargetNamePropertyHandle.IsValid() && TargetNamePropertyHandle->IsValidHandle())
 			{
-				TargetNamePropertyHandle->NotifyPostChange();
+				TargetNamePropertyHandle->NotifyPostChange(EPropertyChangeType::ValueSet);
 			}
 		}
 	}
@@ -839,7 +839,7 @@ void FLocalizationTargetDetailCustomization::OnTargetDependencyCheckStateChanged
 
 	if (TargetDependenciesPropertyHandle.IsValid() && TargetDependenciesPropertyHandle->IsValidHandle())
 	{
-		TargetDependenciesPropertyHandle->NotifyPostChange();
+		TargetDependenciesPropertyHandle->NotifyPostChange(State == ECheckBoxState::Checked ? EPropertyChangeType::ArrayAdd : EPropertyChangeType::ArrayRemove);
 	}
 
 	RebuildTargetDependenciesBox();
@@ -1129,7 +1129,7 @@ void FLocalizationTargetDetailCustomization::UpdateTargetFromReports()
 		LocalizationTarget->UpdateStatusFromConflictReport();
 		for (const TSharedPtr<IPropertyHandle>& WordCountPropertyHandle : WordCountPropertyHandles)
 		{
-			WordCountPropertyHandle->NotifyPostChange();
+			WordCountPropertyHandle->NotifyPostChange(EPropertyChangeType::ValueSet);
 		}
 	}
 }
diff --git a/Engine/Source/Editor/LocalizationDashboard/Private/SLocalizationTargetEditorCultureRow.cpp b/Engine/Source/Editor/LocalizationDashboard/Private/SLocalizationTargetEditorCultureRow.cpp
--- a/Engine/Source/Editor/LocalizationDashboard/Private/SLocalizationTargetEditorCultureRow.cpp
+++ b/Engine/Source/Editor/LocalizationDashboard/Private/SLocalizationTargetEditorCultureRow.cpp
@@ -297,7 +297,7 @@ void SLocalizationTargetEditorCultureRow::OnNativeCultureCheckStateChanged( cons
 
 				TargetSettingsPropertyHandle->NotifyPreChange();
 				TargetSettings->NativeCultureIndex = CultureIndex;
-				TargetSettingsPropertyHandle->NotifyPostChange();
+				TargetSettingsPropertyHandle->NotifyPostChange(EPropertyChangeType::ValueSet);
 			}
 			break;
 		}
@@ -385,7 +385,7 @@ void SLocalizationTargetEditorCultureRow::UpdateTargetFromReports()
 		LocalizationTarget->UpdateStatusFromConflictReport();
 		for (const TSharedPtr<IPropertyHandle>& WordCountPropertyHandle : WordCountPropertyHandles)
 		{
-			WordCountPropertyHandle->NotifyPostChange();
+			WordCountPropertyHandle->NotifyPostChange(EPropertyChangeType::ValueSet);
 		}
 	}
 }
diff --git a/Engine/Source/Editor/MaterialEditor/Private/MaterialLayersFunctionsCustomization.cpp b/Engine/Source/Editor/MaterialEditor/Private/MaterialLayersFunctionsCustomization.cpp
--- a/Engine/Source/Editor/MaterialEditor/Private/MaterialLayersFunctionsCustomization.cpp
+++ b/Engine/Source/Editor/MaterialEditor/Private/MaterialLayersFunctionsCustomization.cpp
@@ -81,7 +81,7 @@ void FMaterialLayersFunctionsCustomization::ResetToDefault()
 	const FScopedTransaction Transaction(LOCTEXT("ResetMaterialLayersFunctions", "Reset all Layers and Blends"));
 	SavedStructPropertyHandle->NotifyPreChange();
 	*MaterialLayersFunctions = FMaterialLayersFunctions();
-	SavedStructPropertyHandle->NotifyPostChange();
+	SavedStructPropertyHandle->NotifyPostChange(EPropertyChangeType::ValueSet);
 
 	// Refresh the header so the reset to default button is no longer visible
 	SavedLayoutBuilder->ForceRefreshDetails();
@@ -251,7 +251,7 @@ void FMaterialLayersFunctionsCustomization::OnNameChanged(const FText& InText, E
 	const FScopedTransaction Transaction(LOCTEXT("RenamedSection", "Renamed layer and blend section"));
 	SavedStructPropertyHandle->NotifyPreChange();
 	MaterialLayersFunctions->LayerNames[Counter] = InText;
-	SavedStructPropertyHandle->NotifyPostChange();
+	SavedStructPropertyHandle->NotifyPostChange(EPropertyChangeType::ValueSet);
 };
 #endif
 
@@ -274,7 +274,7 @@ void FMaterialLayersFunctionsCustomization::AddLayer()
 	const FScopedTransaction Transaction(LOCTEXT("AddLayerAndBlend", "Add a new Layer and a Blend into it"));
 	SavedStructPropertyHandle->NotifyPreChange();
 	MaterialLayersFunctions->AppendBlendedLayer();
-	SavedStructPropertyHandle->NotifyPostChange();
+	SavedStructPropertyHandle->NotifyPostChange(EPropertyChangeType::ArrayAdd);
 }
 
 void FMaterialLayersFunctionsCustomization::RemoveLayer(int32 Index)
@@ -282,7 +282,7 @@ void FMaterialLayersFunctionsCustomization::RemoveLayer(int32 Index)
 	const FScopedTransaction Transaction(LOCTEXT("RemoveLayerAndBlend", "Remove a Layer and the attached Blend"));
 	SavedStructPropertyHandle->NotifyPreChange();
 	MaterialLayersFunctions->RemoveBlendedLayerAt(Index);
-	SavedStructPropertyHandle->NotifyPostChange();
+	SavedStructPropertyHandle->NotifyPostChange(EPropertyChangeType::ArrayRemove);
 }
 
 
@@ -329,7 +329,7 @@ void FMaterialLayerFunctionElement::ResetLayerAssetToDefault(TSharedPtr<IPropert
 	}
 	}
 
-	InPropertyHandle->NotifyPostChange();
+	InPropertyHandle->NotifyPostChange(EPropertyChangeType::ValueSet);
 	InCustomization->RebuildChildren();
 }
 
diff --git a/Engine/Source/Editor/MaterialEditor/Private/MaterialPropertyHelpers.cpp b/Engine/Source/Editor/MaterialEditor/Private/MaterialPropertyHelpers.cpp
--- a/Engine/Source/Editor/MaterialEditor/Private/MaterialPropertyHelpers.cpp
+++ b/Engine/Source/Editor/MaterialEditor/Private/MaterialPropertyHelpers.cpp
@@ -149,7 +149,7 @@ void FMaterialPropertyHelpers::OnMaterialLayerAssetChanged(const struct FAssetDa
 		break;
 		}
 	}
-	InHandle->NotifyPostChange();
+	InHandle->NotifyPostChange(EPropertyChangeType::ValueSet);
 }
 
 bool FMaterialPropertyHelpers::FilterLayerAssets(const struct FAssetData& InAssetData, FMaterialLayersFunctions* LayerFunction, EMaterialParameterAssociation MaterialType, int32 Index)
@@ -1187,7 +1187,7 @@ void FMaterialPropertyHelpers::SetVectorChannelMaskValue(const FString& StringVa
 			MaterialInstanceEditor->CopyToSourceInstance();
 		}
 
-		PropertyHandle->NotifyPostChange();
+		PropertyHandle->NotifyPostChange(EPropertyChangeType::ValueSet);
 	}
 }
 
diff --git a/Engine/Source/Editor/MaterialEditor/Private/SMaterialLayersFunctionsTree.cpp b/Engine/Source/Editor/MaterialEditor/Private/SMaterialLayersFunctionsTree.cpp
--- a/Engine/Source/Editor/MaterialEditor/Private/SMaterialLayersFunctionsTree.cpp
+++ b/Engine/Source/Editor/MaterialEditor/Private/SMaterialLayersFunctionsTree.cpp
@@ -137,7 +137,7 @@ void SMaterialLayersFunctionsInstanceTreeItem::OnNameChanged(const FText& InText
 	InTree->FunctionInstanceHandle->NotifyPreChange();
 	InTree->FunctionInstance->LayerNames[Counter] = InText;
 	InTree->MaterialEditorInstance->CopyToSourceInstance(true);
-	InTree->FunctionInstanceHandle->NotifyPostChange();
+	InTree->FunctionInstanceHandle->NotifyPostChange(EPropertyChangeType::ValueSet);
 }
 
 FReply SMaterialLayersFunctionsInstanceTreeItem::OnLayerDrop(const FDragDropEvent& DragDropEvent)
@@ -193,7 +193,7 @@ FReply SMaterialLayersFunctionsInstanceTreeItem::OnLayerDrop(const FDragDropEven
 
 				Tree->OnExpansionChanged(SwappablePropertyData, bOriginalSwappingExpansion);
 				Tree->OnExpansionChanged(SwappingPropertyData, bOriginalSwappableExpansion);
-				Tree->FunctionInstanceHandle->NotifyPostChange();
+				Tree->FunctionInstanceHandle->NotifyPostChange(EPropertyChangeType::ValueSet);
 				Tree->CreateGroupsWidget();
 				Tree->RequestTreeRefresh();
 				Tree->SetParentsExpansionState();
@@ -1201,7 +1201,7 @@ void SMaterialLayersFunctionsInstanceTree::AddLayer()
 	const FScopedTransaction Transaction(LOCTEXT("AddLayerAndBlend", "Add a new Layer and a Blend into it"));
 	FunctionInstanceHandle->NotifyPreChange();
 	FunctionInstance->AppendBlendedLayer();
-	FunctionInstanceHandle->NotifyPostChange();
+	FunctionInstanceHandle->NotifyPostChange(EPropertyChangeType::ArrayAdd);
 	CreateGroupsWidget();
 	RequestTreeRefresh();
 }
@@ -1211,7 +1211,7 @@ void SMaterialLayersFunctionsInstanceTree::RemoveLayer(int32 Index)
 	const FScopedTransaction Transaction(LOCTEXT("RemoveLayerAndBlend", "Remove a Layer and the attached Blend"));
 	FunctionInstanceHandle->NotifyPreChange();
 	FunctionInstance->RemoveBlendedLayerAt(Index);
-	FunctionInstanceHandle->NotifyPostChange();
+	FunctionInstanceHandle->NotifyPostChange(EPropertyChangeType::ArrayRemove);
 	CreateGroupsWidget();
 	RequestTreeRefresh();
 }
@@ -1221,7 +1221,7 @@ FReply SMaterialLayersFunctionsInstanceTree::UnlinkLayer(int32 Index)
 	const FScopedTransaction Transaction(LOCTEXT("UnlinkLayerFromParent", "Unlink a layer from the parent"));
 	FunctionInstanceHandle->NotifyPreChange();
 	FunctionInstance->UnlinkLayerFromParent(Index);
-	FunctionInstanceHandle->NotifyPostChange();
+	FunctionInstanceHandle->NotifyPostChange(EPropertyChangeType::ValueSet);
 	CreateGroupsWidget();
 	return FReply::Handled();
 }
@@ -1231,7 +1231,7 @@ FReply SMaterialLayersFunctionsInstanceTree::RelinkLayersToParent()
 	const FScopedTransaction Transaction(LOCTEXT("RelinkLayersToParent", "Relink layers to parent"));
 	FunctionInstanceHandle->NotifyPreChange();
 	FunctionInstance->RelinkLayersToParent();
-	FunctionInstanceHandle->NotifyPostChange();
+	FunctionInstanceHandle->NotifyPostChange(EPropertyChangeType::ValueSet);
 	MaterialEditorInstance->RegenerateArrays();
 	CreateGroupsWidget();
 	return FReply::Handled();
@@ -1263,7 +1263,7 @@ FReply SMaterialLayersFunctionsInstanceTree::ToggleLayerVisibility(int32 Index)
 		const FScopedTransaction Transaction(LOCTEXT("ToggleLayerAndBlendVisibility", "Toggles visibility for a blended layer"));
 		FunctionInstanceHandle->NotifyPreChange();
 		FunctionInstance->ToggleBlendedLayerVisibility(Index);
-		FunctionInstanceHandle->NotifyPostChange();
+		FunctionInstanceHandle->NotifyPostChange(EPropertyChangeType::ValueSet);
 		CreateGroupsWidget();
 		return FReply::Handled();
 	}
@@ -1286,7 +1286,7 @@ FReply SMaterialLayersFunctionsInstanceTree::ToggleLayerVisibility(int32 Index)
 		}
 
 		bLayerIsolated = !bLayerIsolated;
-		FunctionInstanceHandle->NotifyPostChange();
+		FunctionInstanceHandle->NotifyPostChange(EPropertyChangeType::ValueSet);
 		CreateGroupsWidget();
 		return FReply::Handled();
 	}
diff --git a/Engine/Source/Editor/MovieSceneTools/Private/MovieSceneBuiltInEasingFunctionCustomization.cpp b/Engine/Source/Editor/MovieSceneTools/Private/MovieSceneBuiltInEasingFunctionCustomization.cpp
--- a/Engine/Source/Editor/MovieSceneTools/Private/MovieSceneBuiltInEasingFunctionCustomization.cpp
+++ b/Engine/Source/Editor/MovieSceneTools/Private/MovieSceneBuiltInEasingFunctionCustomization.cpp
@@ -239,7 +239,7 @@ FReply FMovieSceneBuiltInEasingFunctionCustomization::SetType(EMovieSceneBuiltIn
 		*((EMovieSceneBuiltInEasing*)Ptr) = NewType;
 	}
 
-	TypeProperty->NotifyPostChange();
+	TypeProperty->NotifyPostChange(EPropertyChangeType::ValueSet);
 	TypeProperty->NotifyFinishedChangingProperties();
 
 	return FReply::Unhandled();
diff --git a/Engine/Source/Editor/MovieSceneTools/Private/MovieSceneObjectBindingIDCustomization.cpp b/Engine/Source/Editor/MovieSceneTools/Private/MovieSceneObjectBindingIDCustomization.cpp
--- a/Engine/Source/Editor/MovieSceneTools/Private/MovieSceneObjectBindingIDCustomization.cpp
+++ b/Engine/Source/Editor/MovieSceneTools/Private/MovieSceneObjectBindingIDCustomization.cpp
@@ -140,7 +140,7 @@ void FMovieSceneObjectBindingIDCustomization::SetCurrentValue(const FMovieSceneO
 		*static_cast<FMovieSceneObjectBindingID*>(Ptrs[0]) = InObjectBinding;
 	}
 	
-	StructProperty->NotifyPostChange();
+	StructProperty->NotifyPostChange(EPropertyChangeType::ValueSet);
 	StructProperty->NotifyFinishedChangingProperties();
 }
 
diff --git a/Engine/Source/Editor/MovieSceneTools/Private/MovieSceneSectionDetailsCustomization.cpp b/Engine/Source/Editor/MovieSceneTools/Private/MovieSceneSectionDetailsCustomization.cpp
--- a/Engine/Source/Editor/MovieSceneTools/Private/MovieSceneSectionDetailsCustomization.cpp
+++ b/Engine/Source/Editor/MovieSceneTools/Private/MovieSceneSectionDetailsCustomization.cpp
@@ -206,7 +206,7 @@ void FMovieSceneSectionDetailsCustomization::OnRangeStartTextCommitted(const FTe
 		}
 	}
 
-	MovieSceneSectionPropertyHandle->NotifyPostChange();
+	MovieSceneSectionPropertyHandle->NotifyPostChange(EPropertyChangeType::ValueSet);
 	MovieSceneSectionPropertyHandle->NotifyFinishedChangingProperties();
 
 	GEditor->EndTransaction();
@@ -406,7 +406,7 @@ void FMovieSceneSectionDetailsCustomization::OnRangeEndTextCommitted(const FText
 		}
 	}
 
-	MovieSceneSectionPropertyHandle->NotifyPostChange();
+	MovieSceneSectionPropertyHandle->NotifyPostChange(EPropertyChangeType::ValueSet);
 	MovieSceneSectionPropertyHandle->NotifyFinishedChangingProperties();
 
 	GEditor->EndTransaction();
diff --git a/Engine/Source/Editor/Persona/Private/AnimGraphNodeDetails.cpp b/Engine/Source/Editor/Persona/Private/AnimGraphNodeDetails.cpp
--- a/Engine/Source/Editor/Persona/Private/AnimGraphNodeDetails.cpp
+++ b/Engine/Source/Editor/Persona/Private/AnimGraphNodeDetails.cpp
@@ -347,7 +347,7 @@ void UpdateInputScaleBiasWithMinValue(float MinValue, TSharedRef<class IProperty
 		*BiasPtr = -MinValue * *ScalePtr;
 	}
 
-	InputBiasScaleStructPropertyHandle->NotifyPostChange();
+	InputBiasScaleStructPropertyHandle->NotifyPostChange(EPropertyChangeType::ValueSet);
 }
 
 void UpdateInputScaleBiasWithMaxValue(float MaxValue, TSharedRef<class IPropertyHandle> InputBiasScaleStructPropertyHandle)
@@ -374,7 +374,7 @@ void UpdateInputScaleBiasWithMaxValue(float MaxValue, TSharedRef<class IProperty
 		*BiasPtr = -MinValue * *ScalePtr;
 	}
 
-	InputBiasScaleStructPropertyHandle->NotifyPostChange();
+	InputBiasScaleStructPropertyHandle->NotifyPostChange(EPropertyChangeType::ValueSet);
 }
 
 TOptional<float> GetMinValueInputScaleBias(TSharedRef<class IPropertyHandle> InputBiasScaleStructPropertyHandle)
diff --git a/Engine/Source/Editor/PropertyEditor/Private/PropertyHandleImpl.h b/Engine/Source/Editor/PropertyEditor/Private/PropertyHandleImpl.h
--- a/Engine/Source/Editor/PropertyEditor/Private/PropertyHandleImpl.h
+++ b/Engine/Source/Editor/PropertyEditor/Private/PropertyHandleImpl.h
@@ -472,9 +472,9 @@ public:
 	virtual bool GeneratePossibleValues(TArray< TSharedPtr<FString> >& OutOptionStrings, TArray< FText >& OutToolTips, TArray<bool>& OutRestrictedItems) override;
 	virtual FPropertyAccess::Result SetObjectValueFromSelection() override;
 	virtual void NotifyPreChange() override;
-	virtual void NotifyPostChange( EPropertyChangeType::Type ChangeType = EPropertyChangeType::Unspecified ) override;
+	virtual void NotifyPostChange(EPropertyChangeType::Type ChangeType) override;
 	virtual void NotifyFinishedChangingProperties() override;
-	virtual void AddRestriction( TSharedRef<const FPropertyRestriction> Restriction )override;
+	virtual void AddRestriction(TSharedRef<const FPropertyRestriction> Restriction) override;
 	virtual bool IsHidden(const FString& Value) const override;
 	virtual bool IsHidden(const FString& Value, TArray<FText>& OutReasons) const override;
 	virtual bool IsDisabled(const FString& Value) const override;
@@ -483,7 +483,7 @@ public:
 	virtual bool IsRestricted(const FString& Value, TArray<FText>& OutReasons) const override;
 	virtual bool GenerateRestrictionToolTip(const FString& Value, FText& OutTooltip) const override;
 	virtual void SetIgnoreValidation(bool bInIgnore) override;
-	virtual TArray<TSharedPtr<IPropertyHandle>> AddChildStructure( TSharedRef<FStructOnScope> ChildStructure ) override;
+	virtual TArray<TSharedPtr<IPropertyHandle>> AddChildStructure(TSharedRef<FStructOnScope> ChildStructure) override;
 	virtual bool CanResetToDefault() const override;
 	virtual void ExecuteCustomResetToDefault(const FResetToDefaultOverride& InOnCustomResetToDefault) override;
 	virtual FName GetDefaultCategoryName() const override;
diff --git a/Engine/Source/Editor/PropertyEditor/Private/PropertyNode.h b/Engine/Source/Editor/PropertyEditor/Private/PropertyNode.h
--- a/Engine/Source/Editor/PropertyEditor/Private/PropertyNode.h
+++ b/Engine/Source/Editor/PropertyEditor/Private/PropertyNode.h
@@ -1073,7 +1073,7 @@ protected:
 	TWeakPtr< class FDetailTreeNode > TreeNode;
 
 	/**
-	 * Stores metadata for this instasnce of the property (in contrast
+	 * Stores metadata for this instance of the property (in contrast
 	 * to regular metadata, which is stored per-class)
 	 */
 	TMap<FName, FString> InstanceMetaData;
diff --git a/Engine/Source/Editor/PropertyEditor/Private/SDetailSingleItemRow.cpp b/Engine/Source/Editor/PropertyEditor/Private/SDetailSingleItemRow.cpp
--- a/Engine/Source/Editor/PropertyEditor/Private/SDetailSingleItemRow.cpp
+++ b/Engine/Source/Editor/PropertyEditor/Private/SDetailSingleItemRow.cpp
@@ -227,7 +227,7 @@ FReply SDetailSingleItemRow::OnArrayDrop(const FDragDropEvent& DragDropEvent)
 		ParentHandle->MoveElementTo(OriginalIndex, NewIndex);
 
 		FPropertyChangedEvent MoveEvent(SwappingHandle->GetParentHandle()->GetProperty(), EPropertyChangeType::Unspecified);
-		SwappingHandle->GetParentHandle()->NotifyPostChange();
+		SwappingHandle->GetParentHandle()->NotifyPostChange(EPropertyChangeType::Unspecified);
 		if (DetailsView->GetPropertyUtilities().IsValid())
 		{
 			DetailsView->GetPropertyUtilities()->NotifyFinishedChangingProperties(MoveEvent);
diff --git a/Engine/Source/Editor/PropertyEditor/Public/PropertyHandle.h b/Engine/Source/Editor/PropertyEditor/Public/PropertyHandle.h
--- a/Engine/Source/Editor/PropertyEditor/Public/PropertyHandle.h
+++ b/Engine/Source/Editor/PropertyEditor/Public/PropertyHandle.h
@@ -318,7 +318,7 @@ public:
 	 * Called to manually notify root objects that this property has changed
 	 * This does not need to be called when SetValue functions are used since it will be called automatically
 	 */
-	virtual void NotifyPostChange( EPropertyChangeType::Type ChangeType = EPropertyChangeType::Unspecified ) = 0;
+	virtual void NotifyPostChange(EPropertyChangeType::Type ChangeType) = 0;
 
 	/**
 	 * Called to manually notify root objects that this property has finished changing
diff --git a/Engine/Source/Editor/SequenceRecorder/Private/SClassPropertyRecorderSettings.cpp b/Engine/Source/Editor/SequenceRecorder/Private/SClassPropertyRecorderSettings.cpp
--- a/Engine/Source/Editor/SequenceRecorder/Private/SClassPropertyRecorderSettings.cpp
+++ b/Engine/Source/Editor/SequenceRecorder/Private/SClassPropertyRecorderSettings.cpp
@@ -178,7 +178,7 @@ void SClassPropertyRecorderSettings::HandlePropertyCheckStateChanged(ECheckBoxSt
 				PropertiesToRecord->Remove(PropertyPath);
 			}
 
-			PropertiesHandle->NotifyPostChange();
+			PropertiesHandle->NotifyPostChange(InState == ECheckBoxState::Checked ? EPropertyChangeType::ArrayAdd : EPropertyChangeType::ArrayRemove);
 		}
 	}
 }
diff --git a/Engine/Source/Editor/UnrealEd/Private/Settings/ProjectPackagingSettingsCustomization.h b/Engine/Source/Editor/UnrealEd/Private/Settings/ProjectPackagingSettingsCustomization.h
--- a/Engine/Source/Editor/UnrealEd/Private/Settings/ProjectPackagingSettingsCustomization.h
+++ b/Engine/Source/Editor/UnrealEd/Private/Settings/ProjectPackagingSettingsCustomization.h
@@ -333,7 +333,7 @@ protected:
 		RawCultureStringArray->Add(CultureName);
 		if(!IsInBatchSelectOperation)
 		{
-			CulturesPropertyHandle->NotifyPostChange();
+			CulturesPropertyHandle->NotifyPostChange(EPropertyChangeType::ArrayAdd);
 		}
 	}
 
@@ -349,7 +349,7 @@ protected:
 		RawCultureStringArray->Remove(CultureName);
 		if(!IsInBatchSelectOperation)
 		{
-			CulturesPropertyHandle->NotifyPostChange();
+			CulturesPropertyHandle->NotifyPostChange(EPropertyChangeType::ArrayRemove);
 		}
 	}
 
@@ -366,7 +366,7 @@ protected:
 
 	void OnPostBatchSelect()
 	{
-		CulturesPropertyHandle->NotifyPostChange();
+		CulturesPropertyHandle->NotifyPostChange(EPropertyChangeType::ValueSet);
 		IsInBatchSelectOperation = false;
 	}
 
