Route Pre/PostEditChange events during drag operations, not just only when the operation is commited.
We'll probably have to fix some classes to properly check the interaction type and filter out EPropertyChangeType::Interactive events.

diff --git a/Engine/Source/Editor/DetailCustomizations/Private/ComponentTransformDetails.cpp b/Engine/Source/Editor/DetailCustomizations/Private/ComponentTransformDetails.cpp
--- a/Engine/Source/Editor/DetailCustomizations/Private/ComponentTransformDetails.cpp
+++ b/Engine/Source/Editor/DetailCustomizations/Private/ComponentTransformDetails.cpp
@@ -1107,14 +1107,14 @@ void FComponentTransformDetails::OnSetTransform(ETransformField::Type TransformF
 							// Default subobjects must be included in any undo/redo operations
 							SceneComponent->SetFlags(RF_Transactional);
 						}
+					}
 
-						// Have to downcast here because of function overloading and inheritance not playing nicely
-						// We don't call PreEditChange for non commit changes because most classes implement the version that doesn't check the interaction type
-						((UObject*)SceneComponent)->PreEditChange(PropertyChain);
-						if (EditedActor && EditedActor->GetRootComponent() == SceneComponent)
-						{
-							((UObject*)EditedActor)->PreEditChange(PropertyChain);
-						}
+					// Have to downcast here because of function overloading and inheritance not playing nicely
+					// We don't call PreEditChange for non commit changes because most classes implement the version that doesn't check the interaction type
+					((UObject*)SceneComponent)->PreEditChange(PropertyChain);
+					if (EditedActor && EditedActor->GetRootComponent() == SceneComponent)
+					{
+						((UObject*)EditedActor)->PreEditChange(PropertyChain);
 					}
 
 					if (NotifyHook)
@@ -1261,13 +1261,11 @@ void FComponentTransformDetails::OnSetTransform(ETransformField::Type TransformF
 				AActor* EditedActor = SceneComponent->GetOwner();
 				FString SceneComponentPath = SceneComponent->GetPathName(EditedActor);
 				
-				if (bCommitted)
-				{
-					// This can invalidate OldSceneComponent
-					// We don't call PostEditChange for non commit changes because most classes implement the version that doesn't check the interaction type
-					OldSceneComponent->PostEditChangeChainProperty(PropertyChangedChainEvent);
-				}
-				else
+				// This can invalidate OldSceneComponent
+				// We don't call PostEditChange for non commit changes because most classes implement the version that doesn't check the interaction type
+				OldSceneComponent->PostEditChangeChainProperty(PropertyChangedChainEvent);
+
+				if (!bCommitted)
 				{
 					SnapshotTransactionBuffer(OldSceneComponent);
 				}
@@ -1276,12 +1274,10 @@ void FComponentTransformDetails::OnSetTransform(ETransformField::Type TransformF
 
 				if (EditedActor && EditedActor->GetRootComponent() == SceneComponent)
 				{
-					if (bCommitted)
-					{
-						EditedActor->PostEditChangeChainProperty(PropertyChangedChainEvent);
-						SceneComponent = FindObject<USceneComponent>(EditedActor, *SceneComponentPath);
-					}
-					else
+					EditedActor->PostEditChangeChainProperty(PropertyChangedChainEvent);
+					SceneComponent = FindObject<USceneComponent>(EditedActor, *SceneComponentPath);
+
+					if (!bCommitted)
 					{
 						SnapshotTransactionBuffer(EditedActor);
 					}
