Fix unrelevant comment following 14131069.

diff --git a/Engine/Source/Editor/DetailCustomizations/Private/ComponentTransformDetails.cpp b/Engine/Source/Editor/DetailCustomizations/Private/ComponentTransformDetails.cpp
--- a/Engine/Source/Editor/DetailCustomizations/Private/ComponentTransformDetails.cpp
+++ b/Engine/Source/Editor/DetailCustomizations/Private/ComponentTransformDetails.cpp
@@ -1110,7 +1110,6 @@ void FComponentTransformDetails::OnSetTransform(ETransformField::Type TransformF
 					}
 
 					// Have to downcast here because of function overloading and inheritance not playing nicely
-					// We don't call PreEditChange for non commit changes because most classes implement the version that doesn't check the interaction type
 					((UObject*)SceneComponent)->PreEditChange(PropertyChain);
 					if (EditedActor && EditedActor->GetRootComponent() == SceneComponent)
 					{
@@ -1262,7 +1261,6 @@ void FComponentTransformDetails::OnSetTransform(ETransformField::Type TransformF
 				FString SceneComponentPath = SceneComponent->GetPathName(EditedActor);
 				
 				// This can invalidate OldSceneComponent
-				// We don't call PostEditChange for non commit changes because most classes implement the version that doesn't check the interaction type
 				OldSceneComponent->PostEditChangeChainProperty(PropertyChangedChainEvent);
 
 				if (!bCommitted)
