Details View: Include the list of affected UObjects when notifying about a value change for customized math objects.

diff --git a/Engine/Source/Editor/DetailCustomizations/Private/MathStructProxyCustomizations.cpp b/Engine/Source/Editor/DetailCustomizations/Private/MathStructProxyCustomizations.cpp
--- a/Engine/Source/Editor/DetailCustomizations/Private/MathStructProxyCustomizations.cpp
+++ b/Engine/Source/Editor/DetailCustomizations/Private/MathStructProxyCustomizations.cpp
@@ -461,7 +461,7 @@ bool FMatrixStructCustomization::FlushValues( TWeakPtr<IPropertyHandle> Property
 
 	if (PropertyUtilities.IsValid() && !bIsInteractiveChangeInProgress)
 	{
-		FPropertyChangedEvent ChangeEvent(PropertyHandle->GetProperty(), EPropertyChangeType::ValueSet);
+		FPropertyChangedEvent ChangeEvent(PropertyHandle->GetProperty(), EPropertyChangeType::ValueSet, OuterObjects);
 		PropertyUtilities->NotifyFinishedChangingProperties(ChangeEvent);
 	}
 
@@ -600,7 +600,7 @@ bool FTransformStructCustomization::FlushValues( TWeakPtr<IPropertyHandle> Prope
 
 	if (PropertyUtilities.IsValid() && !bIsInteractiveChangeInProgress)
 	{
-		FPropertyChangedEvent ChangeEvent(PropertyHandle->GetProperty(), EPropertyChangeType::ValueSet);
+		FPropertyChangedEvent ChangeEvent(PropertyHandle->GetProperty(), EPropertyChangeType::ValueSet, OuterObjects);
 		PropertyUtilities->NotifyFinishedChangingProperties(ChangeEvent);
 	}
 
@@ -741,7 +741,7 @@ bool FQuatStructCustomization::FlushValues(TWeakPtr<IPropertyHandle> PropertyHan
 
 	if (PropertyUtilities.IsValid() && !bIsInteractiveChangeInProgress)
 	{
-		FPropertyChangedEvent ChangeEvent(PropertyHandle->GetProperty(), EPropertyChangeType::ValueSet);
+		FPropertyChangedEvent ChangeEvent(PropertyHandle->GetProperty(), EPropertyChangeType::ValueSet, OuterObjects);
 		PropertyUtilities->NotifyFinishedChangingProperties(ChangeEvent);
 	}
 
