Editgrating 14474663 to UE5 to fix spline component crash (In spline details customization, update the active spline visualizer when a point is initially selected rather than when customization is created.)

diff --git a/Engine/Source/Editor/ComponentVisualizers/Private/SplineComponentVisualizer.cpp b/Engine/Source/Editor/ComponentVisualizers/Private/SplineComponentVisualizer.cpp
--- a/Engine/Source/Editor/ComponentVisualizers/Private/SplineComponentVisualizer.cpp
+++ b/Engine/Source/Editor/ComponentVisualizers/Private/SplineComponentVisualizer.cpp
@@ -2944,6 +2944,44 @@ bool FSplineComponentVisualizer::CanResetToDefault() const
     }
 }
 
+bool FSplineComponentVisualizer::HandleSelectFirstLastSplinePoint(USplineComponent* InSplineComponent, bool bFirstPoint)
+{
+	const FScopedTransaction Transaction(LOCTEXT("SelectFirstSplinePoint", "Select First Spline Point"));
+
+	check(InSplineComponent);
+	check(SelectionState);
+
+	bool bResetEditedSplineComponent = false;
+	if (GetEditedSplineComponent() != InSplineComponent)
+	{
+		SetEditedSplineComponent(InSplineComponent);
+		bResetEditedSplineComponent = true;
+	}
+
+	OnSelectFirstLastSplinePoint(bFirstPoint);
+
+	return bResetEditedSplineComponent;
+}
+
+bool FSplineComponentVisualizer::HandleSelectAllSplinePoints(USplineComponent* InSplineComponent)
+{
+	const FScopedTransaction Transaction(LOCTEXT("SelectAllSplinePoints", "Select All Spline Points"));
+
+	check(InSplineComponent);
+	check(SelectionState);
+
+	bool bResetEditedSplineComponent = false;
+	if (GetEditedSplineComponent() != InSplineComponent)
+	{
+		SetEditedSplineComponent(InSplineComponent);
+		bResetEditedSplineComponent = true;
+	}
+
+	OnSelectAllSplinePoints();
+
+	return bResetEditedSplineComponent;
+}
+
 void FSplineComponentVisualizer::OnSelectFirstLastSplinePoint(bool bFirstPoint)
 {
 	const FScopedTransaction Transaction(LOCTEXT("SelectFirstSplinePoint", "Select First Spline Point"));
diff --git a/Engine/Source/Editor/ComponentVisualizers/Public/SplineComponentVisualizer.h b/Engine/Source/Editor/ComponentVisualizers/Public/SplineComponentVisualizer.h
--- a/Engine/Source/Editor/ComponentVisualizers/Public/SplineComponentVisualizer.h
+++ b/Engine/Source/Editor/ComponentVisualizers/Public/SplineComponentVisualizer.h
@@ -226,23 +226,20 @@ public:
 	/** Add menu sections to the context menu */
 	virtual void GenerateContextMenuSections(FMenuBuilder& InMenuBuilder) const;
 
-	/** Set the spline component we are currently editing */
-	void SetEditedSplineComponent(const USplineComponent* InSplineComponent);
-		
 	/** Get the spline component we are currently editing */
 	USplineComponent* GetEditedSplineComponent() const;
 
 	const TSet<int32>& GetSelectedKeys() const { check(SelectionState); return SelectionState->GetSelectedKeys(); }
 
-	/** Select first or last spline point */
-	void OnSelectFirstLastSplinePoint(bool bFirstPoint);
+	/** Select first or last spline point, returns true if the spline component being edited has changed */
+	bool HandleSelectFirstLastSplinePoint(USplineComponent* InSplineComponent, bool bFirstPoint);
+
+	/** Select all spline points, , returns true if the spline component being edited has changed */
+	bool HandleSelectAllSplinePoints(USplineComponent* InSplineComponent);
 
 	/** Select next or prev spline point, loops when last point is currently selected */
 	void OnSelectPrevNextSplinePoint(bool bNextPoint, bool bAddToSelection);
 
-	/** Select all spline points, if no spline points selected yet the currently edited spline component will be set as well */
-	void OnSelectAllSplinePoints();
-
 protected:
 
 	/** Determine if any selected key index is out of range (perhaps because something external has modified the spline) */
@@ -356,6 +353,12 @@ protected:
 	void OnResetToDefault();
 	bool CanResetToDefault() const;
 
+	/** Select first or last spline point */
+	void OnSelectFirstLastSplinePoint(bool bFirstPoint);
+
+	/** Select all spline points, if no spline points selected yet the currently edited spline component will be set as well */
+	void OnSelectAllSplinePoints();
+
 	bool CanSelectSplinePoints() const;
 
 	/** Generate the submenu containing available selection actions */
@@ -373,6 +376,9 @@ protected:
 	/** Generate the submenu containing the lock axis types */
 	void GenerateLockAxisSubMenu(FMenuBuilder& MenuBuilder) const;
 
+	/** Helper function to set edited component we are currently editing */
+	void SetEditedSplineComponent(const USplineComponent* InSplineComponent);
+
 	void CreateSplineGeneratorPanel();
 
 	// FGCObject interface
diff --git a/Engine/Source/Editor/DetailCustomizations/Private/SplineComponentDetails.cpp b/Engine/Source/Editor/DetailCustomizations/Private/SplineComponentDetails.cpp
--- a/Engine/Source/Editor/DetailCustomizations/Private/SplineComponentDetails.cpp
+++ b/Engine/Source/Editor/DetailCustomizations/Private/SplineComponentDetails.cpp
@@ -29,6 +29,7 @@
 #include "Widgets/Input/SNumericEntryBox.h"
 #include "ScopedTransaction.h"
 #include "Editor.h"
+#include "LevelEditorViewport.h"
 
 #define LOCTEXT_NAMESPACE "SplineComponentDetails"
 
@@ -199,7 +200,7 @@ private:
 	FSharedRotatorValue Rotation;
 	TSharedValue<ESplinePointType::Type> PointType;
 
-	FSplineComponentVisualizer* SplineVisualizer;
+	TSharedPtr<FSplineComponentVisualizer> SplineVisualizer;
 	FProperty* SplineCurvesProperty;
 	TArray<TSharedPtr<FString>> SplinePointTypes;
 	TSharedPtr<ISplineMetadataDetails> SplineMetaDataDetails;
@@ -210,8 +211,8 @@ FSplinePointDetails::FSplinePointDetails(USplineComponent* InOwningSplineCompone
 	: SplineComp(nullptr)
 {
 	TSharedPtr<FComponentVisualizer> Visualizer = GUnrealEd->FindComponentVisualizer(InOwningSplineComponent->GetClass());
-	SplineVisualizer = (FSplineComponentVisualizer*)Visualizer.Get();
-	check(SplineVisualizer);
+	SplineVisualizer = StaticCastSharedPtr<FSplineComponentVisualizer>(Visualizer);
+	check(SplineVisualizer.IsValid());
 
 	SplineCurvesProperty = FindFProperty<FProperty>(USplineComponent::StaticClass(), GET_MEMBER_NAME_CHECKED(USplineComponent, SplineCurves));
 
@@ -541,7 +542,7 @@ void FSplinePointDetails::GenerateChildContent(IDetailChildrenBuilder& ChildrenB
 	];
 
 	check(SplineComp);
-	check(SplineVisualizer);
+	check(SplineVisualizer.IsValid());
 
 	if (SplineVisualizer->GetEditedSplineComponent() == SplineComp && SplineVisualizer->GetSelectedKeys().Num() > 0)
 	{
@@ -570,13 +571,7 @@ void FSplinePointDetails::Tick(float DeltaTime)
 
 void FSplinePointDetails::UpdateValues()
 {
-	check(SplineVisualizer);
-	check(SplineComp);
-
-	if (SplineVisualizer->GetEditedSplineComponent() != SplineComp)
-	{
-		SplineVisualizer->SetEditedSplineComponent(SplineComp);
-	}
+	check(SplineVisualizer.IsValid());
 
 	bool bNeedsRebuild = false;
 	const TSet<int32>& NewSelectedKeys = SplineVisualizer->GetSelectedKeys();
@@ -872,22 +867,32 @@ void FSplinePointDetails::OnSplinePointTypeChanged(TSharedPtr<FString> NewValue,
 
 FReply FSplinePointDetails::OnSelectFirstLastSplinePoint(bool bFirst)
 {
-	check(SplineVisualizer);
-	SplineVisualizer->OnSelectFirstLastSplinePoint(bFirst);
+	check(SplineVisualizer.IsValid());
+	check(SplineComp);
+	if (SplineVisualizer->HandleSelectFirstLastSplinePoint(SplineComp, bFirst))
+	{
+		TSharedPtr<FComponentVisualizer> Visualizer = StaticCastSharedPtr<FComponentVisualizer>(SplineVisualizer);
+		GUnrealEd->ComponentVisManager.SetActiveComponentVis(GCurrentLevelEditingViewportClient, Visualizer);
+	}
 	return FReply::Handled();
 }
 
 FReply FSplinePointDetails::OnSelectPrevNextSplinePoint(bool bNext, bool bAddToSelection)
 {
-	check(SplineVisualizer);
+	check(SplineVisualizer.IsValid());
 	SplineVisualizer->OnSelectPrevNextSplinePoint(bNext, bAddToSelection);
 	return FReply::Handled();
 }
 
 FReply FSplinePointDetails::OnSelectAllSplinePoints()
 {
-	check(SplineVisualizer);
-	SplineVisualizer->OnSelectAllSplinePoints();
+	check(SplineVisualizer.IsValid());
+	check(SplineComp);
+	if (SplineVisualizer->HandleSelectAllSplinePoints(SplineComp))
+	{
+		TSharedPtr<FComponentVisualizer> Visualizer = StaticCastSharedPtr<FComponentVisualizer>(SplineVisualizer);
+		GUnrealEd->ComponentVisManager.SetActiveComponentVis(GCurrentLevelEditingViewportClient, Visualizer);
+	}
 	return FReply::Handled();
 }
 
