Fixed guid, datetime, and timespan black on black text with reskin changes

diff --git a/Engine/Source/Editor/DetailCustomizations/Private/DateTimeStructCustomization.cpp b/Engine/Source/Editor/DetailCustomizations/Private/DateTimeStructCustomization.cpp
--- a/Engine/Source/Editor/DetailCustomizations/Private/DateTimeStructCustomization.cpp
+++ b/Engine/Source/Editor/DetailCustomizations/Private/DateTimeStructCustomization.cpp
@@ -41,6 +41,7 @@ void FDateTimeStructCustomization::CustomizeHeader(TSharedRef<IPropertyHandle> S
 				.OnTextChanged(this, &FDateTimeStructCustomization::HandleTextBoxTextChanged)
 				.OnTextCommitted(this, &FDateTimeStructCustomization::HandleTextBoxTextCommited)
 				.SelectAllTextOnCommit(true)
+				.Font(IPropertyTypeCustomizationUtils::GetRegularFont())
 				.Text(this, &FDateTimeStructCustomization::HandleTextBoxText)
 		];
 }
@@ -59,11 +60,13 @@ FSlateColor FDateTimeStructCustomization::HandleTextBoxForegroundColor() const
 {
 	if (InputValid)
 	{
-		static const FName InvertedForegroundName("InvertedForeground");
-		return FEditorStyle::GetSlateColor(InvertedForegroundName);
+		static const FName DefaultForeground("Colors.Foreground");
+		return FAppStyle::Get().GetSlateColor(DefaultForeground);
 	}
 
-	return FLinearColor::Red;
+	static const FName Red("Colors.AccentRed");
+
+	return FAppStyle::Get().GetSlateColor(Red);
 }
 
 
diff --git a/Engine/Source/Editor/DetailCustomizations/Private/GuidStructCustomization.cpp b/Engine/Source/Editor/DetailCustomizations/Private/GuidStructCustomization.cpp
--- a/Engine/Source/Editor/DetailCustomizations/Private/GuidStructCustomization.cpp
+++ b/Engine/Source/Editor/DetailCustomizations/Private/GuidStructCustomization.cpp
@@ -44,8 +44,9 @@ void FGuidStructCustomization::CustomizeHeader( TSharedRef<class IPropertyHandle
 			QuickSetMenuBuilder.AddMenuEntry(LOCTEXT("InvalidateAction", "Invalidate"), LOCTEXT("InvalidateActionHint", "Set an invalid globally unique identifier (GUID)."), FSlateIcon(), InvalidateAction);
 		}
 
-		QuickSetSlotContent = SNew(SComboButton)
-			.ContentPadding(FMargin(6.0, 2.0))
+		QuickSetSlotContent = 
+			SNew(SComboButton)
+			.ComboButtonStyle(FAppStyle::Get(), "SimpleComboButton")
 			.MenuContent()
 			[
 				QuickSetMenuBuilder.MakeWidget()
@@ -75,6 +76,7 @@ void FGuidStructCustomization::CustomizeHeader( TSharedRef<class IPropertyHandle
 						.OnTextChanged(this, &FGuidStructCustomization::HandleTextBoxTextChanged)
 						.OnTextCommitted(this, &FGuidStructCustomization::HandleTextBoxTextCommited)
 						.SelectAllTextOnCommit(true)
+						.Font(IPropertyTypeCustomizationUtils::GetRegularFont())
 						.Text(this, &FGuidStructCustomization::HandleTextBoxText)
 				]
 
@@ -114,23 +116,27 @@ void FGuidStructCustomization::HandleGuidActionClicked( EPropertyEditorGuidActio
 	if (Action == EPropertyEditorGuidActions::Generate)
 	{
 		SetGuidValue(FGuid::NewGuid());
+		InputValid = true;
 	}
 	else if (Action == EPropertyEditorGuidActions::Invalidate)
 	{
 		SetGuidValue(FGuid());
+		InputValid = true;
 	}
 }
 
 
-FSlateColor FGuidStructCustomization::HandleTextBoxForegroundColor( ) const
+FSlateColor FGuidStructCustomization::HandleTextBoxForegroundColor() const
 {
 	if (InputValid)
 	{
-		static const FName InvertedForegroundName("InvertedForeground");
-		return FEditorStyle::GetSlateColor(InvertedForegroundName);
+		static const FName DefaultForeground("Colors.Foreground");
+		return FAppStyle::Get().GetSlateColor(DefaultForeground);
 	}
 
-	return FLinearColor::Red;
+	static const FName Red("Colors.AccentRed");
+
+	return FAppStyle::Get().GetSlateColor(Red);
 }
 
 
diff --git a/Engine/Source/Editor/DetailCustomizations/Private/TimespanStructCustomization.cpp b/Engine/Source/Editor/DetailCustomizations/Private/TimespanStructCustomization.cpp
--- a/Engine/Source/Editor/DetailCustomizations/Private/TimespanStructCustomization.cpp
+++ b/Engine/Source/Editor/DetailCustomizations/Private/TimespanStructCustomization.cpp
@@ -52,13 +52,14 @@ void FTimespanStructCustomization::CustomizeHeader(TSharedRef<IPropertyHandle> S
 
 FSlateColor FTimespanStructCustomization::HandleTextBoxForegroundColor() const
 {
-	if (!InputValid)
+	if (InputValid)
 	{
-		return FLinearColor::Red;
+		static const FName DefaultForeground("Colors.Foreground");
+		return FAppStyle::Get().GetSlateColor(DefaultForeground);
 	}
 
-	static const FName InvertedForegroundName("InvertedForeground");
-	return FEditorStyle::GetSlateColor(InvertedForegroundName);
+	static const FName Red("Colors.AccentRed");
+	return FAppStyle::Get().GetSlateColor(Red);
 }
 
 
