Reset to defaults in the Collision section of the Static Mesh editor now use standard details view placement.

diff --git a/Engine/Source/Editor/DetailCustomizations/Private/BodyInstanceCustomization.cpp b/Engine/Source/Editor/DetailCustomizations/Private/BodyInstanceCustomization.cpp
--- a/Engine/Source/Editor/DetailCustomizations/Private/BodyInstanceCustomization.cpp
+++ b/Engine/Source/Editor/DetailCustomizations/Private/BodyInstanceCustomization.cpp
@@ -121,6 +121,10 @@ void FBodyInstanceCustomization::AddCollisionCategory(TSharedRef<class IProperty
 
 	IDetailGroup& CollisionGroup = StructBuilder.AddGroup( TEXT("Collision"), LOCTEXT("CollisionPresetsLabel", "Collision Presets") );
 	CollisionGroup.HeaderRow()
+	.OverrideResetToDefault(FResetToDefaultOverride::Create(
+		TAttribute<bool>::Create([this]() { return ShouldShowResetToDefaultProfile(); }),
+		FSimpleDelegate::CreateLambda([this]() { SetToDefaultProfile(); })
+	))
 	.NameContent()
 	[
 		SNew(SHorizontalBox)
@@ -159,25 +163,6 @@ void FBodyInstanceCustomization::AddCollisionCategory(TSharedRef<class IProperty
 					.ToolTipText(this, &FBodyInstanceCustomization::GetCollisionProfileComboBoxToolTip)
 				]
 			]
-
-			+ SHorizontalBox::Slot()
-			.VAlign(VAlign_Center)
-			.Padding(2.0f)
-			.AutoWidth()
-			[
-				SNew(SButton)
-				.OnClicked(this, &FBodyInstanceCustomization::SetToDefaultProfile)
-				.ContentPadding(0.f)
-				.ToolTipText(LOCTEXT("ResetToDefaultToolTip", "Reset to Default"))
-				.ButtonStyle(FEditorStyle::Get(), "NoBorder")
-				.IsEnabled(this, &FBodyInstanceCustomization::IsCollisionEnabled)
-				.Visibility(this, &FBodyInstanceCustomization::ShouldShowResetToDefaultProfile)
-				.Content()
-				[
-					SNew(SImage)
-					.Image(FEditorStyle::GetBrush("PropertyWindow.DiffersFromDefault"))
-				]
-			]
 		]
 	];
 
@@ -569,6 +554,10 @@ void FBodyInstanceCustomization::CreateCustomCollisionSetup( TSharedRef<class IP
 			CollisionGroup.AddWidgetRow()
 			.IsEnabled(CustomCollisionEnabled)
 			.Visibility(CustomCollisionVisibility)
+			.OverrideResetToDefault(FResetToDefaultOverride::Create(
+				TAttribute<bool>::Create([this, Index]() { return ShouldShowResetToDefaultResponse(Index); }),
+				FSimpleDelegate::CreateLambda([this, Index]() { SetToDefaultResponse(Index); })
+			))
 			.NameContent()
 			[
 				SNew(SBox)
@@ -618,22 +607,6 @@ void FBodyInstanceCustomization::CreateCustomCollisionSetup( TSharedRef<class IP
 					.OnCheckStateChanged( this, &FBodyInstanceCustomization::OnCollisionChannelChanged, Index, ECR_Block )
 					.IsChecked( this, &FBodyInstanceCustomization::IsCollisionChannelChecked, Index, ECR_Block )
 				]
-
-				+SHorizontalBox::Slot()
-				.VAlign(VAlign_Center)
-				[
-					SNew(SButton)
-					.OnClicked(this, &FBodyInstanceCustomization::SetToDefaultResponse, Index)
-					.Visibility(this, &FBodyInstanceCustomization::ShouldShowResetToDefaultResponse, Index)
-					.ContentPadding(0.f)
-					.ToolTipText(LOCTEXT("ResetToDefaultToolTip", "Reset to Default"))
-					.ButtonStyle( FEditorStyle::Get(), "NoBorder" )
-					.Content()
-					[
-						SNew(SImage)
-						.Image( FEditorStyle::GetBrush("PropertyWindow.DiffersFromDefault") )
-					]
-				]
 			];
 		}
 	}
@@ -664,6 +637,10 @@ void FBodyInstanceCustomization::CreateCustomCollisionSetup( TSharedRef<class IP
 			CollisionGroup.AddWidgetRow()
 			.IsEnabled(CustomCollisionEnabled)
 			.Visibility(CustomCollisionVisibility)
+			.OverrideResetToDefault(FResetToDefaultOverride::Create(
+				TAttribute<bool>::Create([this, Index]() { return ShouldShowResetToDefaultResponse(Index); }),
+				FSimpleDelegate::CreateLambda([this, Index]() { SetToDefaultResponse(Index); })
+			))
 			.NameContent()
 			[
 				SNew(SBox)
@@ -713,22 +690,6 @@ void FBodyInstanceCustomization::CreateCustomCollisionSetup( TSharedRef<class IP
 					.OnCheckStateChanged( this, &FBodyInstanceCustomization::OnCollisionChannelChanged, Index, ECR_Block )
 					.IsChecked( this, &FBodyInstanceCustomization::IsCollisionChannelChecked, Index, ECR_Block )
 				]
-
-				+SHorizontalBox::Slot()
-				.VAlign(VAlign_Center)
-				[
-					SNew(SButton)
-					.OnClicked(this, &FBodyInstanceCustomization::SetToDefaultResponse, Index)
-					.Visibility(this, &FBodyInstanceCustomization::ShouldShowResetToDefaultResponse, Index)
-					.ContentPadding(0.f)
-					.ToolTipText(LOCTEXT("ResetToDefaultToolTip", "Reset to Default"))
-					.ButtonStyle( FEditorStyle::Get(), "NoBorder" )
-					.Content()
-					[
-						SNew(SImage)
-						.Image( FEditorStyle::GetBrush("PropertyWindow.DiffersFromDefault") )
-					]
-				]
 			];
 		}
 	}
@@ -943,53 +904,45 @@ void FBodyInstanceCustomization::UpdateCollisionProfile()
 	CollsionProfileComboBox.Get()->SetSelectedItem(CollisionProfileComboList[AreAllCollisionUsingDefault() ? GetDefaultIndex() : GetCustomIndex()]);
 }
 
-FReply FBodyInstanceCustomization::SetToDefaultProfile()
+void FBodyInstanceCustomization::SetToDefaultProfile()
 {
 	// trigger transaction before UpdateCollisionProfile
 	const FScopedTransaction Transaction( LOCTEXT( "ResetCollisionProfile", "Reset Collision Profile" ) );
 	MarkAllBodiesDefaultCollision(false);
 	CollisionProfileNameHandle.Get()->ResetToDefault();
 	UpdateCollisionProfile();
-	return FReply::Handled();
 }
 
-EVisibility FBodyInstanceCustomization::ShouldShowResetToDefaultProfile() const
+bool FBodyInstanceCustomization::ShouldShowResetToDefaultProfile() const
 {
-	if (CollisionProfileNameHandle.Get()->DiffersFromDefault())
-	{
-		return EVisibility::Visible;
-	}
-
-	return EVisibility::Hidden;
+	return CollisionProfileNameHandle.Get()->DiffersFromDefault();
 }
 
-FReply FBodyInstanceCustomization::SetToDefaultResponse(int32 ValidIndex)
+void FBodyInstanceCustomization::SetToDefaultResponse(int32 Index)
 {
-	if ( ValidCollisionChannels.IsValidIndex(ValidIndex) )
+	if (ValidCollisionChannels.IsValidIndex(Index))
 	{
 		const FScopedTransaction Transaction( LOCTEXT( "ResetCollisionResponse", "Reset Collision Response" ) );
-		const ECollisionResponse DefaultResponse = FCollisionResponseContainer::GetDefaultResponseContainer().GetResponse(ValidCollisionChannels[ValidIndex].CollisionChannel);
+		const ECollisionResponse DefaultResponse = FCollisionResponseContainer::GetDefaultResponseContainer().GetResponse(ValidCollisionChannels[Index].CollisionChannel);
 
-		SetResponse(ValidIndex, DefaultResponse);
-		return FReply::Handled();
+		SetResponse(Index, DefaultResponse);
 	}
-
-	return FReply::Unhandled();
 }
 
-EVisibility FBodyInstanceCustomization::ShouldShowResetToDefaultResponse(int32 ValidIndex) const
+bool FBodyInstanceCustomization::ShouldShowResetToDefaultResponse(int32 Index) const
 {
-	if ( ValidCollisionChannels.IsValidIndex(ValidIndex) )
+	
+	if (ValidCollisionChannels.IsValidIndex(Index))
 	{
-		const ECollisionResponse DefaultResponse = FCollisionResponseContainer::GetDefaultResponseContainer().GetResponse(ValidCollisionChannels[ValidIndex].CollisionChannel);
+		const ECollisionResponse DefaultResponse = FCollisionResponseContainer::GetDefaultResponseContainer().GetResponse(ValidCollisionChannels[Index].CollisionChannel);
 
-		if (IsCollisionChannelChecked(ValidIndex, DefaultResponse) != ECheckBoxState::Checked)
+		if (IsCollisionChannelChecked(Index, DefaultResponse) != ECheckBoxState::Checked)
 		{
-			return EVisibility::Visible;
+			return true;
 		}
 	}
 
-	return EVisibility::Hidden;
+	return false;
 }
 
 bool FBodyInstanceCustomization::AreAllCollisionUsingDefault() const
diff --git a/Engine/Source/Editor/DetailCustomizations/Private/BodyInstanceCustomization.h b/Engine/Source/Editor/DetailCustomizations/Private/BodyInstanceCustomization.h
--- a/Engine/Source/Editor/DetailCustomizations/Private/BodyInstanceCustomization.h
+++ b/Engine/Source/Editor/DetailCustomizations/Private/BodyInstanceCustomization.h
@@ -61,11 +61,11 @@ private:
 	int32 InitializeObjectTypeComboList();
 
 	// set to default for profile setting
-	FReply SetToDefaultProfile();
-	EVisibility ShouldShowResetToDefaultProfile() const;
+	void SetToDefaultProfile();
+	bool ShouldShowResetToDefaultProfile() const;
 
-	FReply SetToDefaultResponse(int32 ValidIndex);
-	EVisibility ShouldShowResetToDefaultResponse(int32 ValidIndex) const;
+	void SetToDefaultResponse(int32 ValidIndex);
+	bool ShouldShowResetToDefaultResponse(int32 ValidIndex) const;
 
 	// collision channel check boxes
 	void OnCollisionChannelChanged(ECheckBoxState InNewValue, int32 ValidIndex, ECollisionResponse InCollisionResponse);
