Move blend profile customization to its own PropertyTypeCustomization. This is needed to support blend profiles in other contexts (i.e. as struct/array members or outside an AnimGraphNode)
#rnx

diff --git a/Engine/Source/Editor/DetailCustomizations/Private/AnimTransitionNodeDetails.cpp b/Engine/Source/Editor/DetailCustomizations/Private/AnimTransitionNodeDetails.cpp
--- a/Engine/Source/Editor/DetailCustomizations/Private/AnimTransitionNodeDetails.cpp
+++ b/Engine/Source/Editor/DetailCustomizations/Private/AnimTransitionNodeDetails.cpp
@@ -76,9 +76,6 @@ void FAnimTransitionNodeDetails::CustomizeDetails( IDetailLayoutBuilder& DetailB
 	UAnimStateTransitionNode* TransNode = TransitionNode.Get();
 	IDetailCategoryBuilder& TransitionCategory = DetailBuilder.EditCategory("Transition", LOCTEXT("TransitionCategoryTitle", "Transition") );
 
-	// Added below via CreateBlendProfilePicker
-	DetailBuilder.HideProperty(GET_MEMBER_NAME_CHECKED(UAnimStateTransitionNode, BlendProfile));
-
 	if (bTransitionToConduit)
 	{
 		// Transitions to conduits are just shorthand for some other real transition;
@@ -89,6 +86,7 @@ void FAnimTransitionNodeDetails::CustomizeDetails( IDetailLayoutBuilder& DetailB
 		DetailBuilder.HideProperty(GET_MEMBER_NAME_CHECKED(UAnimStateTransitionNode, CustomBlendCurve));
 		DetailBuilder.HideProperty(GET_MEMBER_NAME_CHECKED(UAnimStateTransitionNode, LogicType));
 		DetailBuilder.HideProperty(GET_MEMBER_NAME_CHECKED(UAnimStateTransitionNode, PriorityOrder));
+		DetailBuilder.HideProperty(GET_MEMBER_NAME_CHECKED(UAnimStateTransitionNode, BlendProfile));
 	}
 	else
 	{
@@ -230,37 +228,10 @@ void FAnimTransitionNodeDetails::CustomizeDetails( IDetailLayoutBuilder& DetailB
 		CrossfadeCategory.AddProperty(GET_MEMBER_NAME_CHECKED(UAnimStateTransitionNode, CrossfadeDuration)).DisplayName( LOCTEXT("DurationLabel", "Duration") );
 		CrossfadeCategory.AddProperty(GET_MEMBER_NAME_CHECKED(UAnimStateTransitionNode, BlendMode)).DisplayName(LOCTEXT("ModeLabel", "Mode")).IsEnabled(BlendSettingsEnabledAttribute);
 		CrossfadeCategory.AddProperty(GET_MEMBER_NAME_CHECKED(UAnimStateTransitionNode, CustomBlendCurve)).DisplayName(LOCTEXT("CurveLabel", "Custom Blend Curve")).IsEnabled(BlendSettingsEnabledAttribute);
-
-		USkeleton* TargetSkeleton = TransNode ? TransNode->GetAnimBlueprint()->TargetSkeleton : nullptr;
-
-		if(TargetSkeleton)
-		{
+		CrossfadeCategory.AddProperty(GET_MEMBER_NAME_CHECKED(UAnimStateTransitionNode, BlendProfile)).DisplayName(LOCTEXT("BlendProfileLabel", "Blend Profile")).IsEnabled(BlendSettingsEnabledAttribute);
+		
+		{// Blend Profile Mode
 			TSharedPtr<IPropertyHandle> BlendProfileHandle = DetailBuilder.GetProperty(GET_MEMBER_NAME_CHECKED(UAnimStateTransitionNode, BlendProfile));
-			UObject* BlendProfilePropertyValue = nullptr;
-			BlendProfileHandle->GetValue(BlendProfilePropertyValue);
-			UBlendProfile* CurrentProfile = Cast<UBlendProfile>(BlendProfilePropertyValue);
-
-			ISkeletonEditorModule& SkeletonEditorModule = FModuleManager::LoadModuleChecked<ISkeletonEditorModule>("SkeletonEditor");
-
-			FBlendProfilePickerArgs Args;
-			Args.InitialProfile = CurrentProfile;
-			Args.OnBlendProfileSelected = FOnBlendProfileSelected::CreateSP(this, &FAnimTransitionNodeDetails::OnBlendProfileChanged, BlendProfileHandle);
-			Args.bAllowNew = false;
-			Args.bAllowClear = true;
-			Args.bAllowRemove = false;
-
-			CrossfadeCategory.AddCustomRow(LOCTEXT("BlendProfileLabel", "Blend Profile"))
-				.IsEnabled(BlendSettingsEnabledAttribute)
-				.NameContent()
-				[
-					BlendProfileHandle->CreatePropertyNameWidget()
-				]
-				.ValueContent()
-				[
-					SkeletonEditorModule.CreateBlendProfilePicker(TargetSkeleton, Args)
-				];
-
-			// Blend Profile Mode
 			auto BlendProfileEnabled = [BlendProfileHandle, BlendSettingsEnabledAttribute]()
 			{
 				if (BlendProfileHandle->IsValidHandle())
diff --git a/Engine/Source/Editor/Persona/Private/AnimGraphNodeDetails.cpp b/Engine/Source/Editor/Persona/Private/AnimGraphNodeDetails.cpp
--- a/Engine/Source/Editor/Persona/Private/AnimGraphNodeDetails.cpp
+++ b/Engine/Source/Editor/Persona/Private/AnimGraphNodeDetails.cpp
@@ -254,26 +254,6 @@ TSharedRef<SWidget> FAnimGraphNodeDetails::CreatePropertyWidget(FProperty* Targe
 				.AllowClear(bAllowClear)
 				.OnShouldFilterAsset(FOnShouldFilterAsset::CreateSP(this, &FAnimGraphNodeDetails::OnShouldFilterAnimAsset, NodeClass));
 		}
-		else if(ObjectProperty->PropertyClass->IsChildOf(UBlendProfile::StaticClass()) && TargetSkeleton)
-		{
-			TSharedPtr<IPropertyHandle> PropertyPtr(TargetPropertyHandle);
-
-			UObject* PropertyValue = nullptr;
-			TargetPropertyHandle->GetValue(PropertyValue);
-
-			UBlendProfile* CurrentProfile = Cast<UBlendProfile>(PropertyValue);
-
-			FBlendProfilePickerArgs Args;
-			Args.bAllowNew = false;
-			Args.bAllowRemove = false;
-			Args.bAllowClear = true;
-			Args.OnBlendProfileSelected = FOnBlendProfileSelected::CreateSP(this, &FAnimGraphNodeDetails::OnBlendProfileChanged, PropertyPtr);
-			Args.InitialProfile = CurrentProfile;
-
-			ISkeletonEditorModule& SkeletonEditorModule = FModuleManager::Get().LoadModuleChecked<ISkeletonEditorModule>("SkeletonEditor");
-			return SkeletonEditorModule.CreateBlendProfilePicker(this->TargetSkeleton, Args);
-
-		}
 	}
 
 	return SNullWidget::NullWidget;
diff --git a/Engine/Source/Editor/Persona/Private/BlendProfileCustomization.cpp b/Engine/Source/Editor/Persona/Private/BlendProfileCustomization.cpp
new file mode 100644
--- /dev/null
+++ b/Engine/Source/Editor/Persona/Private/BlendProfileCustomization.cpp
@@ -0,0 +1,89 @@
+// Copyright Epic Games, Inc. All Rights Reserved.
+
+#include "BlendProfileCustomization.h"
+
+#include "Editor.h"
+#include "PropertyHandle.h"
+#include "DetailWidgetRow.h"
+#include "BlendProfilePicker.h"
+#include "Animation/BlendProfile.h"
+#include "ISkeletonEditorModule.h"
+#include "Kismet2/BlueprintEditorUtils.h"
+#include "Animation/AnimBlueprint.h"
+
+#define LOCTEXT_NAMESPACE "BlendProfileCustomization"
+
+TSharedRef<IPropertyTypeCustomization> FBlendProfileCustomization::MakeInstance()
+{
+	return MakeShareable(new FBlendProfileCustomization);
+}
+
+void FBlendProfileCustomization::CustomizeHeader(TSharedRef<class IPropertyHandle> InStructPropertyHandle, class FDetailWidgetRow& HeaderRow, IPropertyTypeCustomizationUtils& StructCustomizationUtils)
+{
+	TSharedRef<SWidget> ValueCustomWidget = SNullWidget::NullWidget;
+
+	TArray<UObject*> OuterObjects;
+	InStructPropertyHandle->GetOuterObjects(OuterObjects);
+
+	if (OuterObjects.Num() > 0)
+	{
+		// try to get skeleton from first outer
+		if (USkeleton* TargetSkeleton = GetSkeletonFromOuter(OuterObjects[0]))
+		{
+			TSharedPtr<IPropertyHandle> PropertyPtr(InStructPropertyHandle);
+
+			UObject* PropertyValue = nullptr;
+			InStructPropertyHandle->GetValue(PropertyValue);
+			UBlendProfile* CurrentProfile = Cast<UBlendProfile>(PropertyValue);
+
+			FBlendProfilePickerArgs Args;
+			Args.bAllowNew = false;
+			Args.bAllowRemove = false;
+			Args.bAllowClear = true;
+			Args.OnBlendProfileSelected = FOnBlendProfileSelected::CreateSP(this, &FBlendProfileCustomization::OnBlendProfileChanged, PropertyPtr);
+			Args.InitialProfile = CurrentProfile;
+
+			ISkeletonEditorModule& SkeletonEditorModule = FModuleManager::Get().LoadModuleChecked<ISkeletonEditorModule>("SkeletonEditor");
+			ValueCustomWidget = SkeletonEditorModule.CreateBlendProfilePicker(TargetSkeleton, Args);
+		}
+	}
+
+	HeaderRow.NameContent()
+		[
+			InStructPropertyHandle->CreatePropertyNameWidget()
+		]
+	.ValueContent()
+		[
+			// If we can't find a skeleton, we can't use custom SWidget. Default to regular one.
+			ValueCustomWidget != SNullWidget::NullWidget ? ValueCustomWidget : InStructPropertyHandle->CreatePropertyValueWidget()
+		];
+}
+
+void FBlendProfileCustomization::OnBlendProfileChanged(UBlendProfile* NewProfile, TSharedPtr<IPropertyHandle> PropertyHandle)
+{
+	if (PropertyHandle.IsValid())
+	{
+		PropertyHandle->SetValue(NewProfile);
+	}
+}
+
+USkeleton* FBlendProfileCustomization::GetSkeletonFromOuter(const UObject* Outer)
+{
+	if (const UEdGraphNode* OuterEdGraphNode = Cast<UEdGraphNode>(Outer))
+	{
+		// If outer is a node, we can get the skeleton from its anim blueprint.
+		if (const UAnimBlueprint* AnimBlueprint = Cast<UAnimBlueprint>(FBlueprintEditorUtils::FindBlueprintForNode(OuterEdGraphNode)))
+		{
+			return AnimBlueprint->TargetSkeleton;
+		}
+	}
+	else if (const UAnimationAsset* OuterAnimAsset = Cast<UAnimationAsset>(Outer))
+	{
+		// If outer is an anim asset, grab the skeleton
+		return OuterAnimAsset->GetSkeleton();
+	}
+
+	return nullptr;
+}
+
+#undef LOCTEXT_NAMESPACE
diff --git a/Engine/Source/Editor/Persona/Private/BlendProfileCustomization.h b/Engine/Source/Editor/Persona/Private/BlendProfileCustomization.h
new file mode 100644
--- /dev/null
+++ b/Engine/Source/Editor/Persona/Private/BlendProfileCustomization.h
@@ -0,0 +1,24 @@
+// Copyright Epic Games, Inc. All Rights Reserved.
+
+#pragma once
+
+#include "CoreMinimal.h"
+#include "Editor/PropertyEditor/Public/IPropertyTypeCustomization.h"
+
+class IPropertyHandle;
+
+class FBlendProfileCustomization : public IPropertyTypeCustomization
+{
+public:
+	static TSharedRef<IPropertyTypeCustomization> MakeInstance();
+
+	virtual void CustomizeHeader(TSharedRef<class IPropertyHandle> StructPropertyHandle, class FDetailWidgetRow& HeaderRow, IPropertyTypeCustomizationUtils& StructCustomizationUtils) override;
+	
+	// Override to do nothing
+	virtual void CustomizeChildren(TSharedRef<IPropertyHandle> InStructPropertyHandle, class IDetailChildrenBuilder& ChildBuilder, IPropertyTypeCustomizationUtils& StructCustomizationUtils) override {}
+
+private:
+
+	void OnBlendProfileChanged(class UBlendProfile* NewProfile, TSharedPtr<IPropertyHandle> PropertyHandle);
+	class USkeleton* GetSkeletonFromOuter(const UObject* Outer);
+};
\ No newline at end of file
diff --git a/Engine/Source/Editor/Persona/Private/PersonaModule.cpp b/Engine/Source/Editor/Persona/Private/PersonaModule.cpp
--- a/Engine/Source/Editor/Persona/Private/PersonaModule.cpp
+++ b/Engine/Source/Editor/Persona/Private/PersonaModule.cpp
@@ -20,6 +20,7 @@
 #include "SkeletalMeshSocketDetails.h"
 #include "AnimNotifyDetails.h"
 #include "AnimGraphNodeDetails.h"
+#include "BlendProfileCustomization.h"
 #include "AnimInstanceDetails.h"
 #include "IEditableSkeleton.h"
 #include "IPersonaToolkit.h"
@@ -142,6 +143,7 @@ void FPersonaModule::StartupModule()
 		PropertyModule.RegisterCustomClassLayout( "AnimInstance", FOnGetDetailCustomizationInstance::CreateStatic(&FAnimInstanceDetails::MakeInstance));
 		PropertyModule.RegisterCustomClassLayout("BlendSpaceBase", FOnGetDetailCustomizationInstance::CreateStatic(&FBlendSpaceDetails::MakeInstance));	
 
+		PropertyModule.RegisterCustomPropertyTypeLayout("BlendProfile", FOnGetPropertyTypeCustomizationInstance::CreateStatic(&FBlendProfileCustomization::MakeInstance));
 		PropertyModule.RegisterCustomPropertyTypeLayout( "InputScaleBias", FOnGetPropertyTypeCustomizationInstance::CreateStatic( &FInputScaleBiasCustomization::MakeInstance ) );
 		PropertyModule.RegisterCustomPropertyTypeLayout( "BoneReference", FOnGetPropertyTypeCustomizationInstance::CreateStatic( &FBoneReferenceCustomization::MakeInstance ) );
 		PropertyModule.RegisterCustomPropertyTypeLayout("BoneSocketTarget", FOnGetPropertyTypeCustomizationInstance::CreateStatic(&FBoneSocketTargetCustomization::MakeInstance));
