Reskin numeric entry boxes and related details customizations (vectors, rotators, transforms)

diff --git a/Engine/Source/Editor/DetailCustomizations/Private/ComponentTransformDetails.cpp b/Engine/Source/Editor/DetailCustomizations/Private/ComponentTransformDetails.cpp
--- a/Engine/Source/Editor/DetailCustomizations/Private/ComponentTransformDetails.cpp
+++ b/Engine/Source/Editor/DetailCustomizations/Private/ComponentTransformDetails.cpp
@@ -162,11 +162,9 @@ TSharedRef<SWidget> FComponentTransformDetails::BuildTransformFieldLabel( ETrans
 
 	MenuBuilder.EndSection();
 
-
-	return 
+	TSharedRef<SWidget> NameContent =
 		SNew(SComboButton)
-		.ContentPadding( 0 )
-		.ButtonStyle( FEditorStyle::Get(), "NoBorder" )
+		.ContentPadding(0)
 		.ForegroundColor( FSlateColor::UseForeground() )
 		.MenuContent()
 		[
@@ -176,12 +174,44 @@ TSharedRef<SWidget> FComponentTransformDetails::BuildTransformFieldLabel( ETrans
 		[
 			SNew( SBox )
 			.Padding( FMargin( 0.0f, 0.0f, 2.0f, 0.0f ) )
+			.MinDesiredWidth(50.f)
 			[
 				SNew(STextBlock)
 				.Text(this, &FComponentTransformDetails::GetTransformFieldText, TransformField)
 				.Font(IDetailLayoutBuilder::GetDetailFont())
 			]
 		];
+
+	if(TransformField == ETransformField::Scale)
+	{
+		NameContent =
+			SNew(SHorizontalBox)
+			+ SHorizontalBox::Slot()
+			.AutoWidth()
+			.VAlign(VAlign_Center)
+			.Padding(FMargin(0.0f, 0.0f, 4.0f, 0.0f))
+			[
+				// Add a checkbox to toggle between preserving the ratio of x,y,z components of scale when a value is entered
+				SNew(SCheckBox)
+				.IsChecked(this, &FComponentTransformDetails::IsPreserveScaleRatioChecked)
+				.IsEnabled(this, &FComponentTransformDetails::GetIsEnabled)
+				.OnCheckStateChanged(this, &FComponentTransformDetails::OnPreserveScaleRatioToggled)
+				.Style(FEditorStyle::Get(), "TransparentCheckBox")
+				.ToolTipText(LOCTEXT("PreserveScaleToolTip", "When locked, scales uniformly based on the current xyz scale values so the object maintains its shape in each direction when scaled"))
+				[
+					SNew(SImage)
+					.Image(this, &FComponentTransformDetails::GetPreserveScaleRatioImage)
+					.ColorAndOpacity(FSlateColor::UseForeground())
+				]
+			]
+			+ SHorizontalBox::Slot()
+			.VAlign(VAlign_Center)
+			[
+				NameContent
+			];
+	}
+
+	return NameContent;
 }
 
 FText FComponentTransformDetails::GetTransformFieldText( ETransformField::Type TransformField ) const
@@ -343,32 +373,26 @@ void FComponentTransformDetails::GenerateChildContent( IDetailChildrenBuilder& C
 		.ValueContent()
 		.MinDesiredWidth(125.0f * 3.0f)
 		.MaxDesiredWidth(125.0f * 3.0f)
+		.VAlign(VAlign_Center)
 		[
-			SNew( SHorizontalBox )
-			+SHorizontalBox::Slot()
-			.FillWidth(1)
-			.VAlign( VAlign_Center )
-			[
-				SNew( SVectorInputBox )
-				.X( this, &FComponentTransformDetails::GetLocationX )
-				.Y( this, &FComponentTransformDetails::GetLocationY )
-				.Z( this, &FComponentTransformDetails::GetLocationZ )
-				.bColorAxisLabels( true )
-				.AllowResponsiveLayout( true )
-				.IsEnabled( this, &FComponentTransformDetails::GetIsEnabled )
-				.OnXChanged( this, &FComponentTransformDetails::OnSetTransformAxis, ETextCommit::Default, ETransformField::Location, EAxisList::X, false )
-				.OnYChanged( this, &FComponentTransformDetails::OnSetTransformAxis, ETextCommit::Default, ETransformField::Location, EAxisList::Y, false )
-				.OnZChanged( this, &FComponentTransformDetails::OnSetTransformAxis, ETextCommit::Default, ETransformField::Location, EAxisList::Z, false )
-				.OnXCommitted( this, &FComponentTransformDetails::OnSetTransformAxis, ETransformField::Location, EAxisList::X, true )
-				.OnYCommitted( this, &FComponentTransformDetails::OnSetTransformAxis, ETransformField::Location, EAxisList::Y, true )
-				.OnZCommitted( this, &FComponentTransformDetails::OnSetTransformAxis, ETransformField::Location, EAxisList::Z, true )
-				.Font( FontInfo )
-				.TypeInterface( TypeInterface )
-				.AllowSpin( SelectedObjects.Num() == 1 )
-				.SpinDelta( 1 )
-				.OnBeginSliderMovement( this, &FComponentTransformDetails::OnBeginLocationSlider )
-				.OnEndSliderMovement( this, &FComponentTransformDetails::OnEndLocationSlider )
-			]
+			SNew(SVectorInputBox)
+			.X(this, &FComponentTransformDetails::GetLocationX)
+			.Y(this, &FComponentTransformDetails::GetLocationY)
+			.Z(this, &FComponentTransformDetails::GetLocationZ)
+			.bColorAxisLabels(true)
+			.IsEnabled(this, &FComponentTransformDetails::GetIsEnabled)
+			.OnXChanged(this, &FComponentTransformDetails::OnSetTransformAxis, ETextCommit::Default, ETransformField::Location, EAxisList::X, false)
+			.OnYChanged(this, &FComponentTransformDetails::OnSetTransformAxis, ETextCommit::Default, ETransformField::Location, EAxisList::Y, false)
+			.OnZChanged(this, &FComponentTransformDetails::OnSetTransformAxis, ETextCommit::Default, ETransformField::Location, EAxisList::Z, false)
+			.OnXCommitted(this, &FComponentTransformDetails::OnSetTransformAxis, ETransformField::Location, EAxisList::X, true)
+			.OnYCommitted(this, &FComponentTransformDetails::OnSetTransformAxis, ETransformField::Location, EAxisList::Y, true)
+			.OnZCommitted(this, &FComponentTransformDetails::OnSetTransformAxis, ETransformField::Location, EAxisList::Z, true)
+			.Font(FontInfo)
+			.TypeInterface(TypeInterface)
+			.AllowSpin(SelectedObjects.Num() == 1)
+			.SpinDelta(1)
+			.OnBeginSliderMovement(this, &FComponentTransformDetails::OnBeginLocationSlider)
+			.OnEndSliderMovement(this, &FComponentTransformDetails::OnEndLocationSlider)
 		];
 	}
 	
@@ -394,31 +418,25 @@ void FComponentTransformDetails::GenerateChildContent( IDetailChildrenBuilder& C
 		.ValueContent()
 		.MinDesiredWidth(125.0f * 3.0f)
 		.MaxDesiredWidth(125.0f * 3.0f)
+		.VAlign(VAlign_Center)
 		[
-			SNew( SHorizontalBox )
-			+SHorizontalBox::Slot()
-			.FillWidth(1)
-			.VAlign( VAlign_Center )
-			[
-				SNew( SRotatorInputBox )
-				.AllowSpin( SelectedObjects.Num() == 1 ) 
-				.Roll( this, &FComponentTransformDetails::GetRotationX )
-				.Pitch( this, &FComponentTransformDetails::GetRotationY )
-				.Yaw( this, &FComponentTransformDetails::GetRotationZ )
-				.AllowResponsiveLayout( true )
-				.bColorAxisLabels( true )
-				.IsEnabled( this, &FComponentTransformDetails::GetIsEnabled )
-				.OnBeginSliderMovement( this, &FComponentTransformDetails::OnBeginRotationSlider )
-				.OnEndSliderMovement( this, &FComponentTransformDetails::OnEndRotationSlider )
-				.OnRollChanged( this, &FComponentTransformDetails::OnSetTransformAxis, ETextCommit::Default, ETransformField::Rotation, EAxisList::X, false )
-				.OnPitchChanged( this, &FComponentTransformDetails::OnSetTransformAxis, ETextCommit::Default, ETransformField::Rotation, EAxisList::Y, false )
-				.OnYawChanged( this, &FComponentTransformDetails::OnSetTransformAxis, ETextCommit::Default, ETransformField::Rotation, EAxisList::Z, false )
-				.OnRollCommitted( this, &FComponentTransformDetails::OnSetTransformAxis, ETransformField::Rotation, EAxisList::X, true )
-				.OnPitchCommitted( this, &FComponentTransformDetails::OnSetTransformAxis, ETransformField::Rotation, EAxisList::Y, true )
-				.OnYawCommitted( this, &FComponentTransformDetails::OnSetTransformAxis, ETransformField::Rotation, EAxisList::Z, true )
-				.TypeInterface( TypeInterface )
-				.Font( FontInfo )
-			]
+			SNew( SRotatorInputBox )
+			.AllowSpin( SelectedObjects.Num() == 1 ) 
+			.Roll( this, &FComponentTransformDetails::GetRotationX )
+			.Pitch( this, &FComponentTransformDetails::GetRotationY )
+			.Yaw( this, &FComponentTransformDetails::GetRotationZ )
+			.bColorAxisLabels( true )
+			.IsEnabled( this, &FComponentTransformDetails::GetIsEnabled )
+			.OnBeginSliderMovement( this, &FComponentTransformDetails::OnBeginRotationSlider )
+			.OnEndSliderMovement( this, &FComponentTransformDetails::OnEndRotationSlider )
+			.OnRollChanged( this, &FComponentTransformDetails::OnSetTransformAxis, ETextCommit::Default, ETransformField::Rotation, EAxisList::X, false )
+			.OnPitchChanged( this, &FComponentTransformDetails::OnSetTransformAxis, ETextCommit::Default, ETransformField::Rotation, EAxisList::Y, false )
+			.OnYawChanged( this, &FComponentTransformDetails::OnSetTransformAxis, ETextCommit::Default, ETransformField::Rotation, EAxisList::Z, false )
+			.OnRollCommitted( this, &FComponentTransformDetails::OnSetTransformAxis, ETransformField::Rotation, EAxisList::X, true )
+			.OnPitchCommitted( this, &FComponentTransformDetails::OnSetTransformAxis, ETransformField::Rotation, EAxisList::Y, true )
+			.OnYawCommitted( this, &FComponentTransformDetails::OnSetTransformAxis, ETransformField::Rotation, EAxisList::Z, true )
+			.TypeInterface( TypeInterface )
+			.Font( FontInfo )
 		];
 	}
 	
@@ -438,51 +456,28 @@ void FComponentTransformDetails::GenerateChildContent( IDetailChildrenBuilder& C
 		.ValueContent()
 		.MinDesiredWidth(125.0f * 3.0f)
 		.MaxDesiredWidth(125.0f * 3.0f)
+		.VAlign(VAlign_Center)
 		[
-			SNew( SHorizontalBox )
-			+SHorizontalBox::Slot()
-			.VAlign( VAlign_Center )
-			.FillWidth(1.0f)
-			[
-				SNew( SVectorInputBox )
-				.X( this, &FComponentTransformDetails::GetScaleX )
-				.Y( this, &FComponentTransformDetails::GetScaleY )
-				.Z( this, &FComponentTransformDetails::GetScaleZ )
-				.bColorAxisLabels( true )
-				.AllowResponsiveLayout( true )
-				.IsEnabled( this, &FComponentTransformDetails::GetIsEnabled )
-				.OnXChanged( this, &FComponentTransformDetails::OnSetTransformAxis, ETextCommit::Default, ETransformField::Scale, EAxisList::X, false )
-				.OnYChanged( this, &FComponentTransformDetails::OnSetTransformAxis, ETextCommit::Default, ETransformField::Scale, EAxisList::Y, false )
-				.OnZChanged( this, &FComponentTransformDetails::OnSetTransformAxis, ETextCommit::Default, ETransformField::Scale, EAxisList::Z, false )
-				.OnXCommitted( this, &FComponentTransformDetails::OnSetTransformAxis, ETransformField::Scale, EAxisList::X, true )
-				.OnYCommitted( this, &FComponentTransformDetails::OnSetTransformAxis, ETransformField::Scale, EAxisList::Y, true )
-				.OnZCommitted( this, &FComponentTransformDetails::OnSetTransformAxis, ETransformField::Scale, EAxisList::Z, true )
-				.ContextMenuExtenderX( this, &FComponentTransformDetails::ExtendXScaleContextMenu )
-				.ContextMenuExtenderY( this, &FComponentTransformDetails::ExtendYScaleContextMenu )
-				.ContextMenuExtenderZ( this, &FComponentTransformDetails::ExtendZScaleContextMenu )
-				.Font( FontInfo )
-				.AllowSpin( SelectedObjects.Num() == 1 )
-				.SpinDelta( 0.0025f )
-				.OnBeginSliderMovement( this, &FComponentTransformDetails::OnBeginScaleSlider )
-				.OnEndSliderMovement( this, &FComponentTransformDetails::OnEndScaleSlider )
-			]
-			+SHorizontalBox::Slot()
-			.AutoWidth()
-			.MaxWidth( 18.0f )
-			[
-				// Add a checkbox to toggle between preserving the ratio of x,y,z components of scale when a value is entered
-				SNew( SCheckBox )
-				.IsChecked( this, &FComponentTransformDetails::IsPreserveScaleRatioChecked )
-				.IsEnabled( this, &FComponentTransformDetails::GetIsEnabled )
-				.OnCheckStateChanged( this, &FComponentTransformDetails::OnPreserveScaleRatioToggled )
-				.Style( FEditorStyle::Get(), "TransparentCheckBox" )
-				.ToolTipText( LOCTEXT("PreserveScaleToolTip", "When locked, scales uniformly based on the current xyz scale values so the object maintains its shape in each direction when scaled" ) )
-				[
-					SNew( SImage )
-					.Image( this, &FComponentTransformDetails::GetPreserveScaleRatioImage )
-					.ColorAndOpacity( FSlateColor::UseForeground() )
-				]
-			]
+			SNew( SVectorInputBox )
+			.X( this, &FComponentTransformDetails::GetScaleX )
+			.Y( this, &FComponentTransformDetails::GetScaleY )
+			.Z( this, &FComponentTransformDetails::GetScaleZ )
+			.bColorAxisLabels( true )
+			.IsEnabled( this, &FComponentTransformDetails::GetIsEnabled )
+			.OnXChanged( this, &FComponentTransformDetails::OnSetTransformAxis, ETextCommit::Default, ETransformField::Scale, EAxisList::X, false )
+			.OnYChanged( this, &FComponentTransformDetails::OnSetTransformAxis, ETextCommit::Default, ETransformField::Scale, EAxisList::Y, false )
+			.OnZChanged( this, &FComponentTransformDetails::OnSetTransformAxis, ETextCommit::Default, ETransformField::Scale, EAxisList::Z, false )
+			.OnXCommitted( this, &FComponentTransformDetails::OnSetTransformAxis, ETransformField::Scale, EAxisList::X, true )
+			.OnYCommitted( this, &FComponentTransformDetails::OnSetTransformAxis, ETransformField::Scale, EAxisList::Y, true )
+			.OnZCommitted( this, &FComponentTransformDetails::OnSetTransformAxis, ETransformField::Scale, EAxisList::Z, true )
+			.ContextMenuExtenderX( this, &FComponentTransformDetails::ExtendXScaleContextMenu )
+			.ContextMenuExtenderY( this, &FComponentTransformDetails::ExtendYScaleContextMenu )
+			.ContextMenuExtenderZ( this, &FComponentTransformDetails::ExtendZScaleContextMenu )
+			.Font( FontInfo )
+			.AllowSpin( SelectedObjects.Num() == 1 )
+			.SpinDelta( 0.0025f )
+			.OnBeginSliderMovement( this, &FComponentTransformDetails::OnBeginScaleSlider )
+			.OnEndSliderMovement(this, &FComponentTransformDetails::OnEndScaleSlider)
 		];
 	}
 }
@@ -556,7 +551,7 @@ bool FComponentTransformDetails::GetIsEnabled() const
 
 const FSlateBrush* FComponentTransformDetails::GetPreserveScaleRatioImage() const
 {
-	return bPreserveScaleRatio ? FEditorStyle::GetBrush( TEXT("GenericLock") ) : FEditorStyle::GetBrush( TEXT("GenericUnlock") ) ;
+	return bPreserveScaleRatio ? FEditorStyle::GetBrush( TEXT("Icons.Lock") ) : FEditorStyle::GetBrush( TEXT("Icons.Unlock") ) ;
 }
 
 ECheckBoxState FComponentTransformDetails::IsPreserveScaleRatioChecked() const
diff --git a/Engine/Source/Editor/DetailCustomizations/Private/MathStructCustomizations.cpp b/Engine/Source/Editor/DetailCustomizations/Private/MathStructCustomizations.cpp
--- a/Engine/Source/Editor/DetailCustomizations/Private/MathStructCustomizations.cpp
+++ b/Engine/Source/Editor/DetailCustomizations/Private/MathStructCustomizations.cpp
@@ -87,12 +87,6 @@ void FMathStructCustomization::MakeHeaderRow(TSharedRef<class IPropertyHandle>&
 		TSharedRef<SWidget> NumericEntryBox = MakeChildWidget(StructPropertyHandle, ChildHandle);
 		NumericEntryBoxWidgetList.Add(NumericEntryBox);
 
-		NumericEntryBox->SetToolTipText(MakeAttributeLambda([StructPropertyHandle] {
-			FText result;
-			StructPropertyHandle->GetValueAsDisplayText(result);
-			return result;
-		}));
-
 		HorizontalBox->AddSlot()
 		.Padding(FMargin(0.0f, 2.0f, bLastChild ? 0.0f : 3.0f, 2.0f))
 		[
@@ -110,6 +104,7 @@ void FMathStructCustomization::MakeHeaderRow(TSharedRef<class IPropertyHandle>&
 		HorizontalBox->AddSlot()
 		.AutoWidth()
 		.MaxWidth(18.0f)
+		.VAlign(VAlign_Center)
 		[
 			// Add a checkbox to toggle between preserving the ratio of x,y,z components of scale when a value is entered
 			SNew(SCheckBox)
@@ -129,7 +124,7 @@ void FMathStructCustomization::MakeHeaderRow(TSharedRef<class IPropertyHandle>&
 
 const FSlateBrush* FMathStructCustomization::GetPreserveScaleRatioImage() const
 {
-	return bPreserveScaleRatio ? FEditorStyle::GetBrush(TEXT("GenericLock")) : FEditorStyle::GetBrush(TEXT("GenericUnlock"));
+	return bPreserveScaleRatio ? FEditorStyle::GetBrush(TEXT("Icons.Lock")) : FEditorStyle::GetBrush(TEXT("Icons.Unlock"));
 }
 
 
@@ -264,36 +259,31 @@ TSharedRef<SWidget> FMathStructCustomization::MakeNumericWidget(
 
 	TWeakPtr<IPropertyHandle> WeakHandlePtr = PropertyHandle;
 
-	return SNew(SNumericEntryBox<NumericType>)
-		.IsEnabled(this, &FMathStructCustomization::IsValueEnabled, WeakHandlePtr)
-		.EditableTextBoxStyle(&FCoreStyle::Get().GetWidgetStyle<FEditableTextBoxStyle>("NormalEditableTextBox"))
-		.Value(this, &FMathStructCustomization::OnGetValue, WeakHandlePtr)
-		.Font(IDetailLayoutBuilder::GetDetailFont())
-		.UndeterminedString(NSLOCTEXT("PropertyEditor", "MultipleValues", "Multiple Values"))
-		.OnValueCommitted(this, &FMathStructCustomization::OnValueCommitted<NumericType>, WeakHandlePtr)
-		.OnValueChanged(this, &FMathStructCustomization::OnValueChanged<NumericType>, WeakHandlePtr)
-		.OnBeginSliderMovement(this, &FMathStructCustomization::OnBeginSliderMovement)
-		.OnEndSliderMovement(this, &FMathStructCustomization::OnEndSliderMovement<NumericType>)
-		.LabelVAlign(VAlign_Center)
-		// Only allow spin on handles with one object.  Otherwise it is not clear what value to spin
-		.AllowSpin(PropertyHandle->GetNumOuterObjects() < 2)
-		.ShiftMouseMovePixelPerDelta(ShiftMouseMovePixelPerDelta)
-		.SupportDynamicSliderMaxValue(SupportDynamicSliderMaxValue)
-		.SupportDynamicSliderMinValue(SupportDynamicSliderMinValue)
-		.OnDynamicSliderMaxValueChanged(this, &FMathStructCustomization::OnDynamicSliderMaxValueChanged<NumericType>)
-		.OnDynamicSliderMinValueChanged(this, &FMathStructCustomization::OnDynamicSliderMinValueChanged<NumericType>)
-		.MinValue(MinValue)
-		.MaxValue(MaxValue)
-		.MinSliderValue(SliderMinValue)
-		.MaxSliderValue(SliderMaxValue)
-		.SliderExponent(SliderExponent)
-		.Delta(Delta)
-		.Label()
-		[
-			SNew(STextBlock)
-				.Font(IDetailLayoutBuilder::GetDetailFont())
-				.Text(PropertyHandle->GetPropertyDisplayName())
-		];
+	return
+		SNew(SNumericEntryBox<NumericType>)
+			.IsEnabled(this, &FMathStructCustomization::IsValueEnabled, WeakHandlePtr)
+			.EditableTextBoxStyle(&FCoreStyle::Get().GetWidgetStyle<FEditableTextBoxStyle>("NormalEditableTextBox"))
+			.Value(this, &FMathStructCustomization::OnGetValue, WeakHandlePtr)
+			.Font(IDetailLayoutBuilder::GetDetailFont())
+			.UndeterminedString(NSLOCTEXT("PropertyEditor", "MultipleValues", "Multiple Values"))
+			.OnValueCommitted(this, &FMathStructCustomization::OnValueCommitted<NumericType>, WeakHandlePtr)
+			.OnValueChanged(this, &FMathStructCustomization::OnValueChanged<NumericType>, WeakHandlePtr)
+			.OnBeginSliderMovement(this, &FMathStructCustomization::OnBeginSliderMovement)
+			.OnEndSliderMovement(this, &FMathStructCustomization::OnEndSliderMovement<NumericType>)
+			// Only allow spin on handles with one object.  Otherwise it is not clear what value to spin
+			.AllowSpin(PropertyHandle->GetNumOuterObjects() < 2)
+			.ShiftMouseMovePixelPerDelta(ShiftMouseMovePixelPerDelta)
+			.SupportDynamicSliderMaxValue(SupportDynamicSliderMaxValue)
+			.SupportDynamicSliderMinValue(SupportDynamicSliderMinValue)
+			.OnDynamicSliderMaxValueChanged(this, &FMathStructCustomization::OnDynamicSliderMaxValueChanged<NumericType>)
+			.OnDynamicSliderMinValueChanged(this, &FMathStructCustomization::OnDynamicSliderMinValueChanged<NumericType>)
+			.MinValue(MinValue)
+			.MaxValue(MaxValue)
+			.MinSliderValue(SliderMinValue)
+			.MaxSliderValue(SliderMaxValue)
+			.SliderExponent(SliderExponent)
+			.Delta(Delta)
+			.ToolTipText(this, &FMathStructCustomization::OnGetValueToolTip<NumericType>, WeakHandlePtr);
 }
 
 template <typename NumericType>
@@ -488,6 +478,20 @@ void FMathStructCustomization::SetValue(NumericType NewValue, EPropertyValueSetF
 	WeakHandlePtr.Pin()->SetValue(NewValue, Flags);
 }
 
+template <typename NumericType>
+FText FMathStructCustomization::OnGetValueToolTip(TWeakPtr<IPropertyHandle> WeakHandlePtr) const
+{
+	if(TSharedPtr<IPropertyHandle> PropertyHandle = WeakHandlePtr.Pin())
+	{
+		TOptional<NumericType> Value = OnGetValue<NumericType>(WeakHandlePtr);
+		if (Value.IsSet())
+		{
+			return FText::Format(LOCTEXT("ValueToolTip", "{0}: {1}"),  PropertyHandle->GetPropertyDisplayName(), FText::AsNumber(Value.GetValue()));
+		}
+	}
+	
+	return FText::GetEmpty();
+}
 
 bool FMathStructCustomization::IsValueEnabled(TWeakPtr<IPropertyHandle> WeakHandlePtr) const
 {
diff --git a/Engine/Source/Editor/DetailCustomizations/Private/MathStructProxyCustomizations.cpp b/Engine/Source/Editor/DetailCustomizations/Private/MathStructProxyCustomizations.cpp
--- a/Engine/Source/Editor/DetailCustomizations/Private/MathStructProxyCustomizations.cpp
+++ b/Engine/Source/Editor/DetailCustomizations/Private/MathStructProxyCustomizations.cpp
@@ -12,6 +12,7 @@
 #include "Widgets/Input/SNumericEntryBox.h"
 #include "HAL/PlatformApplicationMisc.h"
 
+#define LOCTEXT_NAMESPACE "MatrixStructCustomization"
 void FMathStructProxyCustomization::CustomizeChildren( TSharedRef<class IPropertyHandle> StructPropertyHandle, class IDetailChildrenBuilder& StructBuilder, IPropertyTypeCustomizationUtils& StructCustomizationUtils )
 {
 	PropertyUtilities = StructCustomizationUtils.GetPropertyUtilities();
@@ -23,31 +24,31 @@ void FMathStructProxyCustomization::MakeHeaderRow( TSharedRef<class IPropertyHan
 }
 
 template<typename ProxyType, typename NumericType>
-TSharedRef<SWidget> FMathStructProxyCustomization::MakeNumericProxyWidget(TSharedRef<IPropertyHandle>& StructPropertyHandle, TSharedRef< TProxyProperty<ProxyType, NumericType> >& ProxyValue, const FText& Label, bool bRotationInDegrees, const FLinearColor& LabelColor, const FLinearColor& LabelBackgroundColor)
+TSharedRef<SWidget> FMathStructProxyCustomization::MakeNumericProxyWidget(TSharedRef<IPropertyHandle>& StructPropertyHandle, TSharedRef< TProxyProperty<ProxyType, NumericType> >& ProxyValue, const FText& Label, bool bRotationInDegrees, const FLinearColor& LabelBackgroundColor)
 {
 	TWeakPtr<IPropertyHandle> WeakHandlePtr = StructPropertyHandle;
 
-	return 
-		SNew( SNumericEntryBox<NumericType> )
-		.IsEnabled( this, &FMathStructProxyCustomization::IsValueEnabled, WeakHandlePtr )
-		.Value( this, &FMathStructProxyCustomization::OnGetValue<ProxyType, NumericType>, WeakHandlePtr, ProxyValue )
-		.Font( IDetailLayoutBuilder::GetDetailFont() )
-		.UndeterminedString( NSLOCTEXT("PropertyEditor", "MultipleValues", "Multiple Values") )
-		.OnValueCommitted( this, &FMathStructProxyCustomization::OnValueCommitted<ProxyType, NumericType>, WeakHandlePtr, ProxyValue )
-		.OnValueChanged( this, &FMathStructProxyCustomization::OnValueChanged<ProxyType, NumericType>, WeakHandlePtr, ProxyValue )
-		.OnBeginSliderMovement( this, &FMathStructProxyCustomization::OnBeginSliderMovement )
-		.OnEndSliderMovement( this, &FMathStructProxyCustomization::OnEndSliderMovement<ProxyType, NumericType>, WeakHandlePtr, ProxyValue )
-		.LabelVAlign(VAlign_Fill)
-		.LabelPadding(0)
+	return
+		SNew(SNumericEntryBox<NumericType>)
+		.IsEnabled(this, &FMathStructProxyCustomization::IsValueEnabled, WeakHandlePtr)
+		.Value(this, &FMathStructProxyCustomization::OnGetValue<ProxyType, NumericType>, WeakHandlePtr, ProxyValue)
+		.Font(IDetailLayoutBuilder::GetDetailFont())
+		.UndeterminedString(NSLOCTEXT("PropertyEditor", "MultipleValues", "Multiple Values"))
+		.OnValueCommitted(this, &FMathStructProxyCustomization::OnValueCommitted<ProxyType, NumericType>, WeakHandlePtr, ProxyValue)
+		.OnValueChanged(this, &FMathStructProxyCustomization::OnValueChanged<ProxyType, NumericType>, WeakHandlePtr, ProxyValue)
+		.OnBeginSliderMovement(this, &FMathStructProxyCustomization::OnBeginSliderMovement)
+		.OnEndSliderMovement(this, &FMathStructProxyCustomization::OnEndSliderMovement<ProxyType, NumericType>, WeakHandlePtr, ProxyValue)
 		// Only allow spin on handles with one object.  Otherwise it is not clear what value to spin
-		.AllowSpin( StructPropertyHandle->GetNumOuterObjects() == 1 )
+		.AllowSpin(StructPropertyHandle->GetNumOuterObjects() == 1)
 		.MinValue(TOptional<NumericType>())
 		.MaxValue(TOptional<NumericType>())
 		.MaxSliderValue(bRotationInDegrees ? 360.0f : TOptional<NumericType>())
 		.MinSliderValue(bRotationInDegrees ? 0.0f : TOptional<NumericType>())
+		.LabelPadding(FMargin(3))
+		.ToolTipText(this, &FMathStructProxyCustomization::OnGetValueToolTip<ProxyType, NumericType>, WeakHandlePtr, ProxyValue, Label)
 		.Label()
 		[
-			SNumericEntryBox<float>::BuildLabel( Label, LabelColor, LabelBackgroundColor )
+			SNumericEntryBox<NumericType>::BuildNarrowColorLabel(LabelBackgroundColor)
 		];
 }
 
@@ -97,6 +98,18 @@ void FMathStructProxyCustomization::OnEndSliderMovement( NumericType NewValue, T
 }
 
 
+template <typename ProxyType, typename NumericType>
+FText FMathStructProxyCustomization::OnGetValueToolTip(TWeakPtr<IPropertyHandle> WeakHandlePtr, TSharedRef<TProxyProperty<ProxyType, NumericType>> ProxyValue, FText Label) const
+{
+	TOptional<NumericType> Value = OnGetValue<ProxyType, NumericType>(WeakHandlePtr, ProxyValue);
+	if (Value.IsSet())
+	{
+		return FText::Format(LOCTEXT("ValueToolTip", "{0}: {1}"), Label, FText::AsNumber(Value.GetValue()));
+	}
+
+	return FText::GetEmpty();
+}
+
 #define LOCTEXT_NAMESPACE "MatrixStructCustomization"
 
 TSharedRef<IPropertyTypeCustomization> FMatrixStructCustomization::MakeInstance()
@@ -138,17 +151,17 @@ void FMatrixStructCustomization::CustomizeLocation(TSharedRef<class IPropertyHan
 		+ SHorizontalBox::Slot()
 		.Padding(FMargin(0.0f, 2.0f, 3.0f, 2.0f))
 		[
-			MakeNumericProxyWidget<FVector, float>(StructPropertyHandle, CachedTranslationX, LOCTEXT("TranslationX", "X"), false, FLinearColor::White, SNumericEntryBox<float>::RedLabelBackgroundColor)
+			MakeNumericProxyWidget<FVector, float>(StructPropertyHandle, CachedTranslationX, LOCTEXT("TranslationX", "X"), false, SNumericEntryBox<float>::RedLabelBackgroundColor)
 		]
 		+ SHorizontalBox::Slot()
 		.Padding(FMargin(0.0f, 2.0f, 3.0f, 2.0f))
 		[
-			MakeNumericProxyWidget<FVector, float>(StructPropertyHandle, CachedTranslationY, LOCTEXT("TranslationY", "Y"), false, FLinearColor::White, SNumericEntryBox<float>::GreenLabelBackgroundColor)
+			MakeNumericProxyWidget<FVector, float>(StructPropertyHandle, CachedTranslationY, LOCTEXT("TranslationY", "Y"), false, SNumericEntryBox<float>::GreenLabelBackgroundColor)
 		]
 		+ SHorizontalBox::Slot()
 		.Padding(FMargin(0.0f, 2.0f, 0.0f, 2.0f))
 		[
-			MakeNumericProxyWidget<FVector, float>(StructPropertyHandle, CachedTranslationZ, LOCTEXT("TranslationZ", "Z"), false, FLinearColor::White, SNumericEntryBox<float>::BlueLabelBackgroundColor)
+			MakeNumericProxyWidget<FVector, float>(StructPropertyHandle, CachedTranslationZ, LOCTEXT("TranslationZ", "Z"), false, SNumericEntryBox<float>::BlueLabelBackgroundColor)
 		]
 	];
 }
@@ -173,17 +186,17 @@ void FMatrixStructCustomization::CustomizeRotation(TSharedRef<class IPropertyHan
 		+ SHorizontalBox::Slot()
 		.Padding(FMargin(0.0f, 2.0f, 3.0f, 2.0f))
 		[
-			MakeNumericProxyWidget<FRotator, float>(StructPropertyHandle, CachedRotationRoll, LOCTEXT("RotationRoll", "X"), true, FLinearColor::White, SNumericEntryBox<float>::RedLabelBackgroundColor)
+			MakeNumericProxyWidget<FRotator, float>(StructPropertyHandle, CachedRotationRoll, LOCTEXT("RotationRoll", "X"), true, SNumericEntryBox<float>::RedLabelBackgroundColor)
 		]
 		+ SHorizontalBox::Slot()
 		.Padding(FMargin(0.0f, 2.0f, 3.0f, 2.0f))
 		[
-			MakeNumericProxyWidget<FRotator, float>(StructPropertyHandle, CachedRotationPitch, LOCTEXT("RotationPitch", "Y"), true, FLinearColor::White, SNumericEntryBox<float>::GreenLabelBackgroundColor)
+			MakeNumericProxyWidget<FRotator, float>(StructPropertyHandle, CachedRotationPitch, LOCTEXT("RotationPitch", "Y"), true, SNumericEntryBox<float>::GreenLabelBackgroundColor)
 		]
 		+ SHorizontalBox::Slot()
 		.Padding(FMargin(0.0f, 2.0f, 0.0f, 2.0f))
 		[
-			MakeNumericProxyWidget<FRotator, float>(StructPropertyHandle, CachedRotationYaw, LOCTEXT("RotationYaw", "Z"), true, FLinearColor::White, SNumericEntryBox<float>::BlueLabelBackgroundColor)
+			MakeNumericProxyWidget<FRotator, float>(StructPropertyHandle, CachedRotationYaw, LOCTEXT("RotationYaw", "Z"), true, SNumericEntryBox<float>::BlueLabelBackgroundColor)
 		]
 	];
 }
@@ -208,17 +221,17 @@ void FMatrixStructCustomization::CustomizeScale(TSharedRef<class IPropertyHandle
 		+ SHorizontalBox::Slot()
 		.Padding(FMargin(0.0f, 2.0f, 3.0f, 2.0f))
 		[
-			MakeNumericProxyWidget<FVector, float>(StructPropertyHandle, CachedScaleX, LOCTEXT("ScaleX", "X"), false, FLinearColor::White, SNumericEntryBox<float>::RedLabelBackgroundColor)
+			MakeNumericProxyWidget<FVector, float>(StructPropertyHandle, CachedScaleX, LOCTEXT("ScaleX", "X"), false, SNumericEntryBox<float>::RedLabelBackgroundColor)
 		]
 		+ SHorizontalBox::Slot()
 		.Padding(FMargin(0.0f, 2.0f, 3.0f, 2.0f))
 		[
-			MakeNumericProxyWidget<FVector, float>(StructPropertyHandle, CachedScaleY, LOCTEXT("ScaleY", "Y"), false, FLinearColor::White, SNumericEntryBox<float>::GreenLabelBackgroundColor)
+			MakeNumericProxyWidget<FVector, float>(StructPropertyHandle, CachedScaleY, LOCTEXT("ScaleY", "Y"), false, SNumericEntryBox<float>::GreenLabelBackgroundColor)
 		]
 		+ SHorizontalBox::Slot()
 		.Padding(FMargin(0.0f, 2.0f, 0.0f, 2.0f))
 		[
-			MakeNumericProxyWidget<FVector, float>(StructPropertyHandle, CachedScaleZ, LOCTEXT("ScaleZ", "Z"), false, FLinearColor::White, SNumericEntryBox<float>::BlueLabelBackgroundColor)
+			MakeNumericProxyWidget<FVector, float>(StructPropertyHandle, CachedScaleZ, LOCTEXT("ScaleZ", "Z"), false, SNumericEntryBox<float>::BlueLabelBackgroundColor)
 		]
 	];
 }
diff --git a/Engine/Source/Editor/DetailCustomizations/Private/MathStructProxyCustomizations.h b/Engine/Source/Editor/DetailCustomizations/Private/MathStructProxyCustomizations.h
--- a/Engine/Source/Editor/DetailCustomizations/Private/MathStructProxyCustomizations.h
+++ b/Engine/Source/Editor/DetailCustomizations/Private/MathStructProxyCustomizations.h
@@ -176,8 +176,10 @@ protected:
 	 * @return the newly created widget.
 	 */
 	template<typename ProxyType, typename NumericType>
-	TSharedRef<SWidget> MakeNumericProxyWidget(TSharedRef<IPropertyHandle>& StructPropertyHandle, TSharedRef< TProxyProperty<ProxyType, NumericType> >& ProxyValue, const FText& Label, bool bRotationInDegrees = false, const FLinearColor& LabelColor = FCoreStyle::Get().GetColor("DefaultForeground"), const FLinearColor& LabelBackgroundColor = FCoreStyle::Get().GetColor("InvertedForeground"));
+	TSharedRef<SWidget> MakeNumericProxyWidget(TSharedRef<IPropertyHandle>& StructPropertyHandle, TSharedRef< TProxyProperty<ProxyType, NumericType> >& ProxyValue, const FText& Label, bool bRotationInDegrees = false, const FLinearColor& LabelColor = FCoreStyle::Get().GetColor("DefaultForeground"));
 
+	template <typename ProxyType, typename NumericType>
+	FText OnGetValueToolTip(TWeakPtr<IPropertyHandle> WeakHandlePtr, TSharedRef<TProxyProperty<ProxyType, NumericType>> ProxyValue, FText Label) const;
 private:
 	/**
 	 * Gets the value as a float for the provided property handle
diff --git a/Engine/Source/Editor/DetailCustomizations/Private/SplineComponentDetails.cpp b/Engine/Source/Editor/DetailCustomizations/Private/SplineComponentDetails.cpp
--- a/Engine/Source/Editor/DetailCustomizations/Private/SplineComponentDetails.cpp
+++ b/Engine/Source/Editor/DetailCustomizations/Private/SplineComponentDetails.cpp
@@ -406,7 +406,6 @@ void FSplinePointDetails::GenerateChildContent(IDetailChildrenBuilder& ChildrenB
 		.X(this, &FSplinePointDetails::GetPositionX)
 		.Y(this, &FSplinePointDetails::GetPositionY)
 		.Z(this, &FSplinePointDetails::GetPositionZ)
-		.AllowResponsiveLayout(true)
 		.AllowSpin(false)
 		.OnXCommitted(this, &FSplinePointDetails::OnSetPosition, 0)
 		.OnYCommitted(this, &FSplinePointDetails::OnSetPosition, 1)
@@ -433,7 +432,6 @@ void FSplinePointDetails::GenerateChildContent(IDetailChildrenBuilder& ChildrenB
 		.X(this, &FSplinePointDetails::GetArriveTangentX)
 		.Y(this, &FSplinePointDetails::GetArriveTangentY)
 		.Z(this, &FSplinePointDetails::GetArriveTangentZ)
-		.AllowResponsiveLayout(true)
 		.AllowSpin(false)
 		.OnXCommitted(this, &FSplinePointDetails::OnSetArriveTangent, 0)
 		.OnYCommitted(this, &FSplinePointDetails::OnSetArriveTangent, 1)
@@ -460,7 +458,6 @@ void FSplinePointDetails::GenerateChildContent(IDetailChildrenBuilder& ChildrenB
 		.X(this, &FSplinePointDetails::GetLeaveTangentX)
 		.Y(this, &FSplinePointDetails::GetLeaveTangentY)
 		.Z(this, &FSplinePointDetails::GetLeaveTangentZ)
-		.AllowResponsiveLayout(true)
 		.AllowSpin(false)
 		.OnXCommitted(this, &FSplinePointDetails::OnSetLeaveTangent, 0)
 		.OnYCommitted(this, &FSplinePointDetails::OnSetLeaveTangent, 1)
@@ -487,7 +484,6 @@ void FSplinePointDetails::GenerateChildContent(IDetailChildrenBuilder& ChildrenB
 		.X(this, &FSplinePointDetails::GetRotationRoll)
 		.Y(this, &FSplinePointDetails::GetRotationPitch)
 		.Z(this, &FSplinePointDetails::GetRotationYaw)
-		.AllowResponsiveLayout(true)
 		.AllowSpin(false)
 		.OnXCommitted(this, &FSplinePointDetails::OnSetRotation, 0)
 		.OnYCommitted(this, &FSplinePointDetails::OnSetRotation, 1)
@@ -515,7 +511,6 @@ void FSplinePointDetails::GenerateChildContent(IDetailChildrenBuilder& ChildrenB
 		.Y(this, &FSplinePointDetails::GetScaleY)
 		.Z(this, &FSplinePointDetails::GetScaleZ)
 		.AllowSpin(false)
-		.AllowResponsiveLayout(true)
 		.OnXCommitted(this, &FSplinePointDetails::OnSetScale, 0)
 		.OnYCommitted(this, &FSplinePointDetails::OnSetScale, 1)
 		.OnZCommitted(this, &FSplinePointDetails::OnSetScale, 2)
diff --git a/Engine/Source/Editor/DetailCustomizations/Public/Customizations/MathStructCustomizations.h b/Engine/Source/Editor/DetailCustomizations/Public/Customizations/MathStructCustomizations.h
--- a/Engine/Source/Editor/DetailCustomizations/Public/Customizations/MathStructCustomizations.h
+++ b/Engine/Source/Editor/DetailCustomizations/Public/Customizations/MathStructCustomizations.h
@@ -122,6 +122,13 @@ protected:
 	template<typename NumericType>
 	void SetValue(NumericType NewValue, EPropertyValueSetFlags::Type Flags, TWeakPtr<IPropertyHandle> WeakHandlePtr);
 
+	/**
+	 * Gets the tooltip for the value. Displays the property name and the current value
+	 * 
+	 * @praram WeakHandlePtr Handle to the property to get the value from
+	 */
+	template <typename NumericType>
+	FText OnGetValueToolTip(TWeakPtr<IPropertyHandle> WeakHandlePtr) const;
 private:
 
 	/** Gets the brush to use for the lock icon. */
diff --git a/Engine/Source/Editor/LandscapeEditor/Private/LandscapeEditorDetailCustomization_NewLandscape.cpp b/Engine/Source/Editor/LandscapeEditor/Private/LandscapeEditorDetailCustomization_NewLandscape.cpp
--- a/Engine/Source/Editor/LandscapeEditor/Private/LandscapeEditorDetailCustomization_NewLandscape.cpp
+++ b/Engine/Source/Editor/LandscapeEditor/Private/LandscapeEditorDetailCustomization_NewLandscape.cpp
@@ -254,7 +254,6 @@ void FLandscapeEditorDetailCustomization_NewLandscape::CustomizeDetails(IDetailL
 	[
 		SNew(SRotatorInputBox)
 		.bColorAxisLabels(true)
-		.AllowResponsiveLayout(true)
 		.Font(DetailBuilder.GetDetailFont())
 		.Roll_Static(&GetOptionalPropertyValue<float>, PropertyHandle_Rotation_Roll)
 		.Pitch_Static(&GetOptionalPropertyValue<float>, PropertyHandle_Rotation_Pitch)
diff --git a/Engine/Source/Editor/LevelEditor/Private/DlgDeltaTransform.cpp b/Engine/Source/Editor/LevelEditor/Private/DlgDeltaTransform.cpp
--- a/Engine/Source/Editor/LevelEditor/Private/DlgDeltaTransform.cpp
+++ b/Engine/Source/Editor/LevelEditor/Private/DlgDeltaTransform.cpp
@@ -71,7 +71,6 @@ public:
 							.Y( this, &SDlgDeltaTransform::GetDeltaY )
 							.Z( this, &SDlgDeltaTransform::GetDeltaZ )
 							.bColorAxisLabels( true )
-							.AllowResponsiveLayout( true )
 							.AllowSpin( false )
 							.OnXCommitted( this, &SDlgDeltaTransform::OnSetDelta, 0 )
 							.OnYCommitted( this, &SDlgDeltaTransform::OnSetDelta, 1 )
diff --git a/Engine/Source/Editor/StaticMeshEditor/Private/StaticMeshEditorTools.cpp b/Engine/Source/Editor/StaticMeshEditor/Private/StaticMeshEditorTools.cpp
--- a/Engine/Source/Editor/StaticMeshEditor/Private/StaticMeshEditorTools.cpp
+++ b/Engine/Source/Editor/StaticMeshEditor/Private/StaticMeshEditorTools.cpp
@@ -719,7 +719,6 @@ void FMeshBuildSettingsLayout::GenerateChildContent( IDetailChildrenBuilder& Chi
 			.Y(this, &FMeshBuildSettingsLayout::GetBuildScaleY)
 			.Z(this, &FMeshBuildSettingsLayout::GetBuildScaleZ)
 			.bColorAxisLabels(false)
-			.AllowResponsiveLayout(true)
 			.AllowSpin(false)
 			.OnXCommitted(this, &FMeshBuildSettingsLayout::OnBuildScaleXChanged)
 			.OnYCommitted(this, &FMeshBuildSettingsLayout::OnBuildScaleYChanged)
diff --git a/Engine/Source/Editor/WorldBrowser/Private/StreamingLevels/StreamingLevelCustomization.cpp b/Engine/Source/Editor/WorldBrowser/Private/StreamingLevels/StreamingLevelCustomization.cpp
--- a/Engine/Source/Editor/WorldBrowser/Private/StreamingLevels/StreamingLevelCustomization.cpp
+++ b/Engine/Source/Editor/WorldBrowser/Private/StreamingLevels/StreamingLevelCustomization.cpp
@@ -66,7 +66,6 @@ void FStreamingLevelCustomization::CustomizeDetails(IDetailLayoutBuilder& Detail
 			.IsEnabled(this, &FStreamingLevelCustomization::LevelEditTextTransformAllowed)
 			.Font(FEditorStyle::GetFontStyle(TEXT("PropertyWindow.NormalFont")))
 			.bColorAxisLabels(true)
-			.AllowResponsiveLayout(true)
 			.AllowSpin(false)
 			.X(this, &FStreamingLevelCustomization::OnGetLevelPosition, 0)
 			.Y(this, &FStreamingLevelCustomization::OnGetLevelPosition, 1)
diff --git a/Engine/Source/Editor/WorldBrowser/Private/Tiles/STiledLandscapeImportDlg.cpp b/Engine/Source/Editor/WorldBrowser/Private/Tiles/STiledLandscapeImportDlg.cpp
--- a/Engine/Source/Editor/WorldBrowser/Private/Tiles/STiledLandscapeImportDlg.cpp
+++ b/Engine/Source/Editor/WorldBrowser/Private/Tiles/STiledLandscapeImportDlg.cpp
@@ -191,7 +191,6 @@ void STiledLandscapeImportDlg::Construct(const FArguments& InArgs, TSharedPtr<SW
 				[
 					SNew( SVectorInputBox )
 					.bColorAxisLabels( true )
-					.AllowResponsiveLayout( true )
 					.AllowSpin(false)
 					.X( this, &STiledLandscapeImportDlg::GetScaleX )
 					.Y( this, &STiledLandscapeImportDlg::GetScaleY )
diff --git a/Engine/Source/Runtime/Slate/Private/Widgets/Input/SRotatorInputBox.cpp b/Engine/Source/Runtime/Slate/Private/Widgets/Input/SRotatorInputBox.cpp
--- a/Engine/Source/Runtime/Slate/Private/Widgets/Input/SRotatorInputBox.cpp
+++ b/Engine/Source/Runtime/Slate/Private/Widgets/Input/SRotatorInputBox.cpp
@@ -7,36 +7,23 @@
 #include "Widgets/Input/SNumericEntryBox.h"
 #include "Widgets/Layout/SWidgetSwitcher.h"
 
-extern TAutoConsoleVariable<float> CVarCrushThem;
-extern TAutoConsoleVariable<float> CVarStopCrushWhenAbove;
-extern TAutoConsoleVariable<float> CVarStartCrushWhenBelow;
-
 #define LOCTEXT_NAMESPACE "SRotatorInputBox"
 
 void SRotatorInputBox::Construct( const SRotatorInputBox::FArguments& InArgs )
 {
-	bCanBeCrushed = InArgs._AllowResponsiveLayout;
-
-	const FLinearColor LabelColorX = InArgs._bColorAxisLabels ? SNumericEntryBox<float>::RedLabelBackgroundColor : FLinearColor(0.0f, 0.0f, 0.0f, .5f);
-	const FLinearColor LabelColorY = InArgs._bColorAxisLabels ? SNumericEntryBox<float>::GreenLabelBackgroundColor : FLinearColor(0.0f, 0.0f, 0.0f, .5f);
-	const FLinearColor LabelColorZ = InArgs._bColorAxisLabels ? SNumericEntryBox<float>::BlueLabelBackgroundColor : FLinearColor(0.0f, 0.0f, 0.0f, .5f);
-
 	this->ChildSlot
 	[
 		SNew(SHorizontalBox)
 		+SHorizontalBox::Slot()
-		.VAlign( VAlign_Center )
-		.FillWidth( 1.0f )
-		.Padding( 0.0f, 1.0f, 2.0f, 1.0f )
 		[
 			SNew(SNumericEntryBox<float>)
 			.AllowSpin(InArgs._AllowSpin)
 			.MinSliderValue(0.0f)
 			.MaxSliderValue(359.999f)
-			.LabelPadding(0)
+			.LabelPadding(FMargin(3))
 			.Label()
 			[
-				BuildDecoratorLabel(LabelColorX, FLinearColor::White, LOCTEXT("Roll_Label", "X"))
+				InArgs._bColorAxisLabels ? SNumericEntryBox<float>::BuildNarrowColorLabel(SNumericEntryBox<float>::RedLabelBackgroundColor) : SNullWidget::NullWidget
 			]
 			.Font( InArgs._Font )
 			.Value( InArgs._Roll )
@@ -49,24 +36,21 @@ void SRotatorInputBox::Construct( const SRotatorInputBox::FArguments& InArgs )
 			{
 				const TOptional<float>& RollValue = RollAttr.Get();
 				return RollValue.IsSet() 
-					? FText::Format(LOCTEXT("Roll_ToolTip", "Roll Value = {0}"), RollValue.GetValue())
+					? FText::Format(LOCTEXT("Roll_ToolTip", "X(Roll): {0}"), RollValue.GetValue())
 					: LOCTEXT("MultipleValues", "Multiple Values");
 			})
 			.TypeInterface(InArgs._TypeInterface)
 		]
 		+SHorizontalBox::Slot()
-		.VAlign( VAlign_Center )
-		.FillWidth( 1.0f )
-		.Padding( 0.0f, 1.0f, 2.0f, 1.0f )
 		[
 			SNew(SNumericEntryBox<float>)
 			.AllowSpin(InArgs._AllowSpin)
 			.MinSliderValue(0.0f)
 			.MaxSliderValue(359.999f)
-			.LabelPadding(0)
+			.LabelPadding(FMargin(3))
 			.Label()
 			[
-				BuildDecoratorLabel(LabelColorY, FLinearColor::White, LOCTEXT("Pitch_Label", "Y"))
+				InArgs._bColorAxisLabels ? SNumericEntryBox<float>::BuildNarrowColorLabel(SNumericEntryBox<float>::GreenLabelBackgroundColor) : SNullWidget::NullWidget
 			]
 			.Font( InArgs._Font )
 			.Value( InArgs._Pitch )
@@ -79,24 +63,22 @@ void SRotatorInputBox::Construct( const SRotatorInputBox::FArguments& InArgs )
 			{
 				const TOptional<float>& PitchValue = PitchAttr.Get();
 				return PitchValue.IsSet()
-					? FText::Format(LOCTEXT("Pitch_ToolTip", "Pitch Value = {0}"), PitchValue.GetValue())
+					? FText::Format(LOCTEXT("Pitch_ToolTip", "Y(Pitch): {0}"), PitchValue.GetValue())
 					: LOCTEXT("MultipleValues", "Multiple Values");
 			})
 			.TypeInterface(InArgs._TypeInterface)
 		]
 		+SHorizontalBox::Slot()
-		.VAlign( VAlign_Center )
-		.FillWidth( 1.0f )
-		.Padding( 0.0f, 1.0f, 0.0f, 1.0f )
 		[
 			SNew(SNumericEntryBox<float>)
 			.AllowSpin(InArgs._AllowSpin)
 			.MinSliderValue(0.0f)
 			.MaxSliderValue(359.999f)
 			.LabelPadding(0)
+			.LabelPadding(FMargin(3))
 			.Label()
 			[
-				BuildDecoratorLabel(LabelColorZ, FLinearColor::White, LOCTEXT("Yaw_Label", "Z"))
+				InArgs._bColorAxisLabels ? SNumericEntryBox<float>::BuildNarrowColorLabel(SNumericEntryBox<float>::BlueLabelBackgroundColor) : SNullWidget::NullWidget
 			]
 			.Font( InArgs._Font )
 			.Value( InArgs._Yaw )
@@ -109,7 +91,7 @@ void SRotatorInputBox::Construct( const SRotatorInputBox::FArguments& InArgs )
 			{
 				const TOptional<float>& YawValue = YawAttr.Get();
 				return YawValue.IsSet()
-					? FText::Format(LOCTEXT("Yaw_ToolTip", "Yaw Value = {0}"), YawValue.GetValue())
+					? FText::Format(LOCTEXT("Yaw_ToolTip", "Z(Yaw): {0}"), YawValue.GetValue())
 					: LOCTEXT("MultipleValues", "Multiple Values");
 			})
 			.TypeInterface(InArgs._TypeInterface)
@@ -118,73 +100,6 @@ void SRotatorInputBox::Construct( const SRotatorInputBox::FArguments& InArgs )
 
 }
 
-TSharedRef<SWidget> SRotatorInputBox::BuildDecoratorLabel(FLinearColor BackgroundColor, FLinearColor InForegroundColor, FText Label)
-{
-	TSharedRef<SWidget> LabelWidget = SNumericEntryBox<float>::BuildLabel(Label, InForegroundColor, BackgroundColor);
-
-	TSharedRef<SWidget> ResultWidget = LabelWidget;
-
-	if (bCanBeCrushed)
-	{
-		ResultWidget =
-			SNew(SWidgetSwitcher)
-			.WidgetIndex(this, &SRotatorInputBox::GetLabelActiveSlot)
-			+SWidgetSwitcher::Slot()
-			[
-				LabelWidget
-			]
-			+SWidgetSwitcher::Slot()
-			[
-				SNew(SBorder)
-				.BorderImage(FCoreStyle::Get().GetBrush("NumericEntrySpinBox.NarrowDecorator"))
-				.BorderBackgroundColor(BackgroundColor)
-				.ForegroundColor(InForegroundColor)
-				.VAlign(VAlign_Center)
-				.HAlign(HAlign_Left)
-				.Padding(FMargin(5.0f, 0.0f, 0.0f, 0.0f))
-			];
-	}
-
-	return ResultWidget;
-}
-
-int32 SRotatorInputBox::GetLabelActiveSlot() const
-{
-	return bIsBeingCrushed ? 1 : 0;
-}
-
-FMargin SRotatorInputBox::GetTextMargin() const
-{
-	return bIsBeingCrushed ? FMargin(1.0f, 2.0f) : FMargin(4.0f, 2.0f);
-}
-
-void SRotatorInputBox::OnArrangeChildren(const FGeometry& AllottedGeometry, FArrangedChildren& ArrangedChildren) const
-{
-	bool bFoop = bCanBeCrushed && (CVarCrushThem.GetValueOnAnyThread() > 0.0f);
-
-	if (bFoop)
-	{
-		const float AlottedWidth = AllottedGeometry.GetLocalSize().X;
-
-		const float CrushBelow = CVarStartCrushWhenBelow.GetValueOnAnyThread();
-		const float StopCrushing = CVarStopCrushWhenAbove.GetValueOnAnyThread();
-
-		if (bIsBeingCrushed)
-		{
-			bIsBeingCrushed = AlottedWidth < StopCrushing;
-		}
-		else
-		{
-			bIsBeingCrushed = AlottedWidth < CrushBelow;
-		}
-	}
-	else
-	{
-		bIsBeingCrushed = false;
-	}
-
-	SCompoundWidget::OnArrangeChildren(AllottedGeometry, ArrangedChildren);
-}
 
 #undef LOCTEXT_NAMESPACE
 
diff --git a/Engine/Source/Runtime/Slate/Private/Widgets/Input/SVectorInputBox.cpp b/Engine/Source/Runtime/Slate/Private/Widgets/Input/SVectorInputBox.cpp
--- a/Engine/Source/Runtime/Slate/Private/Widgets/Input/SVectorInputBox.cpp
+++ b/Engine/Source/Runtime/Slate/Private/Widgets/Input/SVectorInputBox.cpp
@@ -9,14 +9,8 @@
 
 #define LOCTEXT_NAMESPACE "SVectorInputBox"
 
-TAutoConsoleVariable<float> CVarCrushThem(TEXT("Slate.AllowNumericLabelCrush"), 1.0f, TEXT("Should we crush the vector input box?."));
-TAutoConsoleVariable<float> CVarStopCrushWhenAbove(TEXT("Slate.NumericLabelWidthCrushStop"), 200.0f, TEXT("Stop crushing when the width is above."));
-TAutoConsoleVariable<float> CVarStartCrushWhenBelow(TEXT("Slate.NumericLabelWidthCrushStart"), 190.0f, TEXT("Start crushing when the width is below."));
-
-void SVectorInputBox::Construct( const FArguments& InArgs )
+void SVectorInputBox::Construct(const FArguments& InArgs)
 {
-	bCanBeCrushed = InArgs._AllowResponsiveLayout;
-
 	TSharedRef<SHorizontalBox> HorizontalBox = SNew(SHorizontalBox);
 
 	ChildSlot
@@ -24,46 +18,40 @@ void SVectorInputBox::Construct( const FArguments& InArgs )
 		HorizontalBox
 	];
 
-	ConstructX( InArgs, HorizontalBox );
-	ConstructY( InArgs, HorizontalBox );
-	ConstructZ( InArgs, HorizontalBox );
+	ConstructX(InArgs, HorizontalBox);
+	ConstructY(InArgs, HorizontalBox);
+	ConstructZ(InArgs, HorizontalBox);
 }
 
-void SVectorInputBox::ConstructX( const FArguments& InArgs, TSharedRef<SHorizontalBox> HorizontalBox )
-{	
-	const FLinearColor LabelColor = InArgs._bColorAxisLabels ?  SNumericEntryBox<float>::RedLabelBackgroundColor : FLinearColor(0.0f,0.0f,0.0f,.5f);
-	TSharedRef<SWidget> LabelWidget = BuildDecoratorLabel(LabelColor, FLinearColor::White, LOCTEXT("X_Label", "X"));
-	TAttribute<FMargin> MarginAttribute;
-	if (bCanBeCrushed)
+void SVectorInputBox::ConstructX(const FArguments& InArgs, TSharedRef<SHorizontalBox> HorizontalBox)
+{
+	TSharedRef<SWidget> LabelWidget = SNullWidget::NullWidget;
+	if(InArgs._bColorAxisLabels)
 	{
-		MarginAttribute = TAttribute<FMargin>::Create(TAttribute<FMargin>::FGetter::CreateSP(this, &SVectorInputBox::GetTextMargin));
+		const FLinearColor LabelColor = SNumericEntryBox<float>::RedLabelBackgroundColor;
+		LabelWidget = SNumericEntryBox<float>::BuildNarrowColorLabel(LabelColor);
 	}
 
 	TAttribute<TOptional<float>> Value = InArgs._X;
 
 	HorizontalBox->AddSlot()
-	.VAlign( VAlign_Center )
-	.FillWidth( 1.0f )
-	.Padding( 0.0f, 1.0f, 2.0f, 1.0f )
 	[
-		SNew( SNumericEntryBox<float> )
+		SNew(SNumericEntryBox<float>)
 		.AllowSpin(InArgs._AllowSpin)
-		.Font( InArgs._Font )
-		.Value( InArgs._X )
-		.OnValueChanged( InArgs._OnXChanged )
-		.OnValueCommitted( InArgs._OnXCommitted )
+		.Font(InArgs._Font)
+		.Value(InArgs._X)
+		.OnValueChanged(InArgs._OnXChanged)
+		.OnValueCommitted(InArgs._OnXCommitted)
 		.ToolTipText(MakeAttributeLambda([Value]
 		{
 			if (Value.Get().IsSet())
 			{
-				return FText::Format(LOCTEXT("X_ToolTip", "X Value = {0}"), Value.Get().GetValue());
+				return FText::Format(LOCTEXT("X_ToolTip", "X: {0}"), Value.Get().GetValue());
 			}
 			return LOCTEXT("MultipleValues", "Multiple Values");
 		}))
-		.UndeterminedString( LOCTEXT("MultipleValues", "Multiple Values") )
-		.LabelPadding(0)
-		.OverrideTextMargin(MarginAttribute)
-		.ContextMenuExtender( InArgs._ContextMenuExtenderX )
+		.UndeterminedString(LOCTEXT("MultipleValues", "Multiple Values"))
+		.ContextMenuExtender(InArgs._ContextMenuExtenderX)
 		.TypeInterface(InArgs._TypeInterface)
 		.MinValue(TOptional<float>())
 		.MaxValue(TOptional<float>())
@@ -73,48 +61,43 @@ void SVectorInputBox::ConstructX( const FArguments& InArgs, TSharedRef<SHorizont
 		.Delta(InArgs._SpinDelta)
 		.OnBeginSliderMovement(InArgs._OnBeginSliderMovement)
 		.OnEndSliderMovement(InArgs._OnEndSliderMovement)
+		.LabelPadding(FMargin(3))
 		.Label()
 		[
 			LabelWidget
 		]
 	];
-	
+
 }
 
-void SVectorInputBox::ConstructY( const FArguments& InArgs, TSharedRef<SHorizontalBox> HorizontalBox )
+void SVectorInputBox::ConstructY(const FArguments& InArgs, TSharedRef<SHorizontalBox> HorizontalBox)
 {
-	const FLinearColor LabelColor = InArgs._bColorAxisLabels ?  SNumericEntryBox<float>::GreenLabelBackgroundColor : FLinearColor(0.0f,0.0f,0.0f,.5f);
-	TSharedRef<SWidget> LabelWidget = BuildDecoratorLabel(LabelColor, FLinearColor::White, LOCTEXT("Y_Label", "Y"));
-	TAttribute<FMargin> MarginAttribute;
-	if (bCanBeCrushed)
+	TSharedRef<SWidget> LabelWidget = SNullWidget::NullWidget;
+	if (InArgs._bColorAxisLabels)
 	{
-		MarginAttribute = TAttribute<FMargin>::Create(TAttribute<FMargin>::FGetter::CreateSP(this, &SVectorInputBox::GetTextMargin));
+		const FLinearColor LabelColor = SNumericEntryBox<float>::GreenLabelBackgroundColor;
+		LabelWidget = SNumericEntryBox<float>::BuildNarrowColorLabel(LabelColor);
 	}
 
 	TAttribute<TOptional<float>> Value = InArgs._Y;
 
 	HorizontalBox->AddSlot()
-	.VAlign( VAlign_Center )
-	.FillWidth( 1.0f )
-	.Padding( 0.0f, 1.0f, 2.0f, 1.0f )
 	[
-		SNew( SNumericEntryBox<float> )
+		SNew(SNumericEntryBox<float>)
 		.AllowSpin(InArgs._AllowSpin)
-		.Font( InArgs._Font )
-		.Value( InArgs._Y )
-		.OnValueChanged( InArgs._OnYChanged )
-		.OnValueCommitted( InArgs._OnYCommitted )
+		.Font(InArgs._Font)
+		.Value(InArgs._Y)
+		.OnValueChanged(InArgs._OnYChanged)
+		.OnValueCommitted(InArgs._OnYCommitted)
 		.ToolTipText(MakeAttributeLambda([Value]
 		{
 			if (Value.Get().IsSet())
 			{
-				return FText::Format(LOCTEXT("Y_ToolTip", "Y Value = {0}"), Value.Get().GetValue());
+					return FText::Format(LOCTEXT("Y_ToolTip", "Y: {0}"), Value.Get().GetValue());
 			}
 			return LOCTEXT("MultipleValues", "Multiple Values");
 		}))
-		.UndeterminedString( LOCTEXT("MultipleValues", "Multiple Values") )
-		.LabelPadding(0)
-		.OverrideTextMargin(MarginAttribute)
+		.UndeterminedString(LOCTEXT("MultipleValues", "Multiple Values"))
 		.ContextMenuExtender(InArgs._ContextMenuExtenderY)
 		.TypeInterface(InArgs._TypeInterface)
 		.MinValue(TOptional<float>())
@@ -125,48 +108,43 @@ void SVectorInputBox::ConstructY( const FArguments& InArgs, TSharedRef<SHorizont
 		.Delta(InArgs._SpinDelta)
 		.OnBeginSliderMovement(InArgs._OnBeginSliderMovement)
 		.OnEndSliderMovement(InArgs._OnEndSliderMovement)
+		.LabelPadding(FMargin(3))
 		.Label()
 		[
 			LabelWidget
 		]
-	];
+];
 
 }
 
-void SVectorInputBox::ConstructZ( const FArguments& InArgs, TSharedRef<SHorizontalBox> HorizontalBox )
+void SVectorInputBox::ConstructZ(const FArguments& InArgs, TSharedRef<SHorizontalBox> HorizontalBox)
 {
-	const FLinearColor LabelColor = InArgs._bColorAxisLabels ?  SNumericEntryBox<float>::BlueLabelBackgroundColor : FLinearColor(0.0f,0.0f,0.0f,.5f);
-	TSharedRef<SWidget> LabelWidget = BuildDecoratorLabel(LabelColor, FLinearColor::White, LOCTEXT("Z_Label", "Z"));
-	TAttribute<FMargin> MarginAttribute;
-	if (bCanBeCrushed)
+	TSharedRef<SWidget> LabelWidget = SNullWidget::NullWidget;
+	if (InArgs._bColorAxisLabels)
 	{
-		MarginAttribute = TAttribute<FMargin>::Create(TAttribute<FMargin>::FGetter::CreateSP(this, &SVectorInputBox::GetTextMargin));
+		const FLinearColor LabelColor = SNumericEntryBox<float>::BlueLabelBackgroundColor;
+		LabelWidget = SNumericEntryBox<float>::BuildNarrowColorLabel(LabelColor);
 	}
 
 	TAttribute<TOptional<float>> Value = InArgs._Z;
 
 	HorizontalBox->AddSlot()
-	.VAlign( VAlign_Center )
-	.FillWidth( 1.0f )
-	.Padding( 0.0f, 1.0f, 0.0f, 1.0f )
 	[
-		SNew( SNumericEntryBox<float> )
+		SNew(SNumericEntryBox<float>)
 		.AllowSpin(InArgs._AllowSpin)
-		.Font( InArgs._Font )
-		.Value( InArgs._Z )
-		.OnValueChanged( InArgs._OnZChanged )
-		.OnValueCommitted( InArgs._OnZCommitted )
+		.Font(InArgs._Font)
+		.Value(InArgs._Z)
+		.OnValueChanged(InArgs._OnZChanged)
+		.OnValueCommitted(InArgs._OnZCommitted)
 		.ToolTipText(MakeAttributeLambda([Value]
-		{	
+		{
 			if (Value.Get().IsSet())
 			{
-				return FText::Format(LOCTEXT("Z_ToolTip", "Z Value = {0}"), Value.Get().GetValue());
+				return FText::Format(LOCTEXT("Z_ToolTip", "Z: {0}"), Value.Get().GetValue());
 			}
 			return LOCTEXT("MultipleValues", "Multiple Values");
 		}))
-		.UndeterminedString( LOCTEXT("MultipleValues", "Multiple Values") )
-		.LabelPadding(0)
-		.OverrideTextMargin(MarginAttribute)
+		.UndeterminedString(LOCTEXT("MultipleValues", "Multiple Values"))
 		.ContextMenuExtender(InArgs._ContextMenuExtenderZ)
 		.TypeInterface(InArgs._TypeInterface)
 		.MinValue(TOptional<float>())
@@ -177,6 +155,7 @@ void SVectorInputBox::ConstructZ( const FArguments& InArgs, TSharedRef<SHorizont
 		.Delta(InArgs._SpinDelta)
 		.OnBeginSliderMovement(InArgs._OnBeginSliderMovement)
 		.OnEndSliderMovement(InArgs._OnEndSliderMovement)
+		.LabelPadding(FMargin(3))
 		.Label()
 		[
 			LabelWidget
@@ -184,72 +163,6 @@ void SVectorInputBox::ConstructZ( const FArguments& InArgs, TSharedRef<SHorizont
 	];
 }
 
-TSharedRef<SWidget> SVectorInputBox::BuildDecoratorLabel(FLinearColor BackgroundColor, FLinearColor InForegroundColor, FText Label)
-{
-	TSharedRef<SWidget> LabelWidget = SNumericEntryBox<float>::BuildLabel(Label, InForegroundColor, BackgroundColor);
-
-	TSharedRef<SWidget> ResultWidget = LabelWidget;
-	
-	if (bCanBeCrushed)
-	{
-		ResultWidget =
-			SNew(SWidgetSwitcher)
-			.WidgetIndex(this, &SVectorInputBox::GetLabelActiveSlot)
-			+SWidgetSwitcher::Slot()
-			[
-				LabelWidget
-			]
-			+SWidgetSwitcher::Slot()
-			[
-				SNew(SBorder)
-				.BorderImage(FCoreStyle::Get().GetBrush("NumericEntrySpinBox.NarrowDecorator"))
-				.BorderBackgroundColor(BackgroundColor)
-				.ForegroundColor(InForegroundColor)
-				.VAlign(VAlign_Center)
-				.HAlign(HAlign_Left)
-				.Padding(FMargin(5.0f, 0.0f, 0.0f, 0.0f))
-			];
-	}
-
-	return ResultWidget;
-}
-
-int32 SVectorInputBox::GetLabelActiveSlot() const
-{
-	return bIsBeingCrushed ? 1 : 0;
-}
-
-FMargin SVectorInputBox::GetTextMargin() const
-{
-	return bIsBeingCrushed ? FMargin(1.0f, 2.0f) : FMargin(4.0f, 2.0f);
-}
-
-void SVectorInputBox::OnArrangeChildren(const FGeometry& AllottedGeometry, FArrangedChildren& ArrangedChildren) const
-{
-	bool bFoop = bCanBeCrushed && (CVarCrushThem.GetValueOnAnyThread() > 0.0f);
-
-	if (bFoop)
-	{
-		const float AlottedWidth = AllottedGeometry.GetLocalSize().X;
-
-		const float CrushBelow = CVarStartCrushWhenBelow.GetValueOnAnyThread();
-		const float StopCrushing = CVarStopCrushWhenAbove.GetValueOnAnyThread();
 
-		if (bIsBeingCrushed)
-		{
-			bIsBeingCrushed = AlottedWidth < StopCrushing;
-		}
-		else
-		{
-			bIsBeingCrushed = AlottedWidth < CrushBelow;
-		}
-	}
-	else
-	{
-		bIsBeingCrushed = false;
-	}
-
-	SCompoundWidget::OnArrangeChildren(AllottedGeometry, ArrangedChildren);
-}
 
 #undef LOCTEXT_NAMESPACE
diff --git a/Engine/Source/Runtime/Slate/Public/Widgets/Input/SNumericEntryBox.h b/Engine/Source/Runtime/Slate/Public/Widgets/Input/SNumericEntryBox.h
--- a/Engine/Source/Runtime/Slate/Public/Widgets/Input/SNumericEntryBox.h
+++ b/Engine/Source/Runtime/Slate/Public/Widgets/Input/SNumericEntryBox.h
@@ -201,8 +201,7 @@ public:
 				.OnBeginSliderMovement(InArgs._OnBeginSliderMovement)
 				.OnEndSliderMovement(InArgs._OnEndSliderMovement)
 				.MinDesiredWidth(InArgs._MinDesiredValueWidth)
-				.TypeInterface(Interface)
-				.ToolTipText(this, &SNumericEntryBox<NumericType>::GetValueAsText);
+				.TypeInterface(Interface);
 		}
 
 		// Always create an editable text box.  In the case of an undetermined value being passed in, we cant use the spinbox.
@@ -216,44 +215,41 @@ public:
 			.OnTextCommitted(this, &SNumericEntryBox<NumericType>::OnTextCommitted)
 			.SelectAllTextOnCommit(true)
 			.ContextMenuExtender(InArgs._ContextMenuExtender)
-			.MinDesiredWidth(InArgs._MinDesiredValueWidth)
-			.ToolTipText(this, &SNumericEntryBox<NumericType>::GetValueAsText);
+			.MinDesiredWidth(InArgs._MinDesiredValueWidth);
 
-		TSharedRef<SHorizontalBox> HorizontalBox = SNew(SHorizontalBox);
-	
-		if( InArgs._Label.Widget != SNullWidget::NullWidget )
-		{
-			HorizontalBox->AddSlot()
-				.AutoWidth()
-				.HAlign(HAlign_Left)
-				.VAlign(InArgs._LabelVAlign)
-				.Padding(InArgs._LabelPadding)
-				[
-					InArgs._Label.Widget
-				];
-		}
+		TSharedRef<SOverlay> Overlay = SNew(SOverlay);
 
 		// Add the spin box if we have one
 		if( bAllowSpin )
 		{
-			HorizontalBox->AddSlot()
+			Overlay->AddSlot()
 				.HAlign(HAlign_Fill) 
 				.VAlign(VAlign_Center) 
-				.FillWidth(1)
 				[
 					SpinBox.ToSharedRef()
 				];
 		}
 
-		HorizontalBox->AddSlot()
+		Overlay->AddSlot()
 			.HAlign(HAlign_Fill) 
 			.VAlign(VAlign_Center)
 			.Padding(TextMargin)
-			.FillWidth(1)
 			[
 				EditableText.ToSharedRef()
 			];
 
+
+		if (InArgs._Label.Widget != SNullWidget::NullWidget)
+		{
+			Overlay->AddSlot()
+				.HAlign(HAlign_Left)
+				.VAlign(InArgs._LabelVAlign)
+				.Padding(InArgs._LabelPadding)
+				[
+					InArgs._Label.Widget
+				];
+		}
+
 		ChildSlot
 			[
 				SNew( SBorder )
@@ -262,7 +258,7 @@ public:
 				.ForegroundColor( InArgs._BorderForegroundColor )
 				.Padding(0)
 				[
-					HorizontalBox
+					Overlay
 				]
 			];
 	}
@@ -271,6 +267,7 @@ public:
 	{
 		return
 			SNew(SBorder)
+			.Visibility(EVisibility::HitTestInvisible)
 			.BorderImage(FCoreStyle::Get().GetBrush("NumericEntrySpinBox.Decorator"))
 			.BorderBackgroundColor(BackgroundColor)
 			.ForegroundColor(ForegroundColor)
@@ -283,6 +280,19 @@ public:
 			];
 	}
 
+
+	static TSharedRef<SWidget> BuildNarrowColorLabel(FLinearColor LabelColor)
+	{
+		return
+			SNew(SBorder)
+			.Visibility(EVisibility::HitTestInvisible)
+			.BorderImage(FAppStyle::Get().GetBrush("NumericEntrySpinBox.NarrowDecorator"))
+			.BorderBackgroundColor(LabelColor)
+			.HAlign(HAlign_Left)
+			.Padding(FMargin(2.0f, 0.0f, 0.0f, 0.0f));
+	}
+
+
 	/** Return the internally created SpinBox if bAllowSpin is true */
 	TSharedPtr<SWidget> GetSpinBox() const { return SpinBox; }
 
diff --git a/Engine/Source/Runtime/Slate/Public/Widgets/Input/SRotatorInputBox.h b/Engine/Source/Runtime/Slate/Public/Widgets/Input/SRotatorInputBox.h
--- a/Engine/Source/Runtime/Slate/Public/Widgets/Input/SRotatorInputBox.h
+++ b/Engine/Source/Runtime/Slate/Public/Widgets/Input/SRotatorInputBox.h
@@ -22,8 +22,7 @@ class SLATE_API SRotatorInputBox : public SCompoundWidget
 public:
 	SLATE_BEGIN_ARGS( SRotatorInputBox )
 		: _bColorAxisLabels(false)
-		, _AllowResponsiveLayout(false)
-		, _Font(FCoreStyle::Get().GetFontStyle("NormalFont"))
+		, _Font(FAppStyle::Get().GetFontStyle("NormalFont"))
 		, _AllowSpin(true)
 		{}
 
@@ -39,8 +38,11 @@ public:
 		/** Should the axis labels be colored */
 		SLATE_ARGUMENT( bool, bColorAxisLabels )		
 
-			/** Allow responsive layout to crush the label and margins when there is not a lot of room */
-		SLATE_ARGUMENT(bool, AllowResponsiveLayout)
+		UE_DEPRECATED(5.0, "AllowResponsiveLayout unused as it is no longer necessary.")
+		FArguments& AllowResponsiveLayout(bool bAllow)
+		{
+			return Me();
+		}
 
 		/** Font to use for the text in this box */
 		SLATE_ATTRIBUTE( FSlateFontInfo, Font )
@@ -83,26 +85,4 @@ public:
 	 * @param	InArgs	The declaration data for this widget
 	 */
 	void Construct( const FArguments& InArgs );
-
-
-	// SWidget interface
-	virtual void OnArrangeChildren(const FGeometry& AllottedGeometry, FArrangedChildren& ArrangedChildren) const override;
-	// End of SWidget interface
-
-private:
-	/** Are we allowed to be crushed? */
-	bool bCanBeCrushed;
-
-	/** Are we currently being crushed? */
-	mutable bool bIsBeingCrushed;
-
-private:
-	/** Returns the index of the label widget to use (crushed or uncrushed) */
-	int32 GetLabelActiveSlot() const;
-
-	/** Returns the desired text margin for the label */
-	FMargin GetTextMargin() const;
-
-	/** Creates a decorator label (potentially adding a switcher widget if this is cruhsable) */
-	TSharedRef<SWidget> BuildDecoratorLabel(FLinearColor BackgroundColor, FLinearColor ForegroundColor, FText Label);
 };
diff --git a/Engine/Source/Runtime/Slate/Public/Widgets/Input/SVectorInputBox.h b/Engine/Source/Runtime/Slate/Public/Widgets/Input/SVectorInputBox.h
--- a/Engine/Source/Runtime/Slate/Public/Widgets/Input/SVectorInputBox.h
+++ b/Engine/Source/Runtime/Slate/Public/Widgets/Input/SVectorInputBox.h
@@ -13,8 +13,6 @@
 #include "Framework/MultiBox/MultiBoxExtender.h"
 #include "Widgets/Input/NumericTypeInterface.h"
 
-#undef _X //@todo vogel: needs to be removed once Slate no longer uses _##Name
-
 class FArrangedChildren;
 class SHorizontalBox;
 
@@ -25,13 +23,12 @@ class SLATE_API SVectorInputBox : public SCompoundWidget
 {
 public:
 	SLATE_BEGIN_ARGS( SVectorInputBox )
-		: _Font( FCoreStyle::Get().GetFontStyle("NormalFont") )
+		: _Font( FAppStyle::Get().GetFontStyle("NormalFont") )
 		, _AllowSpin( false )
 		, _SpinDelta( 1 )
 		, _bColorAxisLabels( false )
-		, _AllowResponsiveLayout( false )
 		{}
-
+		
 		/** X Component of the vector */
 		SLATE_ATTRIBUTE( TOptional<float>, X )
 
@@ -54,7 +51,11 @@ public:
 		SLATE_ARGUMENT( bool, bColorAxisLabels )		
 
 		/** Allow responsive layout to crush the label and margins when there is not a lot of room */
-		SLATE_ARGUMENT( bool, AllowResponsiveLayout )
+		UE_DEPRECATED(5.0, "AllowResponsiveLayout unused as it is no longer necessary.")
+		FArguments& AllowResponsiveLayout(bool bAllow)
+		{
+			return Me();
+		}
 
 		/** Called when the x value of the vector is changed */
 		SLATE_EVENT( FOnFloatValueChanged, OnXChanged )
@@ -100,27 +101,8 @@ public:
 	 * @param	InArgs	The declaration data for this widget
 	 */
 	void Construct( const FArguments& InArgs );
-
-	// SWidget interface
-	virtual void OnArrangeChildren(const FGeometry& AllottedGeometry, FArrangedChildren& ArrangedChildren) const override;
-	// End of SWidget interface
-
 private:
-	/** Are we allowed to be crushed? */
-	bool bCanBeCrushed;
-
-	/** Are we currently being crushed? */
-	mutable bool bIsBeingCrushed;
-
-private:
-	/** Returns the index of the label widget to use (crushed or uncrushed) */
-	int32 GetLabelActiveSlot() const;
-
-	/** Returns the desired text margin for the label */
-	FMargin GetTextMargin() const;
 
-	/** Creates a decorator label (potentially adding a switcher widget if this is cruhsable) */
-	TSharedRef<SWidget> BuildDecoratorLabel(FLinearColor BackgroundColor, FLinearColor ForegroundColor, FText Label);
 
 	/**
 	 * Construct widgets for the X Value
diff --git a/Engine/Source/Runtime/SlateCore/Private/Styling/StarshipCoreStyle.cpp b/Engine/Source/Runtime/SlateCore/Private/Styling/StarshipCoreStyle.cpp
--- a/Engine/Source/Runtime/SlateCore/Private/Styling/StarshipCoreStyle.cpp
+++ b/Engine/Source/Runtime/SlateCore/Private/Styling/StarshipCoreStyle.cpp
@@ -499,8 +499,8 @@ TSharedRef<ISlateStyle> FStarshipCoreStyle::Create()
 	// SNumericEntryBox defaults...
 	{
 		Style->Set("NumericEntrySpinBox", FSpinBoxStyle()
-			.SetBackgroundBrush(FSlateRoundedBoxBrush(FStyleColors::Input, InputFocusRadius, FStyleColors::InputOutline, InputFocusThickness))
-			.SetHoveredBackgroundBrush(FSlateRoundedBoxBrush(FStyleColors::Input, InputFocusRadius, FStyleColors::Hover, InputFocusThickness))
+			.SetBackgroundBrush(FSlateNoResource())
+			.SetHoveredBackgroundBrush(FSlateNoResource())
 
 			.SetActiveFillBrush(FSlateRoundedBoxBrush(FStyleColors::Hover, InputFocusRadius, FLinearColor::Transparent, InputFocusThickness))
 			.SetInactiveFillBrush(FSlateRoundedBoxBrush(FStyleColors::Secondary, InputFocusRadius, FLinearColor::Transparent, InputFocusThickness))
@@ -525,7 +525,7 @@ TSharedRef<ISlateStyle> FStarshipCoreStyle::Create()
 
 		Style->Set("NumericEntrySpinBox.Decorator", new BOX_BRUSH("Common/TextBoxLabelBorder", FMargin(5.0f / 16.0f)));
 
-		Style->Set("NumericEntrySpinBox.NarrowDecorator", new BOX_BRUSH("Common/TextBoxLabelBorder", FMargin(2.0f / 16.0f, 4.0f / 16.0f, 2.0f / 16.0f, 4.0f / 16.0f)));
+		Style->Set("NumericEntrySpinBox.NarrowDecorator", new IMAGE_BRUSH_SVG("Starship/CoreWidgets/NumericEntryBox/NarrowDecorator", FVector2D(2,16)));
 	}
 
 	SetupColorPickerStyles(Style);
