Miscellaneous details panel reskin
- Mobility and reset to default

diff --git a/Engine/Source/Editor/DetailCustomizations/Private/MobilityCustomization.cpp b/Engine/Source/Editor/DetailCustomizations/Private/MobilityCustomization.cpp
--- a/Engine/Source/Editor/DetailCustomizations/Private/MobilityCustomization.cpp
+++ b/Engine/Source/Editor/DetailCustomizations/Private/MobilityCustomization.cpp
@@ -2,17 +2,16 @@
 
 #include "Customizations/MobilityCustomization.h"
 #include "Widgets/DeclarativeSyntaxSupport.h"
-#include "Widgets/SBoxPanel.h"
 #include "Styling/SlateTypes.h"
 #include "Widgets/Images/SImage.h"
 #include "Widgets/Text/STextBlock.h"
-#include "Widgets/Layout/SUniformGridPanel.h"
 #include "Widgets/Input/SCheckBox.h"
 #include "EditorStyleSet.h"
 #include "DetailLayoutBuilder.h"
 #include "DetailWidgetRow.h"
 #include "IDetailPropertyRow.h"
 #include "DetailCategoryBuilder.h"
+#include "Widgets/Input/SSegmentedControl.h"
 
 #define LOCTEXT_NAMESPACE "MobilityCustomization"
 
@@ -20,7 +19,10 @@ void FMobilityCustomization::CreateMobilityCustomization(IDetailCategoryBuilder&
 {
 	MobilityHandle = InMobilityHandle;
 
-	TSharedPtr<SUniformGridPanel> ButtonOptionsPanel;
+	TSharedRef<SSegmentedControl<EComponentMobility::Type>> ButtonOptionsPanel =
+		SNew(SSegmentedControl<EComponentMobility::Type>)
+		.Value(this, &FMobilityCustomization::GetActiveMobility)
+		.OnValueChanged(this, &FMobilityCustomization::OnMobilityChanged);
 		
 	IDetailPropertyRow& MobilityRow = Category.AddProperty(MobilityHandle);
 	MobilityRow.CustomWidget()
@@ -34,7 +36,7 @@ void FMobilityCustomization::CreateMobilityCustomization(IDetailCategoryBuilder&
 	.ValueContent()
 	.MaxDesiredWidth(0)
 	[
-		SAssignNew(ButtonOptionsPanel, SUniformGridPanel)
+		ButtonOptionsPanel
 	];
 
 	bool bShowStatic = !( RestrictedMobilityBits & StaticMobilityBitMask );
@@ -49,40 +51,14 @@ void FMobilityCustomization::CreateMobilityCustomization(IDetailCategoryBuilder&
 			: LOCTEXT("Mobility_Static_Tooltip", "A static object can't be changed in game.\n* Allows Baked Lighting\n* Fastest Rendering");
 
 		// Static Mobility
-		ButtonOptionsPanel->AddSlot(0, 0)
+		ButtonOptionsPanel->AddSlot(EComponentMobility::Static)
 		[
-			SNew(SCheckBox)
-			.Style(FEditorStyle::Get(), "Property.ToggleButton.Start")
-			.IsChecked(this, &FMobilityCustomization::IsMobilityActive, EComponentMobility::Static)
-			.OnCheckStateChanged(this, &FMobilityCustomization::OnMobilityChanged, EComponentMobility::Static)
-			.ToolTipText(StaticTooltip)
-			[
-				SNew(SHorizontalBox)
-
-				+ SHorizontalBox::Slot()
-				.AutoWidth()
-				.VAlign(VAlign_Center)
-				.Padding(3, 2)
-				[
-					SNew(SImage)
-					.Image(FEditorStyle::GetBrush("Mobility.Static"))
-				]
-
-				+ SHorizontalBox::Slot()
-				.FillWidth(1.0f)
-				.VAlign(VAlign_Center)
-				.HAlign(HAlign_Center)
-				.Padding(6, 2)
-				[
-					SNew(STextBlock)
-					.Text(LOCTEXT("Static", "Static"))
-					.Font(IDetailLayoutBuilder::GetDetailFont())
-					.ColorAndOpacity(this, &FMobilityCustomization::GetMobilityTextColor, EComponentMobility::Static)
-				]
-			]
-		];
-
-		ColumnIndex++;
+			SNew(STextBlock)
+			.Font(IDetailLayoutBuilder::GetDetailFontBold())
+			.TransformPolicy(ETextTransformPolicy::ToUpper)
+			.Text(LOCTEXT("Static", "Static"))
+		]
+		.ToolTip(StaticTooltip);
 	}
 
 	// Stationary Mobility
@@ -92,40 +68,14 @@ void FMobilityCustomization::CreateMobilityCustomization(IDetailCategoryBuilder&
 			? LOCTEXT("Mobility_Stationary_Tooltip", "A stationary light will only have its shadowing and bounced lighting from static geometry baked by Lightmass, all other lighting will be dynamic.  It can change color and intensity in game.\n* Can't Move\n* Allows Partially Baked Lighting\n* Dynamic Shadows from Movable objects")
 			: LOCTEXT("Mobility_Stationary_Object_Tooltip", "A stationary object can be changed in game but not moved, and enables cached lighting methods. \n* Cached Dynamic Shadows.");
 
-		ButtonOptionsPanel->AddSlot(ColumnIndex, 0)
+		ButtonOptionsPanel->AddSlot(EComponentMobility::Stationary)
 		[
-			SNew(SCheckBox)
-			.IsChecked(this, &FMobilityCustomization::IsMobilityActive, EComponentMobility::Stationary)
-			.Style(FEditorStyle::Get(), ( ColumnIndex == 0 ) ? "Property.ToggleButton.Start" : "Property.ToggleButton.Middle")
-			.OnCheckStateChanged(this, &FMobilityCustomization::OnMobilityChanged, EComponentMobility::Stationary)
-			.ToolTipText(StationaryTooltip)
-			[
-				SNew(SHorizontalBox)
-
-				+ SHorizontalBox::Slot()
-				.AutoWidth()
-				.VAlign(VAlign_Center)
-				.Padding(3, 2)
-				[
-					SNew(SImage)
-					.Image(FEditorStyle::GetBrush("Mobility.Stationary"))
-				]
-
-				+ SHorizontalBox::Slot()
-				.FillWidth(1.0f)
-				.VAlign(VAlign_Center)
-				.HAlign(HAlign_Center)
-				.Padding(6, 2)
-				[
-					SNew(STextBlock)
-					.Text(LOCTEXT("Stationary", "Stationary"))
-					.Font(IDetailLayoutBuilder::GetDetailFont())
-					.ColorAndOpacity(this, &FMobilityCustomization::GetMobilityTextColor, EComponentMobility::Stationary)
-				]
-			]
-		];
-
-		ColumnIndex++;
+			SNew(STextBlock)
+			.Font(IDetailLayoutBuilder::GetDetailFontBold())
+			.TransformPolicy(ETextTransformPolicy::ToUpper)
+			.Text(LOCTEXT("Stationary", "Stationary"))
+		]
+		.ToolTip(StationaryTooltip);
 	}
 
 	FText MovableTooltip = bForLight
@@ -133,51 +83,29 @@ void FMobilityCustomization::CreateMobilityCustomization(IDetailCategoryBuilder&
 			: LOCTEXT("Mobility_Movable_Tooltip", "Movable objects can be moved and changed in game.\n* Totally Dynamic\n* Casts a Dynamic Shadow \n* Slowest Rendering");
 
 	// Movable Mobility
-	ButtonOptionsPanel->AddSlot(ColumnIndex, 0)
+	ButtonOptionsPanel->AddSlot(EComponentMobility::Movable)
 	[
-		SNew(SCheckBox)
-		.IsChecked(this, &FMobilityCustomization::IsMobilityActive, EComponentMobility::Movable)
-		.Style(FEditorStyle::Get(), ( ColumnIndex == 0 ) ? "Property.ToggleButton" : "Property.ToggleButton.End")
-		.OnCheckStateChanged(this, &FMobilityCustomization::OnMobilityChanged, EComponentMobility::Movable)
-		.ToolTipText(MovableTooltip)
-		[
-			SNew(SHorizontalBox)
-
-			+ SHorizontalBox::Slot()
-			.AutoWidth()
-			.VAlign(VAlign_Center)
-			.Padding(3, 2)
-			[
-				SNew(SImage)
-				.Image(FEditorStyle::GetBrush("Mobility.Movable"))
-			]
-
-			+ SHorizontalBox::Slot()
-			.FillWidth(1.0f)
-			.VAlign(VAlign_Center)
-			.HAlign(HAlign_Center)
-			.Padding(6, 2)
-			[
-				SNew(STextBlock)
-				.Text(LOCTEXT("Movable", "Movable"))
-				.Font(IDetailLayoutBuilder::GetDetailFont())
-				.ColorAndOpacity(this, &FMobilityCustomization::GetMobilityTextColor, EComponentMobility::Movable)
-			]
-		]
-	];
+		SNew(STextBlock)
+		.Font(IDetailLayoutBuilder::GetDetailFontBold())
+		.TransformPolicy(ETextTransformPolicy::ToUpper)
+		.Text(LOCTEXT("Movable", "Movable"))
+	]
+	.ToolTip(MovableTooltip);
+
+	ButtonOptionsPanel->RebuildChildren();
 }
 
-ECheckBoxState FMobilityCustomization::IsMobilityActive(EComponentMobility::Type InMobility) const
+EComponentMobility::Type FMobilityCustomization::GetActiveMobility() const
 {
 	if (MobilityHandle.IsValid())
 	{
 		uint8 MobilityByte;
 		MobilityHandle->GetValue(MobilityByte);
 
-		return MobilityByte == InMobility ? ECheckBoxState::Checked : ECheckBoxState::Unchecked;
+		return (EComponentMobility::Type)MobilityByte;
 	}
 
-	return ECheckBoxState::Unchecked;
+	return EComponentMobility::Static;
 }
 
 FSlateColor FMobilityCustomization::GetMobilityTextColor(EComponentMobility::Type InMobility) const
@@ -193,9 +121,9 @@ FSlateColor FMobilityCustomization::GetMobilityTextColor(EComponentMobility::Typ
 	return FSlateColor(FLinearColor(0.72f, 0.72f, 0.72f, 1.f));
 }
 
-void FMobilityCustomization::OnMobilityChanged(ECheckBoxState InCheckedState, EComponentMobility::Type InMobility)
+void FMobilityCustomization::OnMobilityChanged(EComponentMobility::Type InMobility)
 {
-	if (MobilityHandle.IsValid() && InCheckedState == ECheckBoxState::Checked)
+	if (MobilityHandle.IsValid())
 	{
 		MobilityHandle->SetValue((uint8)InMobility);
 	}
diff --git a/Engine/Source/Editor/DetailCustomizations/Public/Customizations/MobilityCustomization.h b/Engine/Source/Editor/DetailCustomizations/Public/Customizations/MobilityCustomization.h
--- a/Engine/Source/Editor/DetailCustomizations/Public/Customizations/MobilityCustomization.h
+++ b/Engine/Source/Editor/DetailCustomizations/Public/Customizations/MobilityCustomization.h
@@ -27,9 +27,9 @@ public:
 
 private:
 
-	ECheckBoxState IsMobilityActive(EComponentMobility::Type InMobility) const;
+	EComponentMobility::Type GetActiveMobility() const;
 	FSlateColor GetMobilityTextColor(EComponentMobility::Type InMobility) const;
-	void OnMobilityChanged(ECheckBoxState InCheckedState, EComponentMobility::Type InMobility);
+	void OnMobilityChanged(EComponentMobility::Type InMobility);
 	FText GetMobilityToolTip() const;
 	
 	TSharedPtr<IPropertyHandle> MobilityHandle;
diff --git a/Engine/Source/Editor/PropertyEditor/Private/UserInterface/PropertyEditor/SResetToDefaultPropertyEditor.cpp b/Engine/Source/Editor/PropertyEditor/Private/UserInterface/PropertyEditor/SResetToDefaultPropertyEditor.cpp
--- a/Engine/Source/Editor/PropertyEditor/Private/UserInterface/PropertyEditor/SResetToDefaultPropertyEditor.cpp
+++ b/Engine/Source/Editor/PropertyEditor/Private/UserInterface/PropertyEditor/SResetToDefaultPropertyEditor.cpp
@@ -31,14 +31,15 @@ void SResetToDefaultPropertyEditor::Construct(const FArguments& InArgs, const TS
 		SNew(SButton)
 		.IsFocusable(false)
 		.ToolTipText(this, &SResetToDefaultPropertyEditor::GetResetToolTip)
-		.ButtonStyle( FEditorStyle::Get(), "NoBorder" )
+		.ButtonStyle(FAppStyle::Get(), "SimpleButton")
 		.ContentPadding(0) 
-		.Visibility( this, &SResetToDefaultPropertyEditor::GetDiffersFromDefaultAsVisibility )
-		.OnClicked( this, &SResetToDefaultPropertyEditor::OnResetClicked )
+		.Visibility(this, &SResetToDefaultPropertyEditor::GetDiffersFromDefaultAsVisibility )
+		.OnClicked(this, &SResetToDefaultPropertyEditor::OnResetClicked)
 		.Content()
 		[
 			SNew(SImage)
-			.Image( FEditorStyle::GetBrush("PropertyWindow.DiffersFromDefault") )
+			.Image(FEditorStyle::GetBrush("PropertyWindow.DiffersFromDefault"))
+			.ColorAndOpacity(FSlateColor::UseForeground())
 		]
 	];
 
diff --git a/Engine/Source/Runtime/SlateCore/Private/Styling/StarshipCoreStyle.cpp b/Engine/Source/Runtime/SlateCore/Private/Styling/StarshipCoreStyle.cpp
--- a/Engine/Source/Runtime/SlateCore/Private/Styling/StarshipCoreStyle.cpp
+++ b/Engine/Source/Runtime/SlateCore/Private/Styling/StarshipCoreStyle.cpp
@@ -1243,7 +1243,10 @@ void FStarshipCoreStyle::SetupCheckboxStyles(TSharedRef<FStyle>& Style)
 		.SetCheckedPressedImage(FSlateNoResource())
 		.SetUndeterminedImage(FSlateNoResource())
 		.SetUndeterminedHoveredImage(FSlateNoResource())
-		.SetUndeterminedPressedImage(FSlateNoResource());
+		.SetUndeterminedPressedImage(FSlateNoResource())
+		.SetForegroundColor(FStyleColors::Foreground)
+		.SetHoveredForegroundColor(FStyleColors::ForegroundHover);
+
 	Style->Set("TransparentCheckBox", BasicTransparentCheckBoxStyle);
 
 	/* Default Style for a toggleable button */
