Mobility property tweaks

diff --git a/Engine/Source/Editor/DetailCustomizations/Private/MobilityCustomization.cpp b/Engine/Source/Editor/DetailCustomizations/Private/MobilityCustomization.cpp
--- a/Engine/Source/Editor/DetailCustomizations/Private/MobilityCustomization.cpp
+++ b/Engine/Source/Editor/DetailCustomizations/Private/MobilityCustomization.cpp
@@ -54,7 +54,7 @@ void FMobilityCustomization::CreateMobilityCustomization(IDetailCategoryBuilder&
 		ButtonOptionsPanel->AddSlot(EComponentMobility::Static)
 		[
 			SNew(STextBlock)
-			.Font(IDetailLayoutBuilder::GetDetailFontBold())
+			.Font(FAppStyle::Get().GetFontStyle("PropertyWindow.MobilityFont"))
 			.TransformPolicy(ETextTransformPolicy::ToUpper)
 			.Text(LOCTEXT("Static", "Static"))
 		]
@@ -71,7 +71,7 @@ void FMobilityCustomization::CreateMobilityCustomization(IDetailCategoryBuilder&
 		ButtonOptionsPanel->AddSlot(EComponentMobility::Stationary)
 		[
 			SNew(STextBlock)
-			.Font(IDetailLayoutBuilder::GetDetailFontBold())
+			.Font(FAppStyle::Get().GetFontStyle("PropertyWindow.MobilityFont"))
 			.TransformPolicy(ETextTransformPolicy::ToUpper)
 			.Text(LOCTEXT("Stationary", "Stationary"))
 		]
@@ -86,7 +86,7 @@ void FMobilityCustomization::CreateMobilityCustomization(IDetailCategoryBuilder&
 	ButtonOptionsPanel->AddSlot(EComponentMobility::Movable)
 	[
 		SNew(STextBlock)
-		.Font(IDetailLayoutBuilder::GetDetailFontBold())
+		.Font(FAppStyle::Get().GetFontStyle("PropertyWindow.MobilityFont"))
 		.TransformPolicy(ETextTransformPolicy::ToUpper)
 		.Text(LOCTEXT("Movable", "Movable"))
 	]
diff --git a/Engine/Source/Editor/EditorStyle/Private/StarshipStyle.cpp b/Engine/Source/Editor/EditorStyle/Private/StarshipStyle.cpp
--- a/Engine/Source/Editor/EditorStyle/Private/StarshipStyle.cpp
+++ b/Engine/Source/Editor/EditorStyle/Private/StarshipStyle.cpp
@@ -3108,6 +3108,11 @@ void FStarshipEditorStyle::FStyle::SetupPropertyEditorStyles()
 		Set( "PropertyWindow.BoldFont",FStyleFonts::Get().SmallBold);
 		Set( "PropertyWindow.ItalicFont", DEFAULT_FONT( "Italic", 8 ) );
 		Set( "PropertyWindow.FilterFont", DEFAULT_FONT( "Regular", 10 ) );
+
+		FSlateFontInfo MobilityFont = FStyleFonts::Get().SmallBold;
+		MobilityFont.LetterSpacing = 200;
+
+		Set("PropertyWindow.MobilityFont", MobilityFont);
 		Set( "PropertyWindow.NoOverlayColor", new FSlateNoResource() );
 		Set( "PropertyWindow.EditConstColor", new FSlateColorBrush( FColor( 152, 152, 152, 80 ) ) );
 		Set( "PropertyWindow.FilteredColor", new FSlateColorBrush( FColor( 0, 255, 0, 80 ) ) );
