Fix no hover cue on transform type drop downs

diff --git a/Engine/Source/Editor/DetailCustomizations/Private/ComponentTransformDetails.cpp b/Engine/Source/Editor/DetailCustomizations/Private/ComponentTransformDetails.cpp
--- a/Engine/Source/Editor/DetailCustomizations/Private/ComponentTransformDetails.cpp
+++ b/Engine/Source/Editor/DetailCustomizations/Private/ComponentTransformDetails.cpp
@@ -165,7 +165,6 @@ TSharedRef<SWidget> FComponentTransformDetails::BuildTransformFieldLabel( ETrans
 	TSharedRef<SWidget> NameContent =
 		SNew(SComboButton)
 		.ContentPadding(0)
-		.ForegroundColor( FSlateColor::UseForeground() )
 		.MenuContent()
 		[
 			MenuBuilder.MakeWidget()
