Fix a number of numeric property label ugliness

diff --git a/Engine/Source/Editor/DetailCustomizations/Private/IntervalStructCustomization.cpp b/Engine/Source/Editor/DetailCustomizations/Private/IntervalStructCustomization.cpp
--- a/Engine/Source/Editor/DetailCustomizations/Private/IntervalStructCustomization.cpp
+++ b/Engine/Source/Editor/DetailCustomizations/Private/IntervalStructCustomization.cpp
@@ -86,67 +86,66 @@ void FIntervalStructCustomization<NumericType>::CustomizeHeader(TSharedRef<IProp
 		StructPropertyHandle->CreatePropertyNameWidget()
 	]
 	.ValueContent()
-	.MinDesiredWidth(200.0f)
-	//.MaxDesiredWidth(200.0f)
+	.MinDesiredWidth(251.0f)
+	.MaxDesiredWidth(251.0f)
 	[
-		SNew(SBox)
-		.Padding(FMargin(0.0f, 3.0f, 0.0f, 2.0f))
+		SNew(SHorizontalBox)
+		+SHorizontalBox::Slot()
+		.VAlign(VAlign_Center)
+		.Padding(FMargin(0.0f, 0.0f, 3.0f, 0.0f))
+		.AutoWidth()
 		[
-			SNew(SHorizontalBox)
-
-			+SHorizontalBox::Slot()
-			.Padding(FMargin(0.0f, 0.0f, 5.0f, 0.0f))
-			.VAlign(VAlign_Center)
-			[
-				SNew(SNumericEntryBox<NumericType>)
-				.Value(this, &FIntervalStructCustomization<NumericType>::OnGetValue, EIntervalField::Min)
-				.MinValue(MinAllowedValue)
-				.MinSliderValue(MinAllowedValue)
-				.MaxValue(this, &FIntervalStructCustomization<NumericType>::OnGetMaxValue)
-				.MaxSliderValue(this, &FIntervalStructCustomization<NumericType>::OnGetMaxValue)
-				.OnValueCommitted(this, &FIntervalStructCustomization<NumericType>::OnValueCommitted, EIntervalField::Min)
-				.OnValueChanged(this, &FIntervalStructCustomization<NumericType>::OnValueChanged, EIntervalField::Min)
-				.OnBeginSliderMovement(this, &FIntervalStructCustomization<NumericType>::OnBeginSliderMovement)
-				.OnEndSliderMovement(this, &FIntervalStructCustomization<NumericType>::OnEndSliderMovement)
-				.UndeterminedString(LOCTEXT("MultipleValues", "Multiple Values"))
-				.Font(IDetailLayoutBuilder::GetDetailFont())
-				.AllowSpin(true)
-				.IsEnabled(this, &FIntervalStructCustomization::IsPropertyEnabled, EIntervalField::Min)
-				.LabelVAlign(VAlign_Center)
-				.Label()
-				[
-					SNew(STextBlock)
-					.Font(IDetailLayoutBuilder::GetDetailFont())
-					.Text(LOCTEXT("MinLabel", "Min"))
-				]
-			]
-
-			+SHorizontalBox::Slot()
-			.Padding(FMargin(0.0f, 0.0f, 5.0f, 0.0f))
-			.VAlign(VAlign_Center)
-			[
-				SNew(SNumericEntryBox<NumericType>)
-				.Value(this, &FIntervalStructCustomization<NumericType>::OnGetValue, EIntervalField::Max)
-				.MinValue(this, &FIntervalStructCustomization<NumericType>::OnGetMinValue)
-				.MinSliderValue(this, &FIntervalStructCustomization<NumericType>::OnGetMinValue)
-				.MaxValue(MaxAllowedValue)
-				.MaxSliderValue(MaxAllowedValue)
-				.OnValueCommitted(this, &FIntervalStructCustomization<NumericType>::OnValueCommitted, EIntervalField::Max)
-				.OnValueChanged(this, &FIntervalStructCustomization<NumericType>::OnValueChanged, EIntervalField::Max)
-				.OnBeginSliderMovement(this, &FIntervalStructCustomization<NumericType>::OnBeginSliderMovement)
-				.OnEndSliderMovement(this, &FIntervalStructCustomization<NumericType>::OnEndSliderMovement)
-				.UndeterminedString(LOCTEXT("MultipleValues", "Multiple Values"))
-				.Font(IDetailLayoutBuilder::GetDetailFont())
-				.AllowSpin(true)
-				.IsEnabled(this, &FIntervalStructCustomization::IsPropertyEnabled, EIntervalField::Max)
-				.LabelVAlign(VAlign_Center)
-				.Label()
-				[
-					SNew(STextBlock)
-					.Font(IDetailLayoutBuilder::GetDetailFont())
-					.Text(LOCTEXT("MaxLabel", "Max"))
-				]
-			]
+			SNew(STextBlock)
+			.Font(IDetailLayoutBuilder::GetDetailFont())
+			.Text(LOCTEXT("MinLabel", "Min"))
+		]
+		+SHorizontalBox::Slot()
+		.Padding(FMargin(0.0f, 0.0f, 3.0f, 0.0f))
+		.VAlign(VAlign_Center)
+		[
+			SNew(SNumericEntryBox<NumericType>)
+			.Value(this, &FIntervalStructCustomization<NumericType>::OnGetValue, EIntervalField::Min)
+			.MinValue(MinAllowedValue)
+			.MinSliderValue(MinAllowedValue)
+			.MaxValue(this, &FIntervalStructCustomization<NumericType>::OnGetMaxValue)
+			.MaxSliderValue(this, &FIntervalStructCustomization<NumericType>::OnGetMaxValue)
+			.OnValueCommitted(this, &FIntervalStructCustomization<NumericType>::OnValueCommitted, EIntervalField::Min)
+			.OnValueChanged(this, &FIntervalStructCustomization<NumericType>::OnValueChanged, EIntervalField::Min)
+			.OnBeginSliderMovement(this, &FIntervalStructCustomization<NumericType>::OnBeginSliderMovement)
+			.OnEndSliderMovement(this, &FIntervalStructCustomization<NumericType>::OnEndSliderMovement)
+			.UndeterminedString(LOCTEXT("MultipleValues", "Multiple Values"))
+			.Font(IDetailLayoutBuilder::GetDetailFont())
+			.AllowSpin(true)
+			.IsEnabled(this, &FIntervalStructCustomization::IsPropertyEnabled, EIntervalField::Min)
+		]
+		+SHorizontalBox::Slot()
+		.VAlign(VAlign_Center)
+		.Padding(2.0f, 0.0f)
+		.AutoWidth()
+		[
+			SNew(STextBlock)
+			.Font(IDetailLayoutBuilder::GetDetailFont())
+			.Text(LOCTEXT("MaxLabel", "Max"))
+		]
+		+SHorizontalBox::Slot()
+		.Padding(FMargin(0.0f, 0.0f, 3.0f, 0.0f))
+		.VAlign(VAlign_Center)
+		[
+			SNew(SNumericEntryBox<NumericType>)
+			.Value(this, &FIntervalStructCustomization<NumericType>::OnGetValue, EIntervalField::Max)
+			.MinValue(this, &FIntervalStructCustomization<NumericType>::OnGetMinValue)
+			.MinSliderValue(this, &FIntervalStructCustomization<NumericType>::OnGetMinValue)
+			.MaxValue(MaxAllowedValue)
+			.MaxSliderValue(MaxAllowedValue)
+			.OnValueCommitted(this, &FIntervalStructCustomization<NumericType>::OnValueCommitted, EIntervalField::Max)
+			.OnValueChanged(this, &FIntervalStructCustomization<NumericType>::OnValueChanged, EIntervalField::Max)
+			.OnBeginSliderMovement(this, &FIntervalStructCustomization<NumericType>::OnBeginSliderMovement)
+			.OnEndSliderMovement(this, &FIntervalStructCustomization<NumericType>::OnEndSliderMovement)
+			.UndeterminedString(LOCTEXT("MultipleValues", "Multiple Values"))
+			.Font(IDetailLayoutBuilder::GetDetailFont())
+			.AllowSpin(true)
+			.IsEnabled(this, &FIntervalStructCustomization::IsPropertyEnabled, EIntervalField::Max)
+			
 		]
 	];
 }
diff --git a/Engine/Source/Editor/DetailCustomizations/Private/MarginCustomization.cpp b/Engine/Source/Editor/DetailCustomizations/Private/MarginCustomization.cpp
--- a/Engine/Source/Editor/DetailCustomizations/Private/MarginCustomization.cpp
+++ b/Engine/Source/Editor/DetailCustomizations/Private/MarginCustomization.cpp
@@ -70,7 +70,7 @@ void FMarginStructCustomization::CustomizeChildren( TSharedRef<class IPropertyHa
 		]
 		.ValueContent()
 		[
-			MakeChildPropertyWidget( ChildIndex, false )
+			MakeChildPropertyWidget( ChildIndex )
 		];
 	}
 }
@@ -89,31 +89,24 @@ TSharedRef<SEditableTextBox> FMarginStructCustomization::MakePropertyWidget()
 	;
 }
 
-TSharedRef<SWidget> FMarginStructCustomization::MakeChildPropertyWidget( int32 PropertyIndex, bool bDisplayLabel ) const
+TSharedRef<SWidget> FMarginStructCustomization::MakeChildPropertyWidget( int32 PropertyIndex) const
 {
 	return
-		SNew( SNumericEntryBox<float> )
-		.Value( this, &FMarginStructCustomization::OnGetValue, PropertyIndex )
-		.Font( IDetailLayoutBuilder::GetDetailFont() )
-		.UndeterminedString( NSLOCTEXT( "PropertyEditor", "MultipleValues", "Multiple Values") )
-		.OnValueCommitted( const_cast<FMarginStructCustomization*>(this), &FMarginStructCustomization::OnValueCommitted, PropertyIndex )
-		.OnValueChanged( const_cast<FMarginStructCustomization*>(this), &FMarginStructCustomization::OnValueChanged, PropertyIndex )
-		.OnBeginSliderMovement( const_cast<FMarginStructCustomization*>(this), &FMarginStructCustomization::OnBeginSliderMovement )
-		.OnEndSliderMovement( const_cast<FMarginStructCustomization*>(this), &FMarginStructCustomization::OnEndSliderMovement )
-		.LabelVAlign( VAlign_Center )
-		.AllowSpin( bIsMarginUsingUVSpace ? true : false )
-		.MinValue( bIsMarginUsingUVSpace ? 0.0f : TNumericLimits<float>::Lowest() )
-		.MaxValue( bIsMarginUsingUVSpace ? 1.0f : TNumericLimits<float>::Max() )
-		.MinSliderValue( bIsMarginUsingUVSpace ? 0.0f : TNumericLimits<float>::Lowest() )
-		.MaxSliderValue( bIsMarginUsingUVSpace ? 1.0f : TNumericLimits<float>::Max()  )
-		.TypeInterface(NumericInterface)
-		.Label()
-		[
-			SNew( STextBlock )
-			.Font( IDetailLayoutBuilder::GetDetailFont() )
-			.Text( ChildPropertyHandles[ PropertyIndex ]->GetPropertyDisplayName() )
-			.Visibility( bDisplayLabel ? EVisibility::Visible : EVisibility::Collapsed )
-		];
+		SNew(SNumericEntryBox<float>)
+		.Value(this, &FMarginStructCustomization::OnGetValue, PropertyIndex)
+		.Font(IDetailLayoutBuilder::GetDetailFont())
+		.UndeterminedString(NSLOCTEXT("PropertyEditor", "MultipleValues", "Multiple Values"))
+		.OnValueCommitted(const_cast<FMarginStructCustomization*>(this), &FMarginStructCustomization::OnValueCommitted, PropertyIndex)
+		.OnValueChanged(const_cast<FMarginStructCustomization*>(this), &FMarginStructCustomization::OnValueChanged, PropertyIndex)
+		.OnBeginSliderMovement(const_cast<FMarginStructCustomization*>(this), &FMarginStructCustomization::OnBeginSliderMovement)
+		.OnEndSliderMovement(const_cast<FMarginStructCustomization*>(this), &FMarginStructCustomization::OnEndSliderMovement)
+		.LabelVAlign(VAlign_Center)
+		.AllowSpin(bIsMarginUsingUVSpace ? true : false)
+		.MinValue(bIsMarginUsingUVSpace ? 0.0f : TNumericLimits<float>::Lowest())
+		.MaxValue(bIsMarginUsingUVSpace ? 1.0f : TNumericLimits<float>::Max())
+		.MinSliderValue(bIsMarginUsingUVSpace ? 0.0f : TNumericLimits<float>::Lowest())
+		.MaxSliderValue(bIsMarginUsingUVSpace ? 1.0f : TNumericLimits<float>::Max())
+		.TypeInterface(NumericInterface);
 }
 
 FText FMarginStructCustomization::GetMarginText() const
diff --git a/Engine/Source/Editor/DetailCustomizations/Private/MarginCustomization.h b/Engine/Source/Editor/DetailCustomizations/Private/MarginCustomization.h
--- a/Engine/Source/Editor/DetailCustomizations/Private/MarginCustomization.h
+++ b/Engine/Source/Editor/DetailCustomizations/Private/MarginCustomization.h
@@ -36,9 +36,8 @@ private:
 	 * Make the widget for the specified property
 	 *
 	 * @param PropertyIndex	The index of the child property to the make the widget for
-	 * @param bDisplayLabel	True if the label is to be visible
 	 */
-	TSharedRef<SWidget> MakeChildPropertyWidget( int32 PropertyIndex, bool bDisplayLabel ) const;
+	TSharedRef<SWidget> MakeChildPropertyWidget( int32 PropertyIndex) const;
 
 	/**
 	 Get the margin property value as text
diff --git a/Engine/Source/Editor/DetailCustomizations/Private/RangeStructCustomization.cpp b/Engine/Source/Editor/DetailCustomizations/Private/RangeStructCustomization.cpp
--- a/Engine/Source/Editor/DetailCustomizations/Private/RangeStructCustomization.cpp
+++ b/Engine/Source/Editor/DetailCustomizations/Private/RangeStructCustomization.cpp
@@ -130,14 +130,21 @@ void FRangeStructCustomization<NumericType>::CustomizeHeader(TSharedRef<IPropert
 	.MaxDesiredWidth(200.0f)
 	[
 		SNew(SVerticalBox)
-
 		+SVerticalBox::Slot()
 		.Padding(FMargin(0.0f, 3.0f, 0.0f, 2.0f))
 		[
 			SNew(SHorizontalBox)
-
 			+SHorizontalBox::Slot()
-			.Padding(FMargin(0.0f, 0.0f, 5.0f, 0.0f))
+			.Padding(FMargin(0.0f, 0.0f, 6.0f, 0.0f))
+			.AutoWidth()
+			.VAlign(VAlign_Center)
+			[
+				SNew(STextBlock)
+				.Font(IDetailLayoutBuilder::GetDetailFont())
+				.Text(LOCTEXT("MinimumBoundLabel", "Min"))
+			]
+			+SHorizontalBox::Slot()
+			.Padding(FMargin(0.0f, 0.0f, 3.0f, 0.0f))
 			.VAlign(VAlign_Center)
 			[
 				SNew(SNumericEntryBox<NumericType>)
@@ -153,20 +160,12 @@ void FRangeStructCustomization<NumericType>::CustomizeHeader(TSharedRef<IPropert
 				.IsEnabled(this, &FRangeStructCustomization<NumericType>::OnQueryIfEnabled, LowerBoundTypeWeakPtr)
 				.Font(IDetailLayoutBuilder::GetDetailFont())
 				.AllowSpin(true)
-				.LabelVAlign(VAlign_Center)
-				.Label()
-				[
-					SNew(STextBlock)
-					.Font(IDetailLayoutBuilder::GetDetailFont())
-					.Text(LOCTEXT("MinimumBoundLabel", "Min"))
-				]
 			]
-
 			+SHorizontalBox::Slot()
 			.VAlign(VAlign_Center)
 			.AutoWidth()
 			[
-				SNew(SComboBox< TSharedPtr<FString> >)
+				SNew(SComboBox<TSharedPtr<FString>>)
 				.OptionsSource(&ComboBoxList)
 				.OnGenerateWidget(this, &FRangeStructCustomization<NumericType>::OnGenerateComboWidget)
 				.OnSelectionChanged(this, &FRangeStructCustomization<NumericType>::OnComboSelectionChanged, LowerBoundTypeWeakPtr)
@@ -182,9 +181,17 @@ void FRangeStructCustomization<NumericType>::CustomizeHeader(TSharedRef<IPropert
 		.Padding(FMargin(0.0f, 2.0f, 0.0f, 3.0f))
 		[
 			SNew(SHorizontalBox)
-
+			+ SHorizontalBox::Slot()
+			.Padding(FMargin(0.0f, 0.0f, 3.0f, 0.0f))
+			.AutoWidth()
+			.VAlign(VAlign_Center)
+			[
+				SNew(STextBlock)
+				.Font(IDetailLayoutBuilder::GetDetailFont())
+				.Text(LOCTEXT("MaximumBoundLabel", "Max"))
+			]
 			+SHorizontalBox::Slot()
-			.Padding(FMargin(0.0f, 0.0f, 5.0f, 0.0f))
+			.Padding(FMargin(0.0f, 0.0f, 3.0f, 0.0f))
 			.VAlign(VAlign_Center)
 			[
 				SNew(SNumericEntryBox<NumericType>)
@@ -200,13 +207,6 @@ void FRangeStructCustomization<NumericType>::CustomizeHeader(TSharedRef<IPropert
 				.IsEnabled(this, &FRangeStructCustomization<NumericType>::OnQueryIfEnabled, UpperBoundTypeWeakPtr)
 				.Font(IDetailLayoutBuilder::GetDetailFont())
 				.AllowSpin(true)
-				.LabelVAlign(VAlign_Center)
-				.Label()
-				[
-					SNew(STextBlock)
-					.Font(IDetailLayoutBuilder::GetDetailFont())
-					.Text(LOCTEXT("MaximumBoundLabel", "Max"))
-				]
 			]
 
 			+SHorizontalBox::Slot()
diff --git a/Engine/Source/Editor/UMGEditor/Private/Designer/SDesignerView.cpp b/Engine/Source/Editor/UMGEditor/Private/Designer/SDesignerView.cpp
--- a/Engine/Source/Editor/UMGEditor/Private/Designer/SDesignerView.cpp
+++ b/Engine/Source/Editor/UMGEditor/Private/Designer/SDesignerView.cpp
@@ -679,7 +679,15 @@ TSharedRef<SWidget> SDesignerView::CreateOverlayUI()
 				.TextStyle(&ToolBarStyle.LabelStyle)
 			]
 		]
-
+		+ SHorizontalBox::Slot()
+		.AutoWidth()
+		.VAlign(VAlign_Center)
+		.Padding(FMargin(2, 0))
+		[
+			SNew(STextBlock)
+			.Visibility(this, &SDesignerView::GetCustomResolutionEntryVisibility)
+			.Text(LOCTEXT("Width", "Width"))
+		]
 		+ SHorizontalBox::Slot()
 		.AutoWidth()
 		.VAlign(VAlign_Center)
@@ -695,14 +703,17 @@ TSharedRef<SWidget> SDesignerView::CreateOverlayUI()
 			.OnValueChanged(this, &SDesignerView::OnCustomResolutionWidthChanged)
 			.Visibility(this, &SDesignerView::GetCustomResolutionEntryVisibility)
 			.MinDesiredValueWidth(50)
-			.LabelPadding(0)
 			.ToolTipText(LOCTEXT("CustomSize_WidthTooltip", "1+\tSets the width of the widget in the designer.\n0\tThe width will match the desired width of the widget."))
-			.Label()
-			[
-				SNumericEntryBox<int32>::BuildLabel(LOCTEXT("Width", "Width"), FLinearColor::White, SNumericEntryBox<int32>::RedLabelBackgroundColor)
-			]
 		]
-
+		+ SHorizontalBox::Slot()
+		.AutoWidth()
+		.VAlign(VAlign_Center)
+		.Padding(FMargin(2, 0))
+		[
+			SNew(STextBlock)
+			.Visibility(this, &SDesignerView::GetCustomResolutionEntryVisibility)
+			.Text(LOCTEXT("Height", "Height"))
+		]
 		+ SHorizontalBox::Slot()
 		.AutoWidth()
 		.VAlign(VAlign_Center)
@@ -719,11 +730,6 @@ TSharedRef<SWidget> SDesignerView::CreateOverlayUI()
 			.Visibility(this, &SDesignerView::GetCustomResolutionEntryVisibility)
 			.MinDesiredValueWidth(50)
 			.ToolTipText(LOCTEXT("CustomSize_HeightTooltip", "1+\tSets the height of the widget in the designer.\n0\tThe height will match the desired height of the widget."))
-			.LabelPadding(0)
-			.Label()
-			[
-				SNumericEntryBox<int32>::BuildLabel(LOCTEXT("Height", "Height"), FLinearColor::White, SNumericEntryBox<int32>::GreenLabelBackgroundColor)
-			]
 		]
 	]
 
