Fixed CVar grouping for non-English by using the enum rather than an FText as the group key

diff --git a/Engine/Source/Editor/DetailCustomizations/Private/DeviceProfileDetails.cpp b/Engine/Source/Editor/DetailCustomizations/Private/DeviceProfileDetails.cpp
--- a/Engine/Source/Editor/DetailCustomizations/Private/DeviceProfileDetails.cpp
+++ b/Engine/Source/Editor/DetailCustomizations/Private/DeviceProfileDetails.cpp
@@ -96,143 +96,107 @@ namespace DeviceProfileCVarFormatHelper
 
 
 	/**
-	 *	Convert the enum to a string version
+	 *	Convert the CVar group enum enum its display name.
 	 *
 	 *	@param CatEnum - The ECVarGroup index
 	 *
-	 *	@return The name of the group.
+	 *	@return The display name of the group.
 	 */
-	FText CategoryTextFromEnum(DeviceProfileCVarFormatHelper::ECVarGroup CatEnum)
+	FText CategoryTextFromEnum(ECVarGroup CatEnum)
 	{
-		FText CategoryText;
-		if (CatEnum == ECVarGroup::CVG_Uncategorized)
-		{
-			CategoryText = LOCTEXT("UncategorizedCVarGroupTitle","Uncategorized");
-		}
-		else if (CatEnum == ECVarGroup::CVG_Rendering)
-		{
-			CategoryText = LOCTEXT("RenderingCVarGroupTitle","Rendering");
-		}
-		else if (CatEnum == ECVarGroup::CVG_Physics)
-		{
-			CategoryText = LOCTEXT("PhysicsCVarGroupTitle","Physics");
-		}
-		else if (CatEnum == ECVarGroup::CVG_Network)
-		{
-			CategoryText = LOCTEXT("NetworkCVarGroupTitle","Network");
-		}
-		else if (CatEnum == ECVarGroup::CVG_Console)
-		{
-			CategoryText = LOCTEXT("ConsoleCVarGroupTitle","Console");
-		}
-		else if (CatEnum == ECVarGroup::CVG_Compatibility)
-		{
-			CategoryText = LOCTEXT("CompatibilityCVarGroupTitle","Compatibility");
-		}
-		else if (CatEnum == ECVarGroup::CVG_UserInterface)
-		{
-			CategoryText = LOCTEXT("UICVarGroupTitle", "User Interface");
-		}
-		else if (CatEnum == ECVarGroup::CVG_ScalabilityGroups)
-		{
-			CategoryText = LOCTEXT("ScalabilityGroupCVarGroupTitle", "Scalability Group");
-		}
-
-		return CategoryText;
+		switch(CatEnum)
+		{
+		case CVG_Rendering:
+			return LOCTEXT("RenderingCVarGroupTitle", "Rendering");
+		case CVG_Physics:
+			return LOCTEXT("PhysicsCVarGroupTitle", "Physics");
+		case CVG_Network:
+			return LOCTEXT("NetworkCVarGroupTitle", "Network");
+		case CVG_Console:
+			return LOCTEXT("ConsoleCVarGroupTitle", "Console");
+		case CVG_Compatibility:
+			return LOCTEXT("CompatibilityCVarGroupTitle", "Compatibility");
+		case CVG_UserInterface:
+			return LOCTEXT("UICVarGroupTitle", "User Interface");
+		case CVG_ScalabilityGroups:
+			return LOCTEXT("ScalabilityGroupCVarGroupTitle", "Scalability Group");
+		default:
+			break;
+		}
+		return LOCTEXT("UncategorizedCVarGroupTitle", "Uncategorized");
 	}
 
 
 	/**
-	*	Convert the Console Variable Category from the CVar prefix
-	*
-	*	@param InPrefix - The Prefix of the Console Variable
-	*
-	*	@return The name of the group.
-	*/
-	FText CategoryTextFromPrefix(const FString& InPrefix)
+	 *	Convert the CVar group enum enum its CVar prefix.
+	 *
+	 *	@param CatEnum - The ECVarGroup index
+	 *
+	 *	@return The prefix of the group.
+	 */
+	FString CategoryPrefixFromEnum(ECVarGroup CatEnum)
 	{
-		FString LowerPrefix = InPrefix.ToLower();
-
-		FText CategoryText;
-		if (LowerPrefix == TEXT("r") || LowerPrefix == TEXT("r."))
-		{
-			CategoryText = LOCTEXT("RenderingCVarGroupTitle", "Rendering");
-		}
-		else if (LowerPrefix == TEXT("p") || LowerPrefix == TEXT("p."))
-		{
-			CategoryText = LOCTEXT("PhysicsCVarGroupTitle", "Physics");
-		}
-		else if (LowerPrefix == TEXT("net") || LowerPrefix == TEXT("net."))
-		{
-			CategoryText = LOCTEXT("NetworkCVarGroupTitle", "Network");
-		}
-		else if (LowerPrefix == TEXT("con") || LowerPrefix == TEXT("con."))
-		{
-			CategoryText = LOCTEXT("ConsoleCVarGroupTitle", "Console");
-		}
-		else if (LowerPrefix == TEXT("compat") || LowerPrefix == TEXT("compat."))
-		{
-			CategoryText = LOCTEXT("CompatibilityCVarGroupTitle", "Compatibility");
-		}
-		else if (LowerPrefix == TEXT("ui") || LowerPrefix == TEXT("ui."))
-		{
-			CategoryText = LOCTEXT("UICVarGroupTitle", "User Interface");
-		}
-		else if (LowerPrefix == TEXT("sg") || LowerPrefix == TEXT("sg."))
-		{
-			CategoryText = LOCTEXT("ScalabilityGroupCVarGroupTitle", "Scalability Group");
-		}
-		else
-		{
-			CategoryText = LOCTEXT("UncategorizedCVarGroupTitle", "Uncategorized");
-		}
-
-		return CategoryText;
+		switch (CatEnum)
+		{
+		case CVG_Rendering:
+			return TEXT("r");
+		case CVG_Physics:
+			return TEXT("p");
+		case CVG_Network:
+			return TEXT("net");
+		case CVG_Console:
+			return TEXT("con");
+		case CVG_Compatibility:
+			return TEXT("compat");
+		case CVG_UserInterface:
+			return TEXT("ui");
+		case CVG_ScalabilityGroups:
+			return TEXT("sg");
+		default:
+			break;
+		}
+		return FString();
 	}
 
 
 	/**
-	*	Convert the Console Variable Category from the CVar prefix
-	*
-	*	@param InPrefix - The Prefix of the Console Variable
-	*
-	*	@return The name of the group.
-	*/
-	FString CVarPrefixFromCategoryString(const FString& CategoryName)
+	 *	Convert the CVar prefix to the CVar group enum entry.
+	 *
+	 *	@param InPrefix - The Prefix of the Console Variable
+	 *
+	 *	@return The enum entry of the group.
+	 */
+	ECVarGroup CategoryEnumFromPrefix(const FString& InPrefix)
 	{
-		FString CVarPrefix;
-
-		FString LowerCategory = CategoryName.ToLower();
-		if (LowerCategory == TEXT("rendering"))
+		if (InPrefix == TEXT("r") || InPrefix == TEXT("r."))
 		{
-			CVarPrefix = TEXT("r");
+			return CVG_Rendering;
 		}
-		else if (LowerCategory == TEXT("physics"))
+		if (InPrefix == TEXT("p") || InPrefix == TEXT("p."))
 		{
-			CVarPrefix = TEXT("p");
+			return CVG_Physics;
 		}
-		else if (LowerCategory == TEXT("network"))
+		if (InPrefix == TEXT("net") || InPrefix == TEXT("net."))
 		{
-			CVarPrefix = TEXT("net");
+			return CVG_Network;
 		}
-		else if (LowerCategory == TEXT("console"))
+		if (InPrefix == TEXT("con") || InPrefix == TEXT("con."))
 		{
-			CVarPrefix = TEXT("con");
+			return CVG_Console;
 		}
-		else if (LowerCategory == TEXT("compatibility"))
+		if (InPrefix == TEXT("compat") || InPrefix == TEXT("compat."))
 		{
-			CVarPrefix = TEXT("compat");
+			return CVG_Compatibility;
 		}
-		else if (LowerCategory == TEXT("user interface"))
+		if (InPrefix == TEXT("ui") || InPrefix == TEXT("ui."))
 		{
-			CVarPrefix = TEXT("ui");
+			return CVG_UserInterface;
 		}
-		else if (LowerCategory == TEXT("scalability group"))
+		if (InPrefix == TEXT("sg") || InPrefix == TEXT("sg."))
 		{
-			CVarPrefix = TEXT("sg");
+			return CVG_ScalabilityGroups;
 		}
-
-		return CVarPrefix;
+		return CVG_Uncategorized;
 	}
 };
 
@@ -342,7 +306,7 @@ void SCVarSelectionPanel::Construct(const FArguments& InArgs, const FString& CVa
 		&FConsoleVariablesAvailableVisitor::OnConsoleVariable,
 		UnprocessedCVars), *CVarPrefix);
 
-	if (DeviceProfileCVarFormatHelper::CategoryTextFromPrefix(CVarPrefix).ToString() == DeviceProfileCVarFormatHelper::CategoryTextFromEnum(DeviceProfileCVarFormatHelper::CVG_Uncategorized).ToString())
+	if (DeviceProfileCVarFormatHelper::CategoryEnumFromPrefix(CVarPrefix) == DeviceProfileCVarFormatHelper::CVG_Uncategorized)
 	{
 		for(TArray<TSharedPtr<FString>>::TIterator CVarIt(UnprocessedCVars); CVarIt; ++CVarIt)
 		{
@@ -619,12 +583,12 @@ void FDeviceProfileConsoleVariablesPropertyDetails::CreateConsoleVariablesProper
 	ensure(CVarsArrayHandle->GetNumElements(CVarCount) == FPropertyAccess::Success);
 
 	// Sort the properties handles into Categories
-	TMap<FString, TArray<TSharedRef<IPropertyHandle>>> CategoryPropertyMap;
+	TMap<DeviceProfileCVarFormatHelper::ECVarGroup, TArray<TSharedRef<IPropertyHandle>>> CategoryPropertyMap;
 
 	// Add all the CVar groups, even if these are empty
 	for (int32 CategoryIdx = 0; CategoryIdx < (int32)DeviceProfileCVarFormatHelper::Max_CVarCategories; CategoryIdx++)
 	{
-		CategoryPropertyMap.FindOrAdd(DeviceProfileCVarFormatHelper::CategoryTextFromEnum((DeviceProfileCVarFormatHelper::ECVarGroup)CategoryIdx).ToString());
+		CategoryPropertyMap.FindOrAdd((DeviceProfileCVarFormatHelper::ECVarGroup)CategoryIdx);
 	}
 
 	for (uint32 CVarPropertyIdx = 0; CVarPropertyIdx < CVarCount; CVarPropertyIdx++)
@@ -639,9 +603,9 @@ void FDeviceProfileConsoleVariablesPropertyDetails::CreateConsoleVariablesProper
 		int32 CategoryEndIndex = CVarName.Find(TEXT("."));
 
 		FString CVarAbrv = CVarName.Left(CategoryEndIndex);
-		FText CVarCategory = DeviceProfileCVarFormatHelper::CategoryTextFromPrefix(CVarAbrv);
+		DeviceProfileCVarFormatHelper::ECVarGroup CVarCategory = DeviceProfileCVarFormatHelper::CategoryEnumFromPrefix(CVarAbrv);
 
-		TArray<TSharedRef<IPropertyHandle>>* CurrentPropertyCategoryGroup = CategoryPropertyMap.Find(CVarCategory.ToString());
+		TArray<TSharedRef<IPropertyHandle>>* CurrentPropertyCategoryGroup = CategoryPropertyMap.Find(CVarCategory);
 		CurrentPropertyCategoryGroup->Add(CVarElementHandle);
 	}
 
@@ -650,9 +614,9 @@ void FDeviceProfileConsoleVariablesPropertyDetails::CreateConsoleVariablesProper
 	for (auto& Current : CategoryPropertyMap)
 	{
 		const TArray<TSharedRef<IPropertyHandle>>& CurrentGroupsProperties = Current.Value;
-		const FText GroupName = FText::FromString(Current.Key);
+		const FText GroupName = DeviceProfileCVarFormatHelper::CategoryTextFromEnum(Current.Key);
 
-		FString CVarPrefix = DeviceProfileCVarFormatHelper::CVarPrefixFromCategoryString(GroupName.ToString());
+		FString CVarPrefix = DeviceProfileCVarFormatHelper::CategoryPrefixFromEnum(Current.Key);
 		if (CVarPrefix.Len() > 0)
 		{
 			CVarPrefix += TEXT(".");
@@ -699,7 +663,7 @@ void FDeviceProfileConsoleVariablesPropertyDetails::CreateConsoleVariablesProper
 				[
 					SNew(SButton)
 					.ButtonStyle(FEditorStyle::Get(), "HoverHintOnly")
-					.OnClicked(this, &FDeviceProfileConsoleVariablesPropertyDetails::OnRemoveAllFromGroup, GroupName)
+					.OnClicked(this, &FDeviceProfileConsoleVariablesPropertyDetails::OnRemoveAllFromGroup, (int32)Current.Key)
 					.ContentPadding(4.0f)
 					.ForegroundColor(FSlateColor::UseForeground())
 					.IsFocusable(false)
@@ -837,14 +801,14 @@ FReply FDeviceProfileConsoleVariablesPropertyDetails::OnRemoveCVarProperty(TShar
 }
 
 
-FReply FDeviceProfileConsoleVariablesPropertyDetails::OnRemoveAllFromGroup(FText GroupName)
+FReply FDeviceProfileConsoleVariablesPropertyDetails::OnRemoveAllFromGroup(int32 CVarCategory)
 {
 	const TSharedPtr<IPropertyHandleArray> CVarsArrayHandle = CVarsHandle->AsArray();
 
 	uint32 CVarCount = 0;
 	ensure(CVarsArrayHandle->GetNumElements(CVarCount) == FPropertyAccess::Success);
 
-	FString CVarPrefix = DeviceProfileCVarFormatHelper::CVarPrefixFromCategoryString(GroupName.ToString());
+	FString CVarPrefix = DeviceProfileCVarFormatHelper::CategoryPrefixFromEnum((DeviceProfileCVarFormatHelper::ECVarGroup)CVarCategory);
 
 	for (int32 CVarPropertyIdx = CVarCount-1; CVarPropertyIdx >= 0; CVarPropertyIdx--)
 	{
diff --git a/Engine/Source/Editor/DetailCustomizations/Private/DeviceProfileDetails.h b/Engine/Source/Editor/DetailCustomizations/Private/DeviceProfileDetails.h
--- a/Engine/Source/Editor/DetailCustomizations/Private/DeviceProfileDetails.h
+++ b/Engine/Source/Editor/DetailCustomizations/Private/DeviceProfileDetails.h
@@ -142,10 +142,10 @@ private:
 	/**
 	 * Action when a group of CVars have been removed from the device profile
 	 *
-	 * @param GroupName - The name of the group all Cvars have been removed for.
+	 * @param CVarCategory - The category of the group all CVars have been removed for (@see ECVarGroup).
 	 * @return Whether the event was handled
 	 */
-	FReply OnRemoveAllFromGroup(FText GroupName);
+	FReply OnRemoveAllFromGroup(int32 CVarCategory);
 
 
 	/**
