Rework per-platform properties to make better use of the space in the details panel. Now each override is a separate tree node and all overrides can be viewed by expanding the property row

diff --git a/Engine/Source/Editor/DetailCustomizations/Private/PerPlatformPropertyCustomization.cpp b/Engine/Source/Editor/DetailCustomizations/Private/PerPlatformPropertyCustomization.cpp
--- a/Engine/Source/Editor/DetailCustomizations/Private/PerPlatformPropertyCustomization.cpp
+++ b/Engine/Source/Editor/DetailCustomizations/Private/PerPlatformPropertyCustomization.cpp
@@ -13,37 +13,35 @@
 #include "Editor.h"
 #include "PropertyHandle.h"
 #include "DetailLayoutBuilder.h"
-#include "SPerPlatformPropertiesWidget.h"
 #include "PlatformInfo.h"
 #include "ScopedTransaction.h"
 #include "IPropertyUtilities.h"
 #include "UObject/MetaData.h"
+#include "Framework/MultiBox/MultiBoxBuilder.h"
+#include "SPerPlatformPropertiesWidget.h"
 
 #define LOCTEXT_NAMESPACE "PerPlatformPropertyCustomization"
 
 template<typename PerPlatformType>
-void FPerPlatformPropertyCustomization<PerPlatformType>::CustomizeHeader(TSharedRef<IPropertyHandle> StructPropertyHandle, FDetailWidgetRow& HeaderRow, IPropertyTypeCustomizationUtils& StructCustomizationUtils)
+void FPerPlatformPropertyCustomization<PerPlatformType>::CustomizeChildren(TSharedRef<IPropertyHandle> StructPropertyHandle, IDetailChildrenBuilder& StructBuilder, IPropertyTypeCustomizationUtils& StructCustomizationUtils)
 {
 	PropertyUtilities = StructCustomizationUtils.GetPropertyUtilities();
 
 	int32 PlatformNumber = PlatformInfo::GetAllPlatformGroupNames().Num();
 
-	HeaderRow.NameContent()
-	[
-		StructPropertyHandle->CreatePropertyNameWidget()
-	]
-	.ValueContent()
-	.MinDesiredWidth(CalcDesiredWidth(StructPropertyHandle))
-	.MaxDesiredWidth((float)(PlatformNumber + 1)*125.0f)
-	[
-		SNew(SPerPlatformPropertiesWidget)
-		.OnGenerateWidget(this, &FPerPlatformPropertyCustomization<PerPlatformType>::GetWidget, StructPropertyHandle)
-		.OnAddPlatform(this, &FPerPlatformPropertyCustomization<PerPlatformType>::AddPlatformOverride, StructPropertyHandle)
-		.OnRemovePlatform(this, &FPerPlatformPropertyCustomization<PerPlatformType>::RemovePlatformOverride, StructPropertyHandle)
-		.PlatformOverrideNames(this, &FPerPlatformPropertyCustomization<PerPlatformType>::GetPlatformOverrideNames, StructPropertyHandle)
-	];
+	TAttribute<TArray<FName>> PlatformOverrideNames = TAttribute<TArray<FName>>::Create(TAttribute<TArray<FName>>::FGetter::CreateSP(this, &FPerPlatformPropertyCustomization<PerPlatformType>::GetPlatformOverrideNames, StructPropertyHandle));
+
+	FPerPlatformPropertyCustomNodeBuilderArgs Args;
+	Args.NameWidget = StructPropertyHandle->CreatePropertyNameWidget();
+	Args.PlatformOverrideNames = PlatformOverrideNames;
+	Args.OnAddPlatformOverride = FOnPlatformOverrideAction::CreateSP(this, &FPerPlatformPropertyCustomization<PerPlatformType>::AddPlatformOverride, StructPropertyHandle);
+	Args.OnRemovePlatformOverride = FOnPlatformOverrideAction::CreateSP(this, &FPerPlatformPropertyCustomization<PerPlatformType>::RemovePlatformOverride, StructPropertyHandle);
+	Args.OnGenerateWidgetForPlatformRow = FOnGenerateWidget::CreateSP(this, &FPerPlatformPropertyCustomization<PerPlatformType>::GetWidget, StructPropertyHandle);
+
+	StructBuilder.AddCustomBuilder(MakeShared<FPerPlatformPropertyCustomNodeBuilder>(MoveTemp(Args)));
 }
 
+
 template<typename PerPlatformType>
 TSharedRef<SWidget> FPerPlatformPropertyCustomization<PerPlatformType>::GetWidget(FName PlatformGroupName, TSharedRef<IPropertyHandle> StructPropertyHandle) const
 {
@@ -99,36 +97,14 @@ TSharedRef<SWidget> FPerPlatformPropertyCustomization<PerPlatformType>::GetWidge
 		}
 	}
 
-	if (EditProperty.IsValid())
+	if (ensure(EditProperty.IsValid()))
 	{
 		return EditProperty->CreatePropertyValueWidget(false);
 	}
-	else
-	{
-		return
-			SNew(STextBlock)
-			.Text(NSLOCTEXT("FPerPlatformPropertyCustomization", "GetWidget", "Could not find valid property"))
-			.ColorAndOpacity(FLinearColor::Red);
-	}
+	
+	return SNullWidget::NullWidget;
 }
 
-template<typename PerPlatformType>
-float FPerPlatformPropertyCustomization<PerPlatformType>::CalcDesiredWidth(TSharedRef<IPropertyHandle> StructPropertyHandle)
-{
-	int32 NumOverrides = 0;
-	TSharedPtr<IPropertyHandle>	MapProperty = StructPropertyHandle->GetChildHandle(FName("PerPlatform"));
-	if (MapProperty.IsValid())
-	{
-		TArray<const void*> RawData;
-		MapProperty->AccessRawData(RawData);
-		for (const void* Data : RawData)
-		{
-			const TMap<FName, typename PerPlatformType::ValueType>* PerPlatformMap = (const TMap<FName, typename PerPlatformType::ValueType>*)(Data);
-			NumOverrides = FMath::Max<int32>(PerPlatformMap->Num(), NumOverrides);
-		}
-	}
-	return (float)(1 + NumOverrides) * 125.f;
-}
 
 template<typename PerPlatformType>
 bool FPerPlatformPropertyCustomization<PerPlatformType>::AddPlatformOverride(FName PlatformGroupName, TSharedRef<IPropertyHandle> StructPropertyHandle)
@@ -164,11 +140,6 @@ bool FPerPlatformPropertyCustomization<PerPlatformType>::AddPlatformOverride(FNa
 							DefaultProperty->GetValue(DefaultValue);
 							ChildProperty->SetValue(DefaultValue);
 
-							if(PropertyUtilities.IsValid())
-							{
-								PropertyUtilities.Pin()->ForceRefresh();
-							}
-
 							return true;
 						}
 					}
@@ -182,7 +153,6 @@ bool FPerPlatformPropertyCustomization<PerPlatformType>::AddPlatformOverride(FNa
 template<typename PerPlatformType>
 bool FPerPlatformPropertyCustomization<PerPlatformType>::RemovePlatformOverride(FName PlatformGroupName, TSharedRef<IPropertyHandle> StructPropertyHandle)
 {
-
 	FScopedTransaction Transaction(LOCTEXT("RemovePlatformOverride", "Remove Platform Override"));
 
 	TSharedPtr<IPropertyHandle>	MapProperty = StructPropertyHandle->GetChildHandle(FName("PerPlatform"));
@@ -200,13 +170,8 @@ bool FPerPlatformPropertyCustomization<PerPlatformType>::RemovePlatformOverride(
 			{
 				if (PlatformName == PlatformGroupName)
 				{
-
 					PerPlatformMap->Remove(PlatformName);
 
-					if (PropertyUtilities.IsValid())
-					{
-						PropertyUtilities.Pin()->ForceRefresh();
-					}
 					return true;
 				}
 			}
@@ -255,4 +220,171 @@ template class FPerPlatformPropertyCustomization<FPerPlatformInt>;
 template class FPerPlatformPropertyCustomization<FPerPlatformFloat>;
 template class FPerPlatformPropertyCustomization<FPerPlatformBool>;
 
-#undef LOCTEXT_NAMESPACE
\ No newline at end of file
+#undef LOCTEXT_NAMESPACE
+
+void FPerPlatformPropertyCustomNodeBuilder::SetOnRebuildChildren(FSimpleDelegate InOnRegenerateChildren)
+{
+	OnRebuildChildren = InOnRegenerateChildren;
+}
+
+void FPerPlatformPropertyCustomNodeBuilder::SetOnToggleExpansion(FOnToggleNodeExpansion InOnToggleExpansion)
+{
+	OnToggleExpansion = InOnToggleExpansion;
+}
+
+void FPerPlatformPropertyCustomNodeBuilder::GenerateHeaderRowContent(FDetailWidgetRow& HeaderRow)
+{
+	// Build Platform menu
+	FMenuBuilder AddPlatformMenuBuilder(true, nullptr, nullptr, true);
+
+	// Platform (group) names
+	const TArray<FName>& PlatformGroupNameArray = PlatformInfo::GetAllPlatformGroupNames();
+	const TArray<FName>& VanillaPlatformNameArray = PlatformInfo::GetAllVanillaPlatformNames();
+
+	// Sanitized platform names
+	TArray<FName> BasePlatformNameArray;
+	// Mapping from platform group name to individual platforms
+	TMultiMap<FName, FName> GroupToPlatform;
+
+	TArray<FName> PlatformOverrides = Args.PlatformOverrideNames.Get();
+
+	// Create mapping from platform to platform groups and remove postfixes and invalid platform names
+	for (const FName& PlatformName : VanillaPlatformNameArray)
+	{
+		// Add platform name if it isn't already set, and also add to group mapping
+		if (!PlatformOverrides.Contains(PlatformName))
+		{
+			BasePlatformNameArray.AddUnique(PlatformName);
+			GroupToPlatform.AddUnique(PlatformInfo::FindPlatformInfo(PlatformName)->DataDrivenPlatformInfo->PlatformGroupName, PlatformName);
+		}
+	}
+
+	// Create section for platform groups 
+	const FName PlatformGroupSection(TEXT("PlatformGroupSection"));
+	AddPlatformMenuBuilder.BeginSection(PlatformGroupSection, FText::FromString(TEXT("Platform Groups")));
+	for (const FName& GroupName : PlatformGroupNameArray)
+	{
+		if (!PlatformOverrides.Contains(GroupName))
+		{
+			const FTextFormat Format = NSLOCTEXT("SPerPlatformPropertiesWidget", "AddOverrideGroupFor", "Add Override for Platforms part of the {0} Platform Group");
+			AddPlatformToMenu(GroupName, Format, AddPlatformMenuBuilder);
+		}
+	}
+	AddPlatformMenuBuilder.EndSection();
+
+	for (const FName& GroupName : PlatformGroupNameArray)
+	{
+		// Create a section for each platform group and their respective platforms
+		AddPlatformMenuBuilder.BeginSection(GroupName, FText::FromName(GroupName));
+
+		TArray<FName> PlatformNames;
+		GroupToPlatform.MultiFind(GroupName, PlatformNames);
+
+		const FTextFormat Format = NSLOCTEXT("SPerPlatformPropertiesWidget", "AddOverrideFor", "Add Override specifically for {0}");
+		for (const FName& PlatformName : PlatformNames)
+		{
+			AddPlatformToMenu(PlatformName, Format, AddPlatformMenuBuilder);
+		}
+
+		AddPlatformMenuBuilder.EndSection();
+	}
+
+
+	HeaderRow.NameContent()
+	[
+		Args.NameWidget.ToSharedRef()
+	]
+	.ValueContent()
+	.MinDesiredWidth(125+28.0f)
+	[
+		SNew(SHorizontalBox)
+		.IsEnabled(Args.IsEnabled)
+		.ToolTipText(NSLOCTEXT("SPerPlatformPropertiesWidget", "DefaultPlatformDesc", "This property can have per-platform or platform group overrides.\nThis is the default value used when no override has been set for a platform or platform group."))
+		+SHorizontalBox::Slot()
+		[
+			SNew(SPerPlatformPropertiesRow, NAME_None)
+			.OnGenerateWidget(Args.OnGenerateWidgetForPlatformRow)
+		]
+		+SHorizontalBox::Slot()
+		.AutoWidth()
+		.Padding(2.0f, 0.0f, 0.0f, 0.0f)
+		.VAlign(VAlign_Center)
+		[
+			SNew(SComboButton)
+			.ComboButtonStyle(FEditorStyle::Get(), "SimpleComboButton")
+			.HasDownArrow(false)
+			.ToolTipText(NSLOCTEXT("SPerPlatformPropertiesWidget", "AddOverrideToolTip", "Add an override for a specific platform or platform group"))
+			.ButtonContent()
+			[
+				SNew(SImage)
+				.Image(FEditorStyle::GetBrush("Icons.PlusCircle"))
+				.ColorAndOpacity(FSlateColor::UseForeground())
+			]
+			.MenuContent()
+			[
+				AddPlatformMenuBuilder.MakeWidget()
+			]	
+		]
+	];
+}
+
+void FPerPlatformPropertyCustomNodeBuilder::GenerateChildContent(IDetailChildrenBuilder& ChildrenBuilder)
+{
+	TArray<FName> PlatformOverrides = Args.PlatformOverrideNames.Get();
+	for (FName PlatformName : PlatformOverrides)
+	{
+		FText PlatformDisplayName = FText::AsCultureInvariant(PlatformName.ToString());
+		FDetailWidgetRow& Row = ChildrenBuilder.AddCustomRow(PlatformDisplayName);
+		Row.IsEnabled(Args.IsEnabled);
+
+		Row.NameContent()
+		[
+			SNew(STextBlock)
+			.Text(PlatformDisplayName)
+			.Font(IDetailLayoutBuilder::GetDetailFont())
+		];
+
+		Row.ValueContent()
+		[
+			SNew(SPerPlatformPropertiesRow, PlatformName)
+			.OnGenerateWidget(Args.OnGenerateWidgetForPlatformRow)
+			.OnRemovePlatform(this, &FPerPlatformPropertyCustomNodeBuilder::OnRemovePlatformOverride)
+		];
+	}
+}
+
+FName FPerPlatformPropertyCustomNodeBuilder::GetName() const
+{
+	static const FName Name("FPerPlatformPropertyCustomNodeBuilder");
+	return Name;
+}
+
+void FPerPlatformPropertyCustomNodeBuilder::OnAddPlatformOverride(const FName PlatformName)
+{
+	if (Args.OnAddPlatformOverride.IsBound() && Args.OnAddPlatformOverride.Execute(PlatformName))
+	{
+		OnRebuildChildren.ExecuteIfBound();
+		OnToggleExpansion.ExecuteIfBound(true);
+	}
+}
+
+bool FPerPlatformPropertyCustomNodeBuilder::OnRemovePlatformOverride(const FName PlatformName)
+{
+	if (Args.OnRemovePlatformOverride.IsBound() && Args.OnRemovePlatformOverride.Execute(PlatformName))
+	{
+		OnRebuildChildren.ExecuteIfBound();
+	}
+
+	return true;
+}
+void FPerPlatformPropertyCustomNodeBuilder::AddPlatformToMenu(const FName PlatformName, const FTextFormat Format, FMenuBuilder& AddPlatformMenuBuilder)
+{
+	const FText MenuText = FText::Format(FText::FromString(TEXT("{0}")), FText::AsCultureInvariant(PlatformName.ToString()));
+	const FText MenuTooltipText = FText::Format(Format, FText::AsCultureInvariant(PlatformName.ToString()));
+	AddPlatformMenuBuilder.AddMenuEntry(
+		MenuText,
+		MenuTooltipText,
+		FSlateIcon(FEditorStyle::GetStyleSetName(), "PerPlatformWidget.AddPlatform"),
+		FUIAction(FExecuteAction::CreateSP(this, &FPerPlatformPropertyCustomNodeBuilder::OnAddPlatformOverride, PlatformName))
+	);
+}
diff --git a/Engine/Source/Editor/DetailCustomizations/Private/PerPlatformPropertyCustomization.h b/Engine/Source/Editor/DetailCustomizations/Private/PerPlatformPropertyCustomization.h
deleted file mode 100644
diff --git a/Engine/Source/Editor/DetailCustomizations/Private/SPerPlatformPropertiesWidget.cpp b/Engine/Source/Editor/DetailCustomizations/Private/SPerPlatformPropertiesWidget.cpp
new file mode 100644
--- /dev/null
+++ b/Engine/Source/Editor/DetailCustomizations/Private/SPerPlatformPropertiesWidget.cpp
@@ -0,0 +1,84 @@
+// Copyright Epic Games, Inc. All Rights Reserved.
+
+#include "SPerPlatformPropertiesWidget.h"
+#include "Layout/Margin.h"
+#include "Fonts/SlateFontInfo.h"
+#include "Styling/CoreStyle.h"
+#include "Framework/MultiBox/MultiBoxBuilder.h"
+#include "Widgets/Input/SButton.h"
+#include "Widgets/Input/SComboButton.h"
+#include "PlatformInfo.h"
+#include "Widgets/Layout/SWrapBox.h"
+#include "EditorStyleSet.h"
+#include "Widgets/Layout/SBox.h"
+#include "DetailLayoutBuilder.h"
+#include "Widgets/Images/SImage.h"
+
+
+void SPerPlatformPropertiesRow::Construct(const FArguments& InArgs, FName PlatformName)
+{
+	this->OnGenerateWidget = InArgs._OnGenerateWidget;
+	this->OnRemovePlatform = InArgs._OnRemovePlatform;
+
+	ChildSlot
+	[
+		MakePerPlatformWidget(PlatformName)
+	];
+
+}
+
+
+FReply SPerPlatformPropertiesRow::RemovePlatform(FName PlatformName)
+{
+	if (OnRemovePlatform.IsBound())
+	{
+		OnRemovePlatform.Execute(PlatformName);
+	}
+
+	return FReply::Handled();
+}
+
+TSharedRef<SWidget> SPerPlatformPropertiesRow::MakePerPlatformWidget(FName InName)
+{
+	TSharedPtr<SHorizontalBox> HorizontalBox;
+
+	TSharedRef<SWidget> Widget = 
+		SNew(SBox)
+		.Padding(FMargin(0.0f, 2.0f, 4.0f, 2.0f))
+		.MinDesiredWidth(50.0f)
+		[
+			SAssignNew(HorizontalBox, SHorizontalBox)
+			.ToolTipText(InName != NAME_None 
+				? FText::Format(NSLOCTEXT("SPerPlatformPropertiesWidget", "PerPlatformDesc", "Override for {0}"), FText::AsCultureInvariant(InName.ToString()))
+				: NSLOCTEXT("SPerPlatformPropertiesWidget", "DefaultDesc", "Default value for properties without an override"))
+			+SHorizontalBox::Slot()
+			.Padding(0.0f, 0.0f, 2.0f, 0.0f)
+			.VAlign(VAlign_Center)
+			[
+				OnGenerateWidget.Execute(InName)
+			]
+		];
+
+	if(InName != NAME_None)
+	{
+		HorizontalBox->AddSlot()
+			.AutoWidth()
+			.Padding(2.0f, 0.0f, 0.0f, 0.0f)
+			.VAlign(VAlign_Center)
+			[
+				SNew(SButton)
+				.ButtonStyle(FEditorStyle::Get(), "SimpleButton")
+				.OnClicked(this, &SPerPlatformPropertiesRow::RemovePlatform, InName)
+				.ToolTipText(FText::Format(NSLOCTEXT("SPerPlatformPropertiesWidget", "RemoveOverrideFor", "Remove Override for {0}"), FText::AsCultureInvariant(InName.ToString())))
+				.Content()
+				[
+					SNew(SImage)
+					.Image(FEditorStyle::GetBrush("Icons.Delete"))
+					.ColorAndOpacity(FSlateColor::UseForeground())
+				]
+			];
+	}
+
+
+	return Widget;
+}
\ No newline at end of file
diff --git a/Engine/Source/Editor/DetailCustomizations/Private/SPerPlatformPropertiesWidget.h b/Engine/Source/Editor/DetailCustomizations/Private/SPerPlatformPropertiesWidget.h
new file mode 100644
--- /dev/null
+++ b/Engine/Source/Editor/DetailCustomizations/Private/SPerPlatformPropertiesWidget.h
@@ -0,0 +1,53 @@
+// Copyright Epic Games, Inc. All Rights Reserved.
+
+#pragma once
+
+#include "CoreMinimal.h"
+#include "Widgets/DeclarativeSyntaxSupport.h"
+#include "Widgets/SWidget.h"
+#include "Widgets/SCompoundWidget.h"
+#include "Framework/SlateDelegates.h"
+#include "IDetailChildrenBuilder.h"
+#include "PerPlatformPropertyCustomization.h"
+
+class FMenuBuilder;
+
+
+/**
+* SPerPlatformPropertiesWidget
+*/
+class SPerPlatformPropertiesRow : public SCompoundWidget
+{
+public:
+
+	SLATE_BEGIN_ARGS(SPerPlatformPropertiesRow)
+		{}
+		
+		SLATE_EVENT(FOnGenerateWidget, OnGenerateWidget)
+		SLATE_EVENT(FOnPlatformOverrideAction, OnRemovePlatform)
+
+	SLATE_END_ARGS()
+
+	/**
+	* Construct this widget
+	*
+	* @param	InArgs	The declaration data for this widget
+	*/
+	void Construct(const FArguments& InArgs, FName PlatformName);
+
+protected:
+	TSharedRef<SWidget> MakePerPlatformWidget(FName InName);
+
+	FReply RemovePlatform(FName PlatformName);
+
+	EActiveTimerReturnType CheckPlatformCount(double InCurrentTime, float InDeltaSeconds);
+
+	void AddPlatformToMenu(const FName& PlatformName, const FTextFormat Format, FMenuBuilder& AddPlatformMenuBuilder);
+
+	FOnGenerateWidget OnGenerateWidget;
+	FOnPlatformOverrideAction OnRemovePlatform;
+	TAttribute<TArray<FName>> PlatformOverrideNames;
+	int32 LastPlatformOverrideNames;
+	bool bAddedMenuItem;
+};
+
diff --git a/Engine/Source/Editor/DetailCustomizations/Public/PerPlatformPropertyCustomization.h b/Engine/Source/Editor/DetailCustomizations/Public/PerPlatformPropertyCustomization.h
new file mode 100644
--- /dev/null
+++ b/Engine/Source/Editor/DetailCustomizations/Public/PerPlatformPropertyCustomization.h
@@ -0,0 +1,93 @@
+// Copyright Epic Games, Inc. All Rights Reserved.
+#pragma once
+
+#include "CoreMinimal.h"
+#include "Widgets/SWidget.h"
+#include "UnrealClient.h"
+#include "IPropertyTypeCustomization.h"
+#include "PerPlatformProperties.h"
+#include "IDetailCustomNodeBuilder.h"
+
+class FDetailWidgetDecl;
+
+typedef typename TSlateDelegates<FName>::FOnGenerateWidget FOnGenerateWidget;
+
+DECLARE_DELEGATE_RetVal_OneParam(bool, FOnPlatformOverrideAction, FName);
+
+struct FPerPlatformPropertyCustomNodeBuilderArgs
+{
+	/** Handle to the default value */
+	TSharedPtr<SWidget> NameWidget;
+
+	/** List of platforms that can override the default value */
+	TAttribute<TArray<FName>> PlatformOverrideNames;
+
+	/** Is overriding per-property enabled */
+	TAttribute<bool> IsEnabled;
+
+	/** Callback to generate a widget for a specific platform row */
+	FOnGenerateWidget OnGenerateWidgetForPlatformRow;
+
+	FOnPlatformOverrideAction OnAddPlatformOverride;
+
+	FOnPlatformOverrideAction OnRemovePlatformOverride;
+};
+
+class DETAILCUSTOMIZATIONS_API FPerPlatformPropertyCustomNodeBuilder : public IDetailCustomNodeBuilder, public TSharedFromThis<FPerPlatformPropertyCustomNodeBuilder>
+{
+public:
+	FPerPlatformPropertyCustomNodeBuilder(const FPerPlatformPropertyCustomNodeBuilderArgs& InArgs)
+		: Args(InArgs)
+	{}
+
+	virtual void Tick(float DeltaTime) override {}
+	virtual bool RequiresTick() const override { return false; }
+
+	virtual void SetOnRebuildChildren(FSimpleDelegate InOnRegenerateChildren) override;
+	virtual void SetOnToggleExpansion(FOnToggleNodeExpansion InOnRegenerateChildren) override;
+	virtual void GenerateHeaderRowContent(FDetailWidgetRow& HeaderRow);
+	virtual void GenerateChildContent(IDetailChildrenBuilder& ChildrenBuilder) override;
+	virtual bool InitiallyCollapsed() const override { return true; }
+	virtual FName GetName() const override;
+private:
+	void OnAddPlatformOverride(FName PlatformName);
+	bool OnRemovePlatformOverride(FName PlatformName);
+	void AddPlatformToMenu(const FName PlatformName, const FTextFormat Format, FMenuBuilder& AddPlatformMenuBuilder);
+private:
+	/** Handle to the default value */
+	FPerPlatformPropertyCustomNodeBuilderArgs Args;
+	FSimpleDelegate OnRebuildChildren;
+	FOnToggleNodeExpansion OnToggleExpansion;
+};
+/**
+* Implements a details panel customization for the FPerPlatform structures.
+*/
+template<typename PerPlatformType>
+class FPerPlatformPropertyCustomization : public IPropertyTypeCustomization
+{
+public:
+	FPerPlatformPropertyCustomization()
+	{}
+
+	// IPropertyTypeCustomization interface
+	virtual void CustomizeHeader(TSharedRef<IPropertyHandle> StructPropertyHandle, FDetailWidgetRow& HeaderRow, IPropertyTypeCustomizationUtils& StructCustomizationUtils) override {}
+	virtual void CustomizeChildren(TSharedRef<IPropertyHandle> StructPropertyHandle, IDetailChildrenBuilder& StructBuilder, IPropertyTypeCustomizationUtils& StructCustomizationUtils) override;
+
+	/**
+	* Creates a new instance.
+	*
+	* @return A new customization for FPerPlatform structs.
+	*/
+	static TSharedRef<IPropertyTypeCustomization> MakeInstance();
+
+protected:
+	TSharedRef<SWidget> GetWidget(FName PlatformGroupName, TSharedRef<IPropertyHandle> StructPropertyHandle) const;
+	TArray<FName> GetPlatformOverrideNames(TSharedRef<IPropertyHandle> StructPropertyHandle) const;
+	bool AddPlatformOverride(FName PlatformGroupName, TSharedRef<IPropertyHandle> StructPropertyHandle);
+	bool RemovePlatformOverride(FName PlatformGroupName, TSharedRef<IPropertyHandle> StructPropertyHandle);
+
+private:
+	/** Cached utils used for resetting customization when layout changes */
+	TWeakPtr<IPropertyUtilities> PropertyUtilities;
+};
+
diff --git a/Engine/Source/Editor/Persona/Private/PersonaMeshDetails.cpp b/Engine/Source/Editor/Persona/Private/PersonaMeshDetails.cpp
--- a/Engine/Source/Editor/Persona/Private/PersonaMeshDetails.cpp
+++ b/Engine/Source/Editor/Persona/Private/PersonaMeshDetails.cpp
@@ -2715,7 +2715,7 @@ void FPersonaMeshDetails::AddLODLevelCategories(IDetailLayoutBuilder& DetailLayo
 				SectionListDelegates.OnEnableSectionItem.BindSP(this, &FPersonaMeshDetails::OnSectionEnabledChanged);
 
 				FName SkeletalMeshSectionListName = FName(*(FString(TEXT("SkeletalMeshSectionListNameLOD_")) + FString::FromInt(LODIndex)));
-				LODCategory.AddCustomBuilder(MakeShareable(new FSectionList(LODCategory.GetParentLayout(), SectionListDelegates, false, 64, LODIndex, SkeletalMeshSectionListName)));
+				LODCategory.AddCustomBuilder(MakeShareable(new FSectionList(LODCategory.GetParentLayout(), SectionListDelegates, false, 48, LODIndex, SkeletalMeshSectionListName)));
 
 				GetPersonaToolkit()->GetPreviewScene()->RegisterOnSelectedLODChanged(FOnSelectedLODChanged::CreateSP(this, &FPersonaMeshDetails::UpdateLODCategoryVisibility));
 			}
diff --git a/Engine/Source/Editor/PropertyEditor/Private/DetailCustomBuilderRow.cpp b/Engine/Source/Editor/PropertyEditor/Private/DetailCustomBuilderRow.cpp
--- a/Engine/Source/Editor/PropertyEditor/Private/DetailCustomBuilderRow.cpp
+++ b/Engine/Source/Editor/PropertyEditor/Private/DetailCustomBuilderRow.cpp
@@ -37,14 +37,15 @@ void FDetailCustomBuilderRow::OnItemNodeInitialized( TSharedRef<FDetailItemNode>
 	IsParentEnabled = InIsParentEnabled;
 
 	const bool bUpdateFilteredNodes = true;
+
 	// Set a delegate on the interface that it will call to rebuild this nodes children
-	FSimpleDelegate OnRegenerateChildren = FSimpleDelegate::CreateSP( InTreeNode, &FDetailItemNode::GenerateChildren, bUpdateFilteredNodes );
+	CustomNodeBuilder->SetOnRebuildChildren(FSimpleDelegate::CreateSP(InTreeNode, &FDetailItemNode::GenerateChildren, bUpdateFilteredNodes));
 
-	CustomNodeBuilder->SetOnRebuildChildren( OnRegenerateChildren );
+	CustomNodeBuilder->SetOnToggleExpansion(FOnToggleNodeExpansion::CreateSP(InTreeNode, &FDetailItemNode::SetExpansionState));
 
-	HeaderRow = MakeShareable( new FDetailWidgetRow );
+	HeaderRow = MakeShared<FDetailWidgetRow>();
 
-	CustomNodeBuilder->GenerateHeaderRowContent( *HeaderRow );
+	CustomNodeBuilder->GenerateHeaderRowContent(*HeaderRow);
 }
 
 FName FDetailCustomBuilderRow::GetCustomBuilderName() const
@@ -59,7 +60,7 @@ TSharedPtr<IPropertyHandle> FDetailCustomBuilderRow::GetPropertyHandle() const
 
 void FDetailCustomBuilderRow::OnGenerateChildren( FDetailNodeList& OutChildren )
 {
-	ChildrenBuilder = MakeShareable( new FCustomChildrenBuilder( ParentCategory.Pin().ToSharedRef() ) );
+	ChildrenBuilder = MakeShared<FCustomChildrenBuilder>(ParentCategory.Pin().ToSharedRef());
 
 	CustomNodeBuilder->GenerateChildContent( *ChildrenBuilder );
 		
diff --git a/Engine/Source/Editor/PropertyEditor/Private/DetailItemNode.cpp b/Engine/Source/Editor/PropertyEditor/Private/DetailItemNode.cpp
--- a/Engine/Source/Editor/PropertyEditor/Private/DetailItemNode.cpp
+++ b/Engine/Source/Editor/PropertyEditor/Private/DetailItemNode.cpp
@@ -219,6 +219,12 @@ void FDetailItemNode::SetExpansionState(bool bWantsExpanded, bool bSaveState)
 	OnItemExpansionChanged(bIsExpanded, bSaveState);
 }
 
+void FDetailItemNode::SetExpansionState(bool bWantsExpanded)
+{
+	const bool bSaveState = true;
+	SetExpansionState(bWantsExpanded, bSaveState);
+}
+
 TSharedRef< ITableRow > FDetailItemNode::GenerateWidgetForTableView( const TSharedRef<STableViewBase>& OwnerTable, bool bAllowFavoriteSystem)
 {
 	FTagMetaData TagMeta(TEXT("DetailRowItem"));
diff --git a/Engine/Source/Editor/PropertyEditor/Private/DetailItemNode.h b/Engine/Source/Editor/PropertyEditor/Private/DetailItemNode.h
--- a/Engine/Source/Editor/PropertyEditor/Private/DetailItemNode.h
+++ b/Engine/Source/Editor/PropertyEditor/Private/DetailItemNode.h
@@ -31,11 +31,10 @@ public:
 	 */
 	void Initialize();
 
-
 	void ToggleExpansion();
 
 	void SetExpansionState(bool bWantsExpanded, bool bSaveState);
-
+	void SetExpansionState(bool bWantsExpanded);
 	/**
 	 * Generates children for this node
 	 *
diff --git a/Engine/Source/Editor/PropertyEditor/Private/PropertyCustomizationHelpers.cpp b/Engine/Source/Editor/PropertyEditor/Private/PropertyCustomizationHelpers.cpp
--- a/Engine/Source/Editor/PropertyEditor/Private/PropertyCustomizationHelpers.cpp
+++ b/Engine/Source/Editor/PropertyEditor/Private/PropertyCustomizationHelpers.cpp
@@ -157,7 +157,7 @@ namespace PropertyCustomizationHelpers
 			SNew( SPropertyEditorButton )
 			.Text( LOCTEXT( "DeleteButtonLabel", "Delete") )
 			.ToolTipText( OptionalToolTipText.Get().IsEmpty() ? LOCTEXT( "DeleteButtonToolTipText", "Delete") : OptionalToolTipText )
-			.Image( FEditorStyle::GetBrush("Icons.X") )
+			.Image( FEditorStyle::GetBrush("Icons.Delete") )
 			.OnClickAction( OnDeleteClicked )
 			.IsEnabled(IsEnabled)
 			.IsFocusable( false );
@@ -1394,7 +1394,6 @@ void SMaterialSlotWidget::Construct(const FArguments& InArgs, int32 SlotIndex, b
 		[
 			SNew(SBox)
 			.VAlign(VAlign_Center)
-			.HAlign(HAlign_Fill)
 			.MinDesiredWidth(160.0f)
 			[
 				SNew(SEditableTextBox)
@@ -1409,14 +1408,13 @@ void SMaterialSlotWidget::Construct(const FArguments& InArgs, int32 SlotIndex, b
 	
 	if (bIsMaterialUsed)
 	{
-		DeleteButton->SetVisibility(EVisibility::Hidden);
+		DeleteButton->SetEnabled(false);
 	}
 	
 
 	SlotNameBox->AddSlot()
 		.AutoWidth()
 		.VAlign(VAlign_Center)
-		.HAlign(HAlign_Left)
 		.Padding(2)
 		[
 			DeleteButton
diff --git a/Engine/Source/Editor/PropertyEditor/Private/SPerPlatformPropertiesWidget.cpp b/Engine/Source/Editor/PropertyEditor/Private/SPerPlatformPropertiesWidget.cpp
deleted file mode 100644
diff --git a/Engine/Source/Editor/PropertyEditor/Public/IDetailCustomNodeBuilder.h b/Engine/Source/Editor/PropertyEditor/Public/IDetailCustomNodeBuilder.h
--- a/Engine/Source/Editor/PropertyEditor/Public/IDetailCustomNodeBuilder.h
+++ b/Engine/Source/Editor/PropertyEditor/Public/IDetailCustomNodeBuilder.h
@@ -8,6 +8,7 @@ class FDetailWidgetRow;
 class IDetailChildrenBuilder;
 class IPropertyHandle;
 
+DECLARE_DELEGATE_OneParam(FOnToggleNodeExpansion, bool)
 /**
  * A custom node that can be given to a details panel category to customize widgets
  */
@@ -21,7 +22,14 @@ public:
 	 *
 	 * @param A delegate to invoke when the tree should rebuild this nodes children
 	 */
-	virtual void SetOnRebuildChildren( FSimpleDelegate InOnRegenerateChildren ) = 0;
+	virtual void SetOnRebuildChildren(FSimpleDelegate InOnRegenerateChildren) {}
+
+	/**
+	 * Sets a delegate that should be used when the custom node wants to expand
+	 *
+	 * @param A delegate to invoke when the tree should expand
+	 */
+	virtual void SetOnToggleExpansion(FOnToggleNodeExpansion InOnToggleExpansion) {}
 
 	/**
 	 * Called to generate content in the header of this node ( the actual node content ). 
@@ -41,17 +49,17 @@ public:
 	/**
 	 * Called each tick if ReqiresTick is true
 	 */
-	virtual void Tick( float DeltaTime ) = 0;
+	virtual void Tick(float DeltaTime) {}
 
 	/**
 	 * @return true if this node requires tick, false otherwise
 	 */
-	virtual bool RequiresTick() const = 0;
+	virtual bool RequiresTick() const { return false; }
 
 	/**
 	 * @return true if this node should be collapsed in the tree
 	 */
-	virtual bool InitiallyCollapsed() const = 0;
+	virtual bool InitiallyCollapsed() const { return true; }
 
 	/**
 	 * @return The name of this custom builder.  This is used as an identifier to save expansion state if needed
diff --git a/Engine/Source/Editor/PropertyEditor/Public/SPerPlatformPropertiesWidget.h b/Engine/Source/Editor/PropertyEditor/Public/SPerPlatformPropertiesWidget.h
deleted file mode 100644
diff --git a/Engine/Source/Editor/StaticMeshEditor/Private/StaticMeshEditorTools.cpp b/Engine/Source/Editor/StaticMeshEditor/Private/StaticMeshEditorTools.cpp
--- a/Engine/Source/Editor/StaticMeshEditor/Private/StaticMeshEditorTools.cpp
+++ b/Engine/Source/Editor/StaticMeshEditor/Private/StaticMeshEditorTools.cpp
@@ -29,9 +29,7 @@
 #include "FbxMeshUtils.h"
 #include "Widgets/Input/SVectorInputBox.h"
 #include "Widgets/Input/SNumericEntryBox.h"
-#include "SPerPlatformPropertiesWidget.h"
 #include "PlatformInfo.h"
-
 #include "ContentStreaming.h"
 #include "EditorDirectories.h"
 #include "EditorFramework/AssetImportData.h"
@@ -48,6 +46,7 @@
 #include "UObject/UObjectGlobals.h"
 #include "Widgets/Input/SFilePathPicker.h"
 #include "Widgets/Input/STextComboBox.h"
+#include "PerPlatformPropertyCustomization.h"
 
 const uint32 MaxHullCount = 64;
 const uint32 MinHullCount = 2;
@@ -1753,7 +1752,7 @@ void FMeshSectionSettingsLayout::AddToCategory( IDetailCategoryBuilder& Category
 	SectionListDelegates.OnPasteSectionItem.BindSP(this, &FMeshSectionSettingsLayout::OnPasteSectionItem);
 	//We need a valid name if we want the section expand state to be saved
 	FName StaticMeshSectionListName = FName(*(FString(TEXT("StaticMeshSectionListNameLOD_")) + FString::FromInt(LODIndex)));
-	CategoryBuilder.AddCustomBuilder(MakeShareable(new FSectionList(CategoryBuilder.GetParentLayout(), SectionListDelegates, true, 64, LODIndex, StaticMeshSectionListName)));
+	CategoryBuilder.AddCustomBuilder(MakeShareable(new FSectionList(CategoryBuilder.GetParentLayout(), SectionListDelegates, true, 48, LODIndex, StaticMeshSectionListName)));
 
 	StaticMeshEditor.RegisterOnSelectedLODChanged(FOnSelectedLODChanged::CreateSP(this, &FMeshSectionSettingsLayout::UpdateLODCategoryVisibility), false);
 }
@@ -3247,45 +3246,41 @@ void FLevelOfDetailSettingsLayout::AddToDetailsPanel( IDetailLayoutBuilder& Deta
 			.OnSelectionChanged(this, &FLevelOfDetailSettingsLayout::OnImportLOD)
 		];
 
-	int32 PlatformNumber = PlatformInfo::GetAllPlatformGroupNames().Num();
+	{
+		TAttribute<TArray<FName>> PlatformOverrideNames = TAttribute<TArray<FName>>::Create(TAttribute<TArray<FName>>::FGetter::CreateSP(this, &FLevelOfDetailSettingsLayout::GetMinLODPlatformOverrideNames));
+		FPerPlatformPropertyCustomNodeBuilderArgs Args;
+		Args.NameWidget =
+			SNew(STextBlock)
+			.IsEnabled(FLevelOfDetailSettingsLayout::GetLODCount() > 1)
+			.Font(IDetailLayoutBuilder::GetDetailFont())
+			.Text(LOCTEXT("MinLOD", "Minimum LOD"));
 
-	LODSettingsCategory.AddCustomRow( LOCTEXT("MinLOD", "Minimum LOD") )
-	.NameContent()
-	[
-		SNew(STextBlock)
-		.Font( IDetailLayoutBuilder::GetDetailFont() )
-		.Text(LOCTEXT("MinLOD", "Minimum LOD"))
-	]
-	.ValueContent()
-	.MinDesiredWidth((float)(StaticMesh->GetMinLOD().PerPlatform.Num() + 1)*125.0f)
-	.MaxDesiredWidth((float)(PlatformNumber + 1)*125.0f)
-	[
-		SNew(SPerPlatformPropertiesWidget)
-		.IsEnabled(FLevelOfDetailSettingsLayout::GetLODCount() > 1)
-		.OnGenerateWidget(this, &FLevelOfDetailSettingsLayout::GetMinLODWidget)
-		.OnAddPlatform(this, &FLevelOfDetailSettingsLayout::AddMinLODPlatformOverride)
-		.OnRemovePlatform(this, &FLevelOfDetailSettingsLayout::RemoveMinLODPlatformOverride)
-		.PlatformOverrideNames(this, &FLevelOfDetailSettingsLayout::GetMinLODPlatformOverrideNames)
-	];
+		Args.PlatformOverrideNames = PlatformOverrideNames;
+		Args.OnAddPlatformOverride = FOnPlatformOverrideAction::CreateSP(this, &FLevelOfDetailSettingsLayout::AddMinLODPlatformOverride);
+		Args.OnRemovePlatformOverride = FOnPlatformOverrideAction::CreateSP(this, &FLevelOfDetailSettingsLayout::RemoveMinLODPlatformOverride);
+		Args.OnGenerateWidgetForPlatformRow = FOnGenerateWidget::CreateSP(this, &FLevelOfDetailSettingsLayout::GetMinLODWidget);
+		Args.IsEnabled = TAttribute<bool>(FLevelOfDetailSettingsLayout::GetLODCount() > 1);
 
-	LODSettingsCategory.AddCustomRow(LOCTEXT("NumStreamedLODs", "Num Streamed LODs"))
-	.NameContent()
-	[
-		SNew(STextBlock)
-		.Font(IDetailLayoutBuilder::GetDetailFont())
-		.Text(LOCTEXT("NumStreamdLODs", "Num Streamed LODs"))
-	]
-	.ValueContent()
-	.MinDesiredWidth((float)(StaticMesh->NumStreamedLODs.PerPlatform.Num() + 1)*125.0f)
-	.MaxDesiredWidth((float)(PlatformNumber + 1)*125.0f)
-	[
-		SNew(SPerPlatformPropertiesWidget)
-		.IsEnabled(FLevelOfDetailSettingsLayout::GetLODCount() > 1)
-		.OnGenerateWidget(this, &FLevelOfDetailSettingsLayout::GetNumStreamedLODsWidget)
-		.OnAddPlatform(this, &FLevelOfDetailSettingsLayout::AddNumStreamedLODsPlatformOverride)
-		.OnRemovePlatform(this, &FLevelOfDetailSettingsLayout::RemoveNumStreamedLODsPlatformOverride)
-		.PlatformOverrideNames(this, &FLevelOfDetailSettingsLayout::GetNumStreamedLODsPlatformOverrideNames)
-	];
+		LODSettingsCategory.AddCustomBuilder(MakeShared<FPerPlatformPropertyCustomNodeBuilder>(MoveTemp(Args)));
+	}
+
+	{
+		TAttribute<TArray<FName>> PlatformOverrideNames = TAttribute<TArray<FName>>::Create(TAttribute<TArray<FName>>::FGetter::CreateSP(this, &FLevelOfDetailSettingsLayout::GetNumStreamedLODsPlatformOverrideNames));
+		FPerPlatformPropertyCustomNodeBuilderArgs Args;
+		Args.NameWidget =
+			SNew(STextBlock)
+			.IsEnabled(FLevelOfDetailSettingsLayout::GetLODCount() > 1)
+			.Font(IDetailLayoutBuilder::GetDetailFont())
+			.Text(LOCTEXT("NumStreamdLODs", "Num Streamed LODs"));
+
+		Args.PlatformOverrideNames = PlatformOverrideNames;
+		Args.OnAddPlatformOverride = FOnPlatformOverrideAction::CreateSP(this, &FLevelOfDetailSettingsLayout::AddNumStreamedLODsPlatformOverride);
+		Args.OnRemovePlatformOverride = FOnPlatformOverrideAction::CreateSP(this, &FLevelOfDetailSettingsLayout::RemoveNumStreamedLODsPlatformOverride);
+		Args.OnGenerateWidgetForPlatformRow = FOnGenerateWidget::CreateSP(this, &FLevelOfDetailSettingsLayout::GetNumStreamedLODsWidget);
+		Args.IsEnabled = TAttribute<bool>(FLevelOfDetailSettingsLayout::GetLODCount() > 1);
+
+		LODSettingsCategory.AddCustomBuilder(MakeShared<FPerPlatformPropertyCustomNodeBuilder>(MoveTemp(Args)));
+	}
 
 	// Add Number of LODs slider.
 	const int32 MinAllowedLOD = 1;
@@ -3572,24 +3567,22 @@ void FLevelOfDetailSettingsLayout::AddLODLevelCategories( IDetailLayoutBuilder&
 
 			int32 PlatformNumber = PlatformInfo::GetAllPlatformGroupNames().Num();
 
-			LODCategory.AddCustomRow(( LOCTEXT("ScreenSizeRow", "ScreenSize")))
-			.NameContent()
-			[
+			TAttribute<TArray<FName>> PlatformOverrideNames = TAttribute<TArray<FName>>::Create(TAttribute<TArray<FName>>::FGetter::CreateSP(this, &FLevelOfDetailSettingsLayout::GetLODScreenSizePlatformOverrideNames, LODIndex));
+
+			FPerPlatformPropertyCustomNodeBuilderArgs Args;
+			Args.NameWidget =
 				SNew(STextBlock)
-				.Font(IDetailLayoutBuilder::GetDetailFont())
-				.Text(LOCTEXT("ScreenSizeName", "Screen Size"))
-			]
-			.ValueContent()
-			.MinDesiredWidth(GetScreenSizeWidgetWidth(LODIndex))
-			.MaxDesiredWidth((float)(PlatformNumber + 1)*125.0f)
-			[
-				SNew(SPerPlatformPropertiesWidget)
 				.IsEnabled(this, &FLevelOfDetailSettingsLayout::CanChangeLODScreenSize)
-				.OnGenerateWidget(this, &FLevelOfDetailSettingsLayout::GetLODScreenSizeWidget, LODIndex)
-				.OnAddPlatform(this, &FLevelOfDetailSettingsLayout::AddLODScreenSizePlatformOverride, LODIndex)
-				.OnRemovePlatform(this, &FLevelOfDetailSettingsLayout::RemoveLODScreenSizePlatformOverride, LODIndex)
-				.PlatformOverrideNames(this, &FLevelOfDetailSettingsLayout::GetLODScreenSizePlatformOverrideNames, LODIndex)
-			];
+				.Font(IDetailLayoutBuilder::GetDetailFont())
+				.Text(LOCTEXT("ScreenSizeName", "Screen Size"));
+
+			Args.PlatformOverrideNames = PlatformOverrideNames;
+			Args.OnAddPlatformOverride = FOnPlatformOverrideAction::CreateSP(this, &FLevelOfDetailSettingsLayout::AddLODScreenSizePlatformOverride, LODIndex);
+			Args.OnRemovePlatformOverride = FOnPlatformOverrideAction::CreateSP(this, &FLevelOfDetailSettingsLayout::RemoveLODScreenSizePlatformOverride, LODIndex);
+			Args.OnGenerateWidgetForPlatformRow = FOnGenerateWidget::CreateSP(this, &FLevelOfDetailSettingsLayout::GetLODScreenSizeWidget, LODIndex);
+			Args.IsEnabled = TAttribute<bool>(this, &FLevelOfDetailSettingsLayout::CanChangeLODScreenSize);
+
+			LODCategory.AddCustomBuilder(MakeShared<FPerPlatformPropertyCustomNodeBuilder>(MoveTemp(Args)));
 
 			if(LODIndex > 0 && StaticMesh->IsMeshDescriptionValid(LODIndex))
 			{
@@ -4483,7 +4476,6 @@ TSharedRef<SWidget> FLevelOfDetailSettingsLayout::OnGenerateLodComboBoxForLodPic
 		.IsEnabled(this, &FLevelOfDetailSettingsLayout::IsLodComboBoxEnabledForLodPicker)
 		.OnGetMenuContent(this, &FLevelOfDetailSettingsLayout::OnGenerateLodMenuForLodPicker)
 		.VAlign(VAlign_Center)
-		.ContentPadding(2)
 		.ButtonContent()
 		[
 			SNew(STextBlock)
diff --git a/Engine/Source/Editor/StaticMeshEditor/StaticMeshEditor.Build.cs b/Engine/Source/Editor/StaticMeshEditor/StaticMeshEditor.Build.cs
--- a/Engine/Source/Editor/StaticMeshEditor/StaticMeshEditor.Build.cs
+++ b/Engine/Source/Editor/StaticMeshEditor/StaticMeshEditor.Build.cs
@@ -44,7 +44,8 @@ public class StaticMeshEditor : ModuleRules
 				"EditorSubsystem",
 				"MeshDescription",
 				"StaticMeshDescription",
-				"ToolMenus"
+				"ToolMenus",
+				"DetailCustomizations",
 			}
 		);
 
