Update mobility settings style in the details panel

diff --git a/Engine/Source/Editor/DetailCustomizations/Private/MobilityCustomization.cpp b/Engine/Source/Editor/DetailCustomizations/Private/MobilityCustomization.cpp
--- a/Engine/Source/Editor/DetailCustomizations/Private/MobilityCustomization.cpp
+++ b/Engine/Source/Editor/DetailCustomizations/Private/MobilityCustomization.cpp
@@ -55,7 +55,6 @@ void FMobilityCustomization::CreateMobilityCustomization(IDetailCategoryBuilder&
 		[
 			SNew(STextBlock)
 			.Font(FAppStyle::Get().GetFontStyle("PropertyWindow.MobilityFont"))
-			.TransformPolicy(ETextTransformPolicy::ToUpper)
 			.Text(LOCTEXT("Static", "Static"))
 		]
 		.ToolTip(StaticTooltip);
@@ -72,7 +71,6 @@ void FMobilityCustomization::CreateMobilityCustomization(IDetailCategoryBuilder&
 		[
 			SNew(STextBlock)
 			.Font(FAppStyle::Get().GetFontStyle("PropertyWindow.MobilityFont"))
-			.TransformPolicy(ETextTransformPolicy::ToUpper)
 			.Text(LOCTEXT("Stationary", "Stationary"))
 		]
 		.ToolTip(StationaryTooltip);
@@ -87,7 +85,6 @@ void FMobilityCustomization::CreateMobilityCustomization(IDetailCategoryBuilder&
 	[
 		SNew(STextBlock)
 		.Font(FAppStyle::Get().GetFontStyle("PropertyWindow.MobilityFont"))
-		.TransformPolicy(ETextTransformPolicy::ToUpper)
 		.Text(LOCTEXT("Movable", "Movable"))
 	]
 	.ToolTip(MovableTooltip);
diff --git a/Engine/Source/Editor/EditorStyle/Private/StarshipStyle.cpp b/Engine/Source/Editor/EditorStyle/Private/StarshipStyle.cpp
--- a/Engine/Source/Editor/EditorStyle/Private/StarshipStyle.cpp
+++ b/Engine/Source/Editor/EditorStyle/Private/StarshipStyle.cpp
@@ -3105,7 +3105,7 @@ void FStarshipEditorStyle::FStyle::SetupPropertyEditorStyles()
 		Set( "PropertyWindow.ItalicFont", DEFAULT_FONT( "Italic", 8 ) );
 		Set( "PropertyWindow.FilterFont", DEFAULT_FONT( "Regular", 10 ) );
 
-		FSlateFontInfo MobilityFont = FStyleFonts::Get().SmallBold;
+		FSlateFontInfo MobilityFont = FStyleFonts::Get().Small;
 		MobilityFont.LetterSpacing = 100;
 
 		Set("PropertyWindow.MobilityFont", MobilityFont);
