Fix up deprecated color block bIgnoreAlpha usage

diff --git a/Engine/Plugins/Compositing/Composure/Source/ComposureLayersEditor/Private/Widgets/SCompElementPickerWindow.cpp b/Engine/Plugins/Compositing/Composure/Source/ComposureLayersEditor/Private/Widgets/SCompElementPickerWindow.cpp
--- a/Engine/Plugins/Compositing/Composure/Source/ComposureLayersEditor/Private/Widgets/SCompElementPickerWindow.cpp
+++ b/Engine/Plugins/Compositing/Composure/Source/ComposureLayersEditor/Private/Widgets/SCompElementPickerWindow.cpp
@@ -301,7 +301,7 @@ TSharedRef<SWidget> SCompElementColorPickerDialog::GeneratePreviewContent()
 			[
 				SNew(SColorBlock)
 					.Color(this, &SCompElementColorPickerDialog::GetPickedColor)
-					.IgnoreAlpha(true)
+					.AlphaDisplayMode(EColorBlockAlphaDisplayMode::Ignore)
 					.Size(FVector2D(16.0f, 16.0f))
 					.Visibility(this, &SCompElementColorPickerDialog::GetAcceptButtonVisibility)
 			]
diff --git a/Engine/Plugins/FX/Niagara/Source/NiagaraEditor/Private/Customizations/NiagaraScriptDetails.cpp b/Engine/Plugins/FX/Niagara/Source/NiagaraEditor/Private/Customizations/NiagaraScriptDetails.cpp
--- a/Engine/Plugins/FX/Niagara/Source/NiagaraEditor/Private/Customizations/NiagaraScriptDetails.cpp
+++ b/Engine/Plugins/FX/Niagara/Source/NiagaraEditor/Private/Customizations/NiagaraScriptDetails.cpp
@@ -426,7 +426,6 @@ private:
 				SAssignNew(ColorBlock, SColorBlock)
 				.Color(this, &SNiagaraScriptHighlightPicker::GetHighlightColor)
 				.ShowBackgroundForAlpha(true)
-				.IgnoreAlpha(false)
 				.OnMouseButtonDown(this, &SNiagaraScriptHighlightPicker::OnMouseButtonDownColorBlock)
 				.Size(FVector2D(EditColorWidth, EditColorHeight))
 			]
diff --git a/Engine/Plugins/FX/Niagara/Source/NiagaraEditor/Private/TypeEditorUtilities/NiagaraColorTypeEditorUtilities.cpp b/Engine/Plugins/FX/Niagara/Source/NiagaraEditor/Private/TypeEditorUtilities/NiagaraColorTypeEditorUtilities.cpp
--- a/Engine/Plugins/FX/Niagara/Source/NiagaraEditor/Private/TypeEditorUtilities/NiagaraColorTypeEditorUtilities.cpp
+++ b/Engine/Plugins/FX/Niagara/Source/NiagaraEditor/Private/TypeEditorUtilities/NiagaraColorTypeEditorUtilities.cpp
@@ -35,7 +35,6 @@ public:
 				SAssignNew(ColorBlock, SColorBlock)
 				.Color(this, &SNiagaraColorParameterEditor::GetColor)
 				.ShowBackgroundForAlpha(true)
-				.IgnoreAlpha(false)
 				.OnMouseButtonDown(this, &SNiagaraColorParameterEditor::OnMouseButtonDownColorBlock)
 				.Size(FVector2D(35.0f, 12.0f))
 			]
diff --git a/Engine/Plugins/Media/TimecodeSynchronizer/Source/TimecodeSynchronizerEditor/Private/AssetEditor/TimecodeSynchronizerEditorToolkit.cpp b/Engine/Plugins/Media/TimecodeSynchronizer/Source/TimecodeSynchronizerEditor/Private/AssetEditor/TimecodeSynchronizerEditorToolkit.cpp
--- a/Engine/Plugins/Media/TimecodeSynchronizer/Source/TimecodeSynchronizerEditor/Private/AssetEditor/TimecodeSynchronizerEditorToolkit.cpp
+++ b/Engine/Plugins/Media/TimecodeSynchronizer/Source/TimecodeSynchronizerEditor/Private/AssetEditor/TimecodeSynchronizerEditorToolkit.cpp
@@ -235,7 +235,7 @@ TSharedRef<SDockTab> FTimecodeSynchronizerEditorToolkit::SpawnPropertiesTab(cons
 				[
 					SNew(SColorBlock)
 					.Color(this, &FTimecodeSynchronizerEditorToolkit::GetProgressColor)
-					.IgnoreAlpha(true)
+					.AlphaDisplayMode(EColorBlockAlphaDisplayMode::Ignore)
 					.Visibility_Lambda([this]() -> EVisibility
 					{
 						if (UTimecodeSynchronizer* Asset = GetTimecodeSynchronizer())
diff --git a/Engine/Plugins/MovieScene/LevelSequenceEditor/Source/LevelSequenceEditor/Private/CinematicViewport/FilmOverlays.cpp b/Engine/Plugins/MovieScene/LevelSequenceEditor/Source/LevelSequenceEditor/Private/CinematicViewport/FilmOverlays.cpp
--- a/Engine/Plugins/MovieScene/LevelSequenceEditor/Source/LevelSequenceEditor/Private/CinematicViewport/FilmOverlays.cpp
+++ b/Engine/Plugins/MovieScene/LevelSequenceEditor/Source/LevelSequenceEditor/Private/CinematicViewport/FilmOverlays.cpp
@@ -57,7 +57,6 @@ namespace WidgetHelpers
 				SNew(SColorBlock)
 				.Color_Lambda(GetValue)
 				.ShowBackgroundForAlpha(true)
-				.IgnoreAlpha(false)
 				.Size(FVector2D(10.0f, 10.0f))
 			];
 	}
diff --git a/Engine/Source/Editor/DetailCustomizations/Private/ColorGradingVectorCustomization.cpp b/Engine/Source/Editor/DetailCustomizations/Private/ColorGradingVectorCustomization.cpp
--- a/Engine/Source/Editor/DetailCustomizations/Private/ColorGradingVectorCustomization.cpp
+++ b/Engine/Source/Editor/DetailCustomizations/Private/ColorGradingVectorCustomization.cpp
@@ -600,7 +600,7 @@ void FColorGradingVectorCustomization::MakeHeaderRow(FDetailWidgetRow& Row, TSha
 					SNew(SColorBlock)
 					.Color(this, &FColorGradingVectorCustomization::OnGetHeaderColorBlock)
 					.ShowBackgroundForAlpha(false)
-					.IgnoreAlpha(true)
+					.AlphaDisplayMode(EColorBlockAlphaDisplayMode::Ignore)
 					.ColorIsHSV(false)
 					.Size(FVector2D(70.0f, 12.0f))
 				]
diff --git a/Engine/Source/Editor/DetailCustomizations/Private/RawDistributionVectorStructCustomization.cpp b/Engine/Source/Editor/DetailCustomizations/Private/RawDistributionVectorStructCustomization.cpp
--- a/Engine/Source/Editor/DetailCustomizations/Private/RawDistributionVectorStructCustomization.cpp
+++ b/Engine/Source/Editor/DetailCustomizations/Private/RawDistributionVectorStructCustomization.cpp
@@ -330,7 +330,7 @@ TSharedRef<SWidget> FReplaceVectorWithLinearColorBuilder::CreateColorWidget(cons
 			SAssignNew(ColorPickerParentWidget, SColorBlock)
 			.Color(this, &FReplaceVectorWithLinearColorBuilder::OnGetColorForColorBlock, StructHandle)
 			.ShowBackgroundForAlpha(false)
-			.IgnoreAlpha(true)
+			.AlphaDisplayMode(EColorBlockAlphaDisplayMode::Ignore)
 			.OnMouseButtonDown(this, &FReplaceVectorWithLinearColorBuilder::OnMouseButtonDownColorBlock, StructHandle)
 			.Size(FVector2D(70.0f, 12.0f))
 		];
diff --git a/Engine/Source/Editor/DetailCustomizations/Public/Customizations/ColorStructCustomization.h b/Engine/Source/Editor/DetailCustomizations/Public/Customizations/ColorStructCustomization.h
--- a/Engine/Source/Editor/DetailCustomizations/Public/Customizations/ColorStructCustomization.h
+++ b/Engine/Source/Editor/DetailCustomizations/Public/Customizations/ColorStructCustomization.h
@@ -12,6 +12,7 @@
 
 class FDetailWidgetRow;
 class SColorPicker;
+class SBorder;
 
 /**
  * Base class for color struct customization (FColor,FLinearColor).
@@ -27,11 +28,12 @@ public:
 protected:
 
 	FColorStructCustomization()
-		: bIgnoreAlpha(false)
+		: ActiveTransaction(INDEX_NONE)
+		, bIgnoreAlpha(false)
 		, bIsInlineColorPickerVisible(false)
 		, bIsInteractive(false)
 		, bDontUpdateWhileEditing(false)
-		, ActiveTransaction(INDEX_NONE)
+
 	{}
 
 protected:
@@ -98,6 +100,15 @@ protected:
 	 */
 	FLinearColor OnGetColorForColorBlock() const;
 	
+	/**
+	 * @return The color that should be displayed in the color block in slate color format                                                           
+	 */
+	FSlateColor  OnGetSlateColorForBlock() const;
+
+	/**
+	 * @return The border color encompassing the entire color block                                                         
+	 */
+	FSlateColor GetColorWidgetBorderColor() const;
 	/**
 	 * Called when the user clicks in the color block (opens inline color picker)
 	 */
@@ -139,6 +150,16 @@ protected:
 	/** Color struct handle */
 	TSharedPtr<IPropertyHandle> StructPropertyHandle;
 
+	/** Cached widget for the color picker to use as a parent */
+	TSharedPtr<SWidget> ColorPickerParentWidget;
+
+	TSharedPtr<SWidget > ColorWidgetBackgroundBorder;
+
+	int32 ActiveTransaction;
+
+	/** Overrides the default state of the sRGB check box */
+	TOptional<bool> sRGBOverride;
+
 	/** True if the property is a linear color property */
 	bool bIsLinearColor;
 
@@ -151,14 +172,10 @@ protected:
 	/** True if the user is performing an interactive color change */
 	bool bIsInteractive;
 
-	/** Cached widget for the color picker to use as a parent */
-	TSharedPtr<SWidget> ColorPickerParentWidget;
 
 	/** The value won;t be updated while editing */
 	bool bDontUpdateWhileEditing;
 
-	/** Overrides the default state of the sRGB check box */
-	TOptional<bool> sRGBOverride;
 
-	int32 ActiveTransaction;
+	
 };
diff --git a/Engine/Source/Editor/Kismet/Private/BlueprintDetailsCustomization.cpp b/Engine/Source/Editor/Kismet/Private/BlueprintDetailsCustomization.cpp
--- a/Engine/Source/Editor/Kismet/Private/BlueprintDetailsCustomization.cpp
+++ b/Engine/Source/Editor/Kismet/Private/BlueprintDetailsCustomization.cpp
@@ -3421,7 +3421,7 @@ void FBlueprintGraphActionDetails::CustomizeDetails( IDetailLayoutBuilder& Detai
 				[
 					SAssignNew( ColorBlock, SColorBlock )
 						.Color( this, &FBlueprintGraphActionDetails::GetNodeTitleColor )
-						.IgnoreAlpha(true)
+						.AlphaDisplayMode(EColorBlockAlphaDisplayMode::Ignore)
 						.OnMouseButtonDown( this, &FBlueprintGraphActionDetails::ColorBlock_OnMouseButtonDown )
 				];
 			}
diff --git a/Engine/Source/Editor/Persona/Private/AnimTimeline/AnimTimelineTrack_FloatCurve.cpp b/Engine/Source/Editor/Persona/Private/AnimTimeline/AnimTimelineTrack_FloatCurve.cpp
--- a/Engine/Source/Editor/Persona/Private/AnimTimeline/AnimTimelineTrack_FloatCurve.cpp
+++ b/Engine/Source/Editor/Persona/Private/AnimTimeline/AnimTimelineTrack_FloatCurve.cpp
@@ -367,7 +367,7 @@ void FAnimTimelineTrack_FloatCurve::AddCurveTrackButton(TSharedPtr<SHorizontalBo
 			SNew(SColorBlock)
 			.Color_Lambda(GetValue)
 			.ShowBackgroundForAlpha(false)
-			.IgnoreAlpha(true)
+			.AlphaDisplayMode(EColorBlockAlphaDisplayMode::Ignore)
 			.Size(FVector2D(OutlinerRightPadding - 2.0f, OutlinerRightPadding))
 		]
 	];
diff --git a/Engine/Source/Editor/PixelInspector/Private/PixelInspectorDetailsCustomization.cpp b/Engine/Source/Editor/PixelInspector/Private/PixelInspectorDetailsCustomization.cpp
--- a/Engine/Source/Editor/PixelInspector/Private/PixelInspectorDetailsCustomization.cpp
+++ b/Engine/Source/Editor/PixelInspector/Private/PixelInspectorDetailsCustomization.cpp
@@ -91,8 +91,7 @@ TSharedRef<SColorBlock> FPixelInspectorDetailsCustomization::CreateColorCell(int
 	int32 SquareSize = FMath::FloorToInt(80.0f / (float)FinalColorContextGridSize);
 	return SNew(SColorBlock)
 		.Color(CellColor)
-		.ShowBackgroundForAlpha(false)
-		.IgnoreAlpha(true)
+		.AlphaDisplayMode(EColorBlockAlphaDisplayMode::Ignore)
 		.Size(FVector2D(SquareSize, SquareSize))
 		.OnMouseButtonDown(this, &FPixelInspectorDetailsCustomization::HandleColorCellMouseButtonDown, RowIndex, ColumnIndex);
 }
diff --git a/Engine/Source/Editor/PropertyEditor/Private/UserInterface/PropertyEditor/SPropertyEditorColor.cpp b/Engine/Source/Editor/PropertyEditor/Private/UserInterface/PropertyEditor/SPropertyEditorColor.cpp
--- a/Engine/Source/Editor/PropertyEditor/Private/UserInterface/PropertyEditor/SPropertyEditorColor.cpp
+++ b/Engine/Source/Editor/PropertyEditor/Private/UserInterface/PropertyEditor/SPropertyEditorColor.cpp
@@ -46,7 +46,7 @@ TSharedRef< SColorBlock > SPropertyEditorColor::ConstructColorBlock()
 	return SNew( SColorBlock )
 		.Color( this, &SPropertyEditorColor::OnGetColor )
 		.ShowBackgroundForAlpha(true)
-		.IgnoreAlpha(bIgnoreAlpha)
+		.AlphaDisplayMode(bIgnoreAlpha ? EColorBlockAlphaDisplayMode::Ignore : EColorBlockAlphaDisplayMode::Combined)
 		.OnMouseButtonDown( this, &SPropertyEditorColor::ColorBlock_OnMouseButtonDown )
 		.Size( FVector2D( 10.0f, 10.0f ) );
 }
@@ -57,7 +57,7 @@ TSharedRef< SColorBlock > SPropertyEditorColor::ConstructAlphaColorBlock()
 	return SNew( SColorBlock )
 		.Color( this, &SPropertyEditorColor::OnGetColor )
 		.ShowBackgroundForAlpha(false)
-		.IgnoreAlpha(true)
+		.AlphaDisplayMode(EColorBlockAlphaDisplayMode::Ignore)
 		.Visibility( this, &SPropertyEditorColor::GetVisibilityForOpaqueDisplay )
 		.OnMouseButtonDown( this, &SPropertyEditorColor::ColorBlock_OnMouseButtonDown )
 		.Size( FVector2D( 10.0f, 10.0f ) );
