Rename UE4.png to UnrealEngine.png and change source file to reflect that

diff --git a/Engine/Source/Editor/DetailCustomizations/Private/LinuxTargetSettingsDetails.cpp b/Engine/Source/Editor/DetailCustomizations/Private/LinuxTargetSettingsDetails.cpp
--- a/Engine/Source/Editor/DetailCustomizations/Private/LinuxTargetSettingsDetails.cpp
+++ b/Engine/Source/Editor/DetailCustomizations/Private/LinuxTargetSettingsDetails.cpp
@@ -128,7 +128,7 @@ static FString GetLinuxIconFilename(ELinuxImageScope::Type Scope)
 
 	if (Scope == ELinuxImageScope::Engine)
 	{
-		FString Filename = FPaths::EngineDir() / FString(TEXT("Source/Runtime/Launch/Resources")) / PlatformName / FString("UE4.png");
+		FString Filename = FPaths::EngineDir() / FString(TEXT("Source/Runtime/Launch/Resources")) / PlatformName / FString("UnrealEngine.png");
 		return FPaths::ConvertRelativePathToFull(Filename);
 	}
 	else
