Fix per-platform properties not allowing new platform overrides to be added

diff --git a/Engine/Source/Editor/DetailCustomizations/Public/PerPlatformPropertyCustomization.h b/Engine/Source/Editor/DetailCustomizations/Public/PerPlatformPropertyCustomization.h
--- a/Engine/Source/Editor/DetailCustomizations/Public/PerPlatformPropertyCustomization.h
+++ b/Engine/Source/Editor/DetailCustomizations/Public/PerPlatformPropertyCustomization.h
@@ -23,7 +23,7 @@ struct FPerPlatformPropertyCustomNodeBuilderArgs
 	TAttribute<TArray<FName>> PlatformOverrideNames;
 
 	/** Is overriding per-property enabled */
-	TAttribute<bool> IsEnabled;
+	TAttribute<bool> IsEnabled = true;
 
 	/** Callback to generate a widget for a specific platform row */
 	FOnGenerateWidget OnGenerateWidgetForPlatformRow;
diff --git a/Engine/Source/Editor/UnrealEd/Classes/Editor/PropertyEditorTestObject.h b/Engine/Source/Editor/UnrealEd/Classes/Editor/PropertyEditorTestObject.h
--- a/Engine/Source/Editor/UnrealEd/Classes/Editor/PropertyEditorTestObject.h
+++ b/Engine/Source/Editor/UnrealEd/Classes/Editor/PropertyEditorTestObject.h
@@ -672,6 +672,12 @@ class UPropertyEditorTestObject : public UObject
 	UPROPERTY(EditAnywhere, Category = AdvancedProperties)
 	FGuid Guid;
 
+	UPROPERTY(EditAnywhere, Category = AdvancedProperties)
+	FPerPlatformFloat PerPlatformFloat;
+
+	UPROPERTY(EditAnywhere, Category = AdvancedProperties)
+	FPerPlatformInt PerPlatformInt;
+
 	UPROPERTY()
 	bool bInlineEditConditionWithoutMetaToggle;
 
