Fixed crash where FCurveStructCustomization/FCurveColorCustomization could be destroyed without making sure the CurveWidget CurveOwner stopped pointing on them.

diff --git a/Engine/Source/Editor/DetailCustomizations/Private/CurveColorCustomization.cpp b/Engine/Source/Editor/DetailCustomizations/Private/CurveColorCustomization.cpp
--- a/Engine/Source/Editor/DetailCustomizations/Private/CurveColorCustomization.cpp
+++ b/Engine/Source/Editor/DetailCustomizations/Private/CurveColorCustomization.cpp
@@ -32,6 +32,11 @@ TSharedRef<IPropertyTypeCustomization> FCurveColorCustomization::MakeInstance()
 
 FCurveColorCustomization::~FCurveColorCustomization()
 {
+	if (CurveWidget.IsValid() && CurveWidget->GetCurveOwner() == this)
+	{
+		CurveWidget->SetCurveOwner(nullptr, false);
+	}
+
 	DestroyPopOutWindow();
 }
 
diff --git a/Engine/Source/Editor/DetailCustomizations/Private/CurveStructCustomization.cpp b/Engine/Source/Editor/DetailCustomizations/Private/CurveStructCustomization.cpp
--- a/Engine/Source/Editor/DetailCustomizations/Private/CurveStructCustomization.cpp
+++ b/Engine/Source/Editor/DetailCustomizations/Private/CurveStructCustomization.cpp
@@ -32,6 +32,11 @@ TSharedRef<IPropertyTypeCustomization> FCurveStructCustomization::MakeInstance()
 
 FCurveStructCustomization::~FCurveStructCustomization()
 {
+	if (CurveWidget.IsValid() && CurveWidget->GetCurveOwner() == this)
+	{
+		CurveWidget->SetCurveOwner(nullptr, false);
+	}
+
 	DestroyPopOutWindow();
 }
 
diff --git a/Engine/Source/Editor/UnrealEd/Public/SCurveEditor.h b/Engine/Source/Editor/UnrealEd/Public/SCurveEditor.h
--- a/Engine/Source/Editor/UnrealEd/Public/SCurveEditor.h
+++ b/Engine/Source/Editor/UnrealEd/Public/SCurveEditor.h
@@ -245,7 +245,7 @@ public:
 	UNREALED_API void SetZoomToFit(bool bNewZoomToFitVertical, bool bNewZoomToFitHorizontal);
 
 	/** Get the currently edited curve */
-	FCurveOwnerInterface* GetCurveOwner() const;
+	UNREALED_API FCurveOwnerInterface* GetCurveOwner() const;
 
 	/** Construct an object of type UCurveFactory and return it's reference*/
 	UNREALED_API UCurveFactory* GetCurveFactory( );
