Force audio stream caching to be enabled for UE5

diff --git a/Engine/Source/Editor/DetailCustomizations/Private/WindowsTargetSettingsDetails.cpp b/Engine/Source/Editor/DetailCustomizations/Private/WindowsTargetSettingsDetails.cpp
--- a/Engine/Source/Editor/DetailCustomizations/Private/WindowsTargetSettingsDetails.cpp
+++ b/Engine/Source/Editor/DetailCustomizations/Private/WindowsTargetSettingsDetails.cpp
@@ -291,29 +291,29 @@ void FWindowsTargetSettingsDetails::CustomizeDetails( IDetailLayoutBuilder& Deta
 	IDetailCategoryBuilder& AudioCategory = DetailBuilder.EditCategory("Audio");
 
 	// Here we add a callback when the 
-	TSharedPtr<IPropertyHandle> AudioStreamCachingPropertyHandle = DetailBuilder.GetProperty("bUseAudioStreamCaching");
-	IDetailCategoryBuilder& AudioStreamCachingCategory = DetailBuilder.EditCategory("Audio");
-	IDetailPropertyRow& AudioStreamCachingPropertyRow = AudioCategory.AddProperty(AudioStreamCachingPropertyHandle);
-	AudioStreamCachingPropertyRow.CustomWidget()
-		.NameContent()
-		[
-			AudioStreamCachingPropertyHandle->CreatePropertyNameWidget()
-		]
-		.ValueContent()
-		.MaxDesiredWidth(500.0f)
-		.MinDesiredWidth(100.0f)
-		[
-			SNew(SHorizontalBox)
-
-			+ SHorizontalBox::Slot()
-			.FillWidth(1.0f)
-			[
-				SNew(SCheckBox)
-				.OnCheckStateChanged(this, &FWindowsTargetSettingsDetails::HandleAudioStreamCachingToggled, AudioStreamCachingPropertyHandle)
-				.IsChecked(this, &FWindowsTargetSettingsDetails::GetAudioStreamCachingToggled, AudioStreamCachingPropertyHandle)
-				.ToolTipText(AudioStreamCachingPropertyHandle->GetToolTipText())
-			]
-		];
+// 	TSharedPtr<IPropertyHandle> AudioStreamCachingPropertyHandle = DetailBuilder.GetProperty("bUseAudioStreamCaching");
+// 	IDetailCategoryBuilder& AudioStreamCachingCategory = DetailBuilder.EditCategory("Audio");
+// 	IDetailPropertyRow& AudioStreamCachingPropertyRow = AudioCategory.AddProperty(AudioStreamCachingPropertyHandle);
+// 	AudioStreamCachingPropertyRow.CustomWidget()
+// 		.NameContent()
+// 		[
+// 			AudioStreamCachingPropertyHandle->CreatePropertyNameWidget()
+// 		]
+// 		.ValueContent()
+// 		.MaxDesiredWidth(500.0f)
+// 		.MinDesiredWidth(100.0f)
+// 		[
+// 			SNew(SHorizontalBox)
+// 
+// 			+ SHorizontalBox::Slot()
+// 			.FillWidth(1.0f)
+// 			[
+// 				SNew(SCheckBox)
+// 				.OnCheckStateChanged(this, &FWindowsTargetSettingsDetails::HandleAudioStreamCachingToggled, AudioStreamCachingPropertyHandle)
+// 				.IsChecked(this, &FWindowsTargetSettingsDetails::GetAudioStreamCachingToggled, AudioStreamCachingPropertyHandle)
+// 				.ToolTipText(AudioStreamCachingPropertyHandle->GetToolTipText())
+// 			]
+// 		];
 }
 
 bool FWindowsTargetSettingsDetails::HandlePreExternalIconCopy(const FString& InChosenImage)
diff --git a/Engine/Source/Runtime/Engine/Private/AudioCompressionSettingsUtils.cpp b/Engine/Source/Runtime/Engine/Private/AudioCompressionSettingsUtils.cpp
--- a/Engine/Source/Runtime/Engine/Private/AudioCompressionSettingsUtils.cpp
+++ b/Engine/Source/Runtime/Engine/Private/AudioCompressionSettingsUtils.cpp
@@ -106,7 +106,8 @@ void CacheAudioCookOverrides(FPlatformAudioCookOverrides& OutOverrides, const TC
 		OutOverrides.SoundCueCookQualityIndex = SoundCueQualityIndex;
 	}
 
-	PlatformFile.GetBool(*CategoryName, TEXT("bUseAudioStreamCaching"), OutOverrides.bUseStreamCaching);
+//	PlatformFile.GetBool(*CategoryName, TEXT("bUseAudioStreamCaching"), OutOverrides.bUseStreamCaching);
+	OutOverrides.bUseStreamCaching = true;
 
 	GConfig->GetBool(*CategoryName, TEXT("bInlineStreamedAudioChunks"), OutOverrides.bInlineStreamedAudioChunks, GEngineIni);
 
