Change KeyStructCustomization to always show all keys
Fixes an issue where Enhanced Input Input Mapping Contexts had missing keys.

diff --git a/Engine/Source/Editor/DetailCustomizations/Private/KeyStructCustomization.cpp b/Engine/Source/Editor/DetailCustomizations/Private/KeyStructCustomization.cpp
--- a/Engine/Source/Editor/DetailCustomizations/Private/KeyStructCustomization.cpp
+++ b/Engine/Source/Editor/DetailCustomizations/Private/KeyStructCustomization.cpp
@@ -60,6 +60,7 @@ void FKeyStructCustomization::CustomizeHeaderOnlyWithButton(TSharedRef<class IPr
 			.OnKeyChanged(this, &FKeyStructCustomization::OnKeyChanged)
 			.Font(StructCustomizationUtils.GetRegularFont())
 			.AllowClear(!StructPropertyHandle->GetProperty()->HasAnyPropertyFlags(CPF_NoClear))
+			.FilterBlueprintBindable(false)
 		]
 		+ SHorizontalBox::Slot()
 		.Padding(InputConstants::PropertyPadding)
