Scale lock icon moved to the right of the label.

diff --git a/Engine/Source/Editor/DetailCustomizations/Private/ComponentTransformDetails.cpp b/Engine/Source/Editor/DetailCustomizations/Private/ComponentTransformDetails.cpp
--- a/Engine/Source/Editor/DetailCustomizations/Private/ComponentTransformDetails.cpp
+++ b/Engine/Source/Editor/DetailCustomizations/Private/ComponentTransformDetails.cpp
@@ -186,9 +186,14 @@ TSharedRef<SWidget> FComponentTransformDetails::BuildTransformFieldLabel( ETrans
 		NameContent =
 			SNew(SHorizontalBox)
 			+ SHorizontalBox::Slot()
+			.VAlign(VAlign_Center)
+			[
+				NameContent
+			]
+			+ SHorizontalBox::Slot()
 			.AutoWidth()
 			.VAlign(VAlign_Center)
-			.Padding(FMargin(0.0f, 0.0f, 4.0f, 0.0f))
+			.Padding(FMargin(4.0f, 0.0f, 0.0f, 0.0f))
 			[
 				// Add a checkbox to toggle between preserving the ratio of x,y,z components of scale when a value is entered
 				SNew(SCheckBox)
@@ -202,11 +207,6 @@ TSharedRef<SWidget> FComponentTransformDetails::BuildTransformFieldLabel( ETrans
 					.Image(this, &FComponentTransformDetails::GetPreserveScaleRatioImage)
 					.ColorAndOpacity(FSlateColor::UseForeground())
 				]
-			]
-			+ SHorizontalBox::Slot()
-			.VAlign(VAlign_Center)
-			[
-				NameContent
 			];
 	}
 
