Removed unnecessary padding in BodyInstanceCustomization that was causing the right side of certain value widgets to be out of alignment.

diff --git a/Engine/Source/Editor/DetailCustomizations/Private/BodyInstanceCustomization.cpp b/Engine/Source/Editor/DetailCustomizations/Private/BodyInstanceCustomization.cpp
--- a/Engine/Source/Editor/DetailCustomizations/Private/BodyInstanceCustomization.cpp
+++ b/Engine/Source/Editor/DetailCustomizations/Private/BodyInstanceCustomization.cpp
@@ -140,28 +140,23 @@ void FBodyInstanceCustomization::AddCollisionCategory(TSharedRef<class IProperty
 	.ValueContent()
 	.MinDesiredWidth(131.0f)
 	[
-		SNew(SVerticalBox)
-		+ SVerticalBox::Slot()
-		.Padding(0.f, 0.f, 10.f, 0.f)
+		SNew(SHorizontalBox)
+		.IsEnabled(this, &FBodyInstanceCustomization::IsCollisionEnabled)
+		+ SHorizontalBox::Slot()
+		.VAlign(VAlign_Center)
 		[
-			SNew(SHorizontalBox)
-			.IsEnabled(this, &FBodyInstanceCustomization::IsCollisionEnabled)
-			+ SHorizontalBox::Slot()
-			.VAlign(VAlign_Center)
+			SAssignNew(CollsionProfileComboBox, SComboBox< TSharedPtr<FString> >)
+			.OptionsSource(&CollisionProfileComboList)
+			.OnGenerateWidget(this, &FBodyInstanceCustomization::MakeCollisionProfileComboWidget)
+			.OnSelectionChanged(this, &FBodyInstanceCustomization::OnCollisionProfileChanged, &CollisionGroup)
+			.OnComboBoxOpening(this, &FBodyInstanceCustomization::OnCollisionProfileComboOpening)
+			.InitiallySelectedItem(DisplayName)
+			.Content()
 			[
-				SAssignNew(CollsionProfileComboBox, SComboBox< TSharedPtr<FString> >)
-				.OptionsSource(&CollisionProfileComboList)
-				.OnGenerateWidget(this, &FBodyInstanceCustomization::MakeCollisionProfileComboWidget)
-				.OnSelectionChanged(this, &FBodyInstanceCustomization::OnCollisionProfileChanged, &CollisionGroup)
-				.OnComboBoxOpening(this, &FBodyInstanceCustomization::OnCollisionProfileComboOpening)
-				.InitiallySelectedItem(DisplayName)
-				.Content()
-				[
-					SNew(STextBlock)
-					.Text(this, &FBodyInstanceCustomization::GetCollisionProfileComboBoxContent)
-					.Font(IDetailLayoutBuilder::GetDetailFont())
-					.ToolTipText(this, &FBodyInstanceCustomization::GetCollisionProfileComboBoxToolTip)
-				]
+				SNew(STextBlock)
+				.Text(this, &FBodyInstanceCustomization::GetCollisionProfileComboBoxContent)
+				.Font(IDetailLayoutBuilder::GetDetailFont())
+				.ToolTipText(this, &FBodyInstanceCustomization::GetCollisionProfileComboBoxToolTip)
 			]
 		]
 	];
@@ -1537,16 +1532,11 @@ void FBodyInstanceCustomizationHelper::AddMassInKg(IDetailCategoryBuilder& Physi
 		]
 		.ValueContent()
 		[
-			SNew(SVerticalBox)
-			+ SVerticalBox::Slot()
-			.Padding(0.f, 0.f, 10.f, 0.f)
-			[
-				SNew(SNumericEntryBox<float>)
-				.IsEnabled(this, &FBodyInstanceCustomizationHelper::IsBodyMassEnabled)
-				.Font(IDetailLayoutBuilder::GetDetailFont())
-				.Value(this, &FBodyInstanceCustomizationHelper::OnGetBodyMass)
-				.OnValueCommitted(this, &FBodyInstanceCustomizationHelper::OnSetBodyMass)
-			]
+			SNew(SNumericEntryBox<float>)
+			.IsEnabled(this, &FBodyInstanceCustomizationHelper::IsBodyMassEnabled)
+			.Font(IDetailLayoutBuilder::GetDetailFont())
+			.Value(this, &FBodyInstanceCustomizationHelper::OnGetBodyMass)
+			.OnValueCommitted(this, &FBodyInstanceCustomizationHelper::OnSetBodyMass)
 		];
 	}
 	else
@@ -1568,7 +1558,7 @@ void FBodyInstanceCustomizationHelper::AddMaxAngularVelocity(IDetailCategoryBuil
 	[
 		SNew(SVerticalBox)
 		+ SVerticalBox::Slot()
-		.Padding(0.f, 0.f, 10.f, 0.f)
+		.VAlign(VAlign_Center)
 		[
 			SNew(SNumericEntryBox<float>)
 			.IsEnabled(false)
@@ -1578,7 +1568,7 @@ void FBodyInstanceCustomizationHelper::AddMaxAngularVelocity(IDetailCategoryBuil
 		]
 
 		+ SVerticalBox::Slot()
-		.Padding(0.f, 0.f, 10.f, 0.f)
+		.VAlign(VAlign_Center)
 		[
 			SNew(SVerticalBox)
 			.Visibility(this, &FBodyInstanceCustomizationHelper::IsMaxAngularVelocityVisible, true)
