Add support for ClampMin and ClampMax on FInterval
#jria UE-108367

diff --git a/Engine/Source/Editor/DetailCustomizations/Private/IntervalStructCustomization.cpp b/Engine/Source/Editor/DetailCustomizations/Private/IntervalStructCustomization.cpp
--- a/Engine/Source/Editor/DetailCustomizations/Private/IntervalStructCustomization.cpp
+++ b/Engine/Source/Editor/DetailCustomizations/Private/IntervalStructCustomization.cpp
@@ -8,6 +8,7 @@
 #include "DetailLayoutBuilder.h"
 #include "Widgets/Input/SNumericEntryBox.h"
 
+#include <limits>
 
 #define LOCTEXT_NAMESPACE "IntervalStructCustomization"
 
@@ -67,15 +68,30 @@ void FIntervalStructCustomization<NumericType>::CustomizeHeader(TSharedRef<IProp
 	auto Property = StructPropertyHandle->GetProperty();
 	check(Property != nullptr);
 
-	if (Property->HasMetaData(TEXT("UIMin")))
-	{
-		MinAllowedValue = TOptional<NumericType>(IntervalMetadata::FGetHelper<NumericType>::GetMetaData(Property, TEXT("UIMin")));
-	}
+	const FString& MetaUIMinString = Property->GetMetaData(TEXT("UIMin"));
+	const FString& MetaUIMaxString = Property->GetMetaData(TEXT("UIMax"));
+	const FString& MetaClampMinString = Property->GetMetaData(TEXT("ClampMin"));
+	const FString& MetaClampMaxString = Property->GetMetaData(TEXT("ClampMax"));
+	const FString& UIMinString = MetaUIMinString.Len() ? MetaUIMinString : MetaClampMinString;
+	const FString& UIMaxString = MetaUIMaxString.Len() ? MetaUIMaxString : MetaClampMaxString;
 
-	if (Property->HasMetaData(TEXT("UIMax")))
-	{
-		MaxAllowedValue = TOptional<NumericType>(IntervalMetadata::FGetHelper<NumericType>::GetMetaData(Property, TEXT("UIMax")));
-	}
+	NumericType ClampMin = std::numeric_limits<NumericType>::lowest();
+	NumericType ClampMax = std::numeric_limits<NumericType>::max();
+	TTypeFromString<NumericType>::FromString(ClampMin, *MetaClampMinString);
+	TTypeFromString<NumericType>::FromString(ClampMax, *MetaClampMaxString);
+
+	NumericType UIMin = std::numeric_limits<NumericType>::lowest();
+	NumericType UIMax = std::numeric_limits<NumericType>::max();
+	TTypeFromString<NumericType>::FromString(UIMin, *UIMinString);
+	TTypeFromString<NumericType>::FromString(UIMax, *UIMaxString);
+
+	const NumericType ActualUIMin = FMath::Max(UIMin, ClampMin);
+	const NumericType ActualUIMax = FMath::Min(UIMax, ClampMax);
+
+	MinAllowedValue = MetaClampMinString.Len() ? ClampMin : TOptional<NumericType>();
+	MaxAllowedValue = MetaClampMaxString.Len() ? ClampMax : TOptional<NumericType>();
+	MinAllowedSliderValue = (UIMinString.Len()) ? ActualUIMin : TOptional<NumericType>();
+	MaxAllowedSliderValue = (UIMaxString.Len()) ? ActualUIMax : TOptional<NumericType>();
 
 	bAllowInvertedInterval = Property->HasMetaData(TEXT("AllowInvertedInterval"));
 	bClampToMinMaxLimits = Property->HasMetaData(TEXT("ClampToMinMaxLimits"));
@@ -106,9 +122,9 @@ void FIntervalStructCustomization<NumericType>::CustomizeHeader(TSharedRef<IProp
 			SNew(SNumericEntryBox<NumericType>)
 			.Value(this, &FIntervalStructCustomization<NumericType>::OnGetValue, EIntervalField::Min)
 			.MinValue(MinAllowedValue)
-			.MinSliderValue(MinAllowedValue)
+			.MinSliderValue(MinAllowedSliderValue)
 			.MaxValue(this, &FIntervalStructCustomization<NumericType>::OnGetMaxValue)
-			.MaxSliderValue(this, &FIntervalStructCustomization<NumericType>::OnGetMaxValue)
+			.MaxSliderValue(this, &FIntervalStructCustomization<NumericType>::OnGetMaxSliderValue)
 			.OnValueCommitted(this, &FIntervalStructCustomization<NumericType>::OnValueCommitted, EIntervalField::Min)
 			.OnValueChanged(this, &FIntervalStructCustomization<NumericType>::OnValueChanged, EIntervalField::Min)
 			.OnBeginSliderMovement(this, &FIntervalStructCustomization<NumericType>::OnBeginSliderMovement)
@@ -134,9 +150,9 @@ void FIntervalStructCustomization<NumericType>::CustomizeHeader(TSharedRef<IProp
 			SNew(SNumericEntryBox<NumericType>)
 			.Value(this, &FIntervalStructCustomization<NumericType>::OnGetValue, EIntervalField::Max)
 			.MinValue(this, &FIntervalStructCustomization<NumericType>::OnGetMinValue)
-			.MinSliderValue(this, &FIntervalStructCustomization<NumericType>::OnGetMinValue)
+			.MinSliderValue(this, &FIntervalStructCustomization<NumericType>::OnGetMinSliderValue)
 			.MaxValue(MaxAllowedValue)
-			.MaxSliderValue(MaxAllowedValue)
+			.MaxSliderValue(MaxAllowedSliderValue)
 			.OnValueCommitted(this, &FIntervalStructCustomization<NumericType>::OnValueCommitted, EIntervalField::Max)
 			.OnValueChanged(this, &FIntervalStructCustomization<NumericType>::OnValueChanged, EIntervalField::Max)
 			.OnBeginSliderMovement(this, &FIntervalStructCustomization<NumericType>::OnBeginSliderMovement)
@@ -189,6 +205,17 @@ TOptional<NumericType> FIntervalStructCustomization<NumericType>::OnGetMinValue(
 	return MinAllowedValue;
 }
 
+template <typename NumericType>
+TOptional<NumericType> FIntervalStructCustomization<NumericType>::OnGetMinSliderValue() const
+{
+	if (bClampToMinMaxLimits)
+	{
+		return GetValue<NumericType>(MinValueHandle.Get());
+	}
+
+	return MinAllowedSliderValue;
+}
+
 template <typename NumericType>
 TOptional<NumericType> FIntervalStructCustomization<NumericType>::OnGetMaxValue() const
 {
@@ -200,6 +227,17 @@ TOptional<NumericType> FIntervalStructCustomization<NumericType>::OnGetMaxValue(
 	return MaxAllowedValue;
 }
 
+template <typename NumericType>
+TOptional<NumericType> FIntervalStructCustomization<NumericType>::OnGetMaxSliderValue() const
+{
+	if (bClampToMinMaxLimits)
+	{
+		return GetValue<NumericType>(MaxValueHandle.Get());
+	}
+
+	return MaxAllowedSliderValue;
+}
+
 template <typename NumericType>
 void FIntervalStructCustomization<NumericType>::SetValue(NumericType NewValue, EIntervalField Field, EPropertyValueSetFlags::Type Flags)
 {
diff --git a/Engine/Source/Editor/DetailCustomizations/Private/IntervalStructCustomization.h b/Engine/Source/Editor/DetailCustomizations/Private/IntervalStructCustomization.h
--- a/Engine/Source/Editor/DetailCustomizations/Private/IntervalStructCustomization.h
+++ b/Engine/Source/Editor/DetailCustomizations/Private/IntervalStructCustomization.h
@@ -59,6 +59,13 @@ protected:
 	 */
 	TOptional<NumericType> OnGetMinValue() const;
 
+	/**
+	 * Gets the minimum slider allowed value
+	 *
+	 * @return	The value or unset if it could not be accessed
+	 */
+	TOptional<NumericType> OnGetMinSliderValue() const;
+
 	/**
 	 * Gets the maximum allowed value
 	 *
@@ -66,6 +73,13 @@ protected:
 	 */
 	TOptional<NumericType> OnGetMaxValue() const;
 
+	/**
+	 * Gets the maximum slider allowed value
+	 *
+	 * @return	The value or unset if it could not be accessed
+	 */
+	TOptional<NumericType> OnGetMaxSliderValue() const;
+
 	/**
 	 * Called when the value is committed from the property editor
 	 *
@@ -121,6 +135,8 @@ protected:
 	// Min and max allowed values from the metadata
 	TOptional<NumericType> MinAllowedValue;
 	TOptional<NumericType> MaxAllowedValue;
+	TOptional<NumericType> MinAllowedSliderValue;
+	TOptional<NumericType> MaxAllowedSliderValue;
 	
 	// Flags whether the slider is being moved at the moment on any of our widgets
 	bool bIsUsingSlider;
