Removed limit on width for a selected Actor's level name.

diff --git a/Engine/Source/Editor/DetailCustomizations/Private/ActorDetails.cpp b/Engine/Source/Editor/DetailCustomizations/Private/ActorDetails.cpp
--- a/Engine/Source/Editor/DetailCustomizations/Private/ActorDetails.cpp
+++ b/Engine/Source/Editor/DetailCustomizations/Private/ActorDetails.cpp
@@ -519,16 +519,17 @@ void FActorDetails::AddActorCategory( IDetailLayoutBuilder& DetailBuilder, const
 		[
 			SNew(SHyperlink)
 				.Style(FEditorStyle::Get(), "HoverOnlyHyperlink")
-				.OnNavigate( this, &FActorDetails::OnNarrowSelectionSetToSpecificLevel, WeakLevelPtr )
+				.OnNavigate(this, &FActorDetails::OnNarrowSelectionSetToSpecificLevel, WeakLevelPtr)
 				.Text(ActorCountDescription)
 				.TextStyle(FEditorStyle::Get(), "DetailsView.HyperlinkStyle")
 				.ToolTipText(Tooltip)
 		]
 		.ValueContent()
+		.MaxDesiredWidth(0)
 		[
 			SNew(STextBlock)
 				.Text(LevelDescription)
-				.Font( IDetailLayoutBuilder::GetDetailFont() )
+				.Font(IDetailLayoutBuilder::GetDetailFont())
 		];
 
 	}
