Removed unnecessary padding around the Collision Presets combo box.

diff --git a/Engine/Source/Editor/DetailCustomizations/Private/BodyInstanceCustomization.cpp b/Engine/Source/Editor/DetailCustomizations/Private/BodyInstanceCustomization.cpp
--- a/Engine/Source/Editor/DetailCustomizations/Private/BodyInstanceCustomization.cpp
+++ b/Engine/Source/Editor/DetailCustomizations/Private/BodyInstanceCustomization.cpp
@@ -138,7 +138,6 @@ void FBodyInstanceCustomization::AddCollisionCategory(TSharedRef<class IProperty
 		]
 	]
 	.ValueContent()
-	.MinDesiredWidth(131.0f)
 	[
 		SNew(SHorizontalBox)
 		.IsEnabled(this, &FBodyInstanceCustomization::IsCollisionEnabled)
