Spline Component Icons

diff --git a/Engine/Source/Editor/DetailCustomizations/Private/SplineComponentDetails.cpp b/Engine/Source/Editor/DetailCustomizations/Private/SplineComponentDetails.cpp
--- a/Engine/Source/Editor/DetailCustomizations/Private/SplineComponentDetails.cpp
+++ b/Engine/Source/Editor/DetailCustomizations/Private/SplineComponentDetails.cpp
@@ -292,104 +292,115 @@ void FSplinePointDetails::GenerateHeaderRowContent(FDetailWidgetRow& NodeRow)
 
 void FSplinePointDetails::GenerateSplinePointSelectionControls(IDetailChildrenBuilder& ChildrenBuilder)
 {
+	FMargin ButtonPadding(2.f, 0.f);
+
 	ChildrenBuilder.AddCustomRow(LOCTEXT("SelectSplinePoints", "Select Spline Points"))
-		.NameContent()
-		[
-			SNew(STextBlock)
-			.Font(IDetailLayoutBuilder::GetDetailFont())
+	.NameContent()
+	[
+		SNew(STextBlock)
+		.Font(IDetailLayoutBuilder::GetDetailFont())
 		.Text(LOCTEXT("SelectSplinePoints", "Select Spline Points"))
-		]
+	]
 	.ValueContent()
-		.MaxDesiredWidth(125.f)
-		.MinDesiredWidth(125.f)
-		[
-			SNew(SHorizontalBox)
-			.Clipping(EWidgetClipping::ClipToBounds)
+	.VAlign(VAlign_Fill)
+	.MaxDesiredWidth(170.f)
+	.MinDesiredWidth(170.f)
+	[
+		SNew(SHorizontalBox)
+		.Clipping(EWidgetClipping::ClipToBounds)
+
 		+ SHorizontalBox::Slot()
-		.HAlign(HAlign_Center)
 		.VAlign(VAlign_Center)
+		.AutoWidth()
+		.Padding(ButtonPadding)
 		[
 			SNew(SButton)
 			.ButtonStyle(FEditorStyle::Get(), "SplineComponentDetails.SelectFirst")
-		.ContentPadding(2.0f)
-		.ToolTipText(LOCTEXT("SelectFirstSplinePointToolTip", "Select first spline point."))
-		.OnClicked(this, &FSplinePointDetails::OnSelectFirstLastSplinePoint, true)
+			.ContentPadding(2.0f)
+			.ToolTipText(LOCTEXT("SelectFirstSplinePointToolTip", "Select first spline point."))
+			.OnClicked(this, &FSplinePointDetails::OnSelectFirstLastSplinePoint, true)
 		]
-	+ SHorizontalBox::Slot()
-		.HAlign(HAlign_Center)
+		+ SHorizontalBox::Slot()
+		.AutoWidth()
 		.VAlign(VAlign_Center)
+		.Padding(ButtonPadding)
 		[
 			SNew(SButton)
 			.ButtonStyle(FEditorStyle::Get(), "SplineComponentDetails.AddPrev")
-		.ContentPadding(2.f)
-		.ToolTipText(LOCTEXT("SelectAddPrevSplinePointToolTip", "Add previous spline point to current selection."))
-		.OnClicked(this, &FSplinePointDetails::OnSelectPrevNextSplinePoint, false, true)
-		.IsEnabled(this, &FSplinePointDetails::ArePointsSelected)
+			.ContentPadding(2.f)
+			.ToolTipText(LOCTEXT("SelectAddPrevSplinePointToolTip", "Add previous spline point to current selection."))
+			.OnClicked(this, &FSplinePointDetails::OnSelectPrevNextSplinePoint, false, true)
+			.IsEnabled(this, &FSplinePointDetails::ArePointsSelected)
 		]
-	+ SHorizontalBox::Slot()
-		.HAlign(HAlign_Center)
+		+ SHorizontalBox::Slot()
+		.AutoWidth()
 		.VAlign(VAlign_Center)
+		.Padding(ButtonPadding)
 		[
 			SNew(SButton)
 			.ButtonStyle(FEditorStyle::Get(), "SplineComponentDetails.SelectPrev")
-		.ContentPadding(2.f)
-		.VAlign(VAlign_Center)
-		.HAlign(HAlign_Center)
-		.ToolTipText(LOCTEXT("SelectPrevSplinePointToolTip", "Select previous spline point."))
-		.OnClicked(this, &FSplinePointDetails::OnSelectPrevNextSplinePoint, false, false)
-		.IsEnabled(this, &FSplinePointDetails::ArePointsSelected)
+			.ContentPadding(2.f)
+			.VAlign(VAlign_Center)
+			.HAlign(HAlign_Center)
+			.ToolTipText(LOCTEXT("SelectPrevSplinePointToolTip", "Select previous spline point."))
+			.OnClicked(this, &FSplinePointDetails::OnSelectPrevNextSplinePoint, false, false)
+			.IsEnabled(this, &FSplinePointDetails::ArePointsSelected)
 		]
-	+ SHorizontalBox::Slot()
-		.HAlign(HAlign_Center)
+		+ SHorizontalBox::Slot()
+		.AutoWidth()
 		.VAlign(VAlign_Center)
+		.Padding(ButtonPadding)
 		[
 			SNew(SButton)
 			.ButtonStyle(FEditorStyle::Get(), "SplineComponentDetails.SelectAll")
-		.ContentPadding(2.f)
-		.VAlign(VAlign_Center)
-		.HAlign(HAlign_Center)
-		.ToolTipText(LOCTEXT("SelectAllSplinePointToolTip", "Select all spline points."))
-		.OnClicked(this, &FSplinePointDetails::OnSelectAllSplinePoints)
+			.ContentPadding(2.f)
+			.VAlign(VAlign_Center)
+			.HAlign(HAlign_Center)
+			.ToolTipText(LOCTEXT("SelectAllSplinePointToolTip", "Select all spline points."))
+			.OnClicked(this, &FSplinePointDetails::OnSelectAllSplinePoints)
 		]
-	+ SHorizontalBox::Slot()
-		.HAlign(HAlign_Center)
+		+ SHorizontalBox::Slot()
+		.AutoWidth()
 		.VAlign(VAlign_Center)
+		.Padding(ButtonPadding)
 		[
 			SNew(SButton)
 			.ButtonStyle(FEditorStyle::Get(), "SplineComponentDetails.SelectNext")
-		.ContentPadding(2.f)
-		.VAlign(VAlign_Center)
-		.HAlign(HAlign_Center)
-		.ToolTipText(LOCTEXT("SelectNextSplinePointToolTip", "Select next spline point."))
-		.OnClicked(this, &FSplinePointDetails::OnSelectPrevNextSplinePoint, true, false)
-		.IsEnabled(this, &FSplinePointDetails::ArePointsSelected)
+			.ContentPadding(2.f)
+			.VAlign(VAlign_Center)
+			.HAlign(HAlign_Center)
+			.ToolTipText(LOCTEXT("SelectNextSplinePointToolTip", "Select next spline point."))
+			.OnClicked(this, &FSplinePointDetails::OnSelectPrevNextSplinePoint, true, false)
+			.IsEnabled(this, &FSplinePointDetails::ArePointsSelected)
 		]
-	+ SHorizontalBox::Slot()
-		.HAlign(HAlign_Center)
+		+ SHorizontalBox::Slot()
+		.AutoWidth()
 		.VAlign(VAlign_Center)
+		.Padding(ButtonPadding)
 		[
 			SNew(SButton)
 			.ButtonStyle(FEditorStyle::Get(), "SplineComponentDetails.AddNext")
-		.ContentPadding(2.f)
-		.VAlign(VAlign_Center)
-		.HAlign(HAlign_Center)
-		.ToolTipText(LOCTEXT("SelectAddNextSplinePointToolTip", "Add next spline point to current selection."))
-		.OnClicked(this, &FSplinePointDetails::OnSelectPrevNextSplinePoint, true, true)
-		.IsEnabled(this, &FSplinePointDetails::ArePointsSelected)
+			.ContentPadding(2.f)
+			.VAlign(VAlign_Center)
+			.HAlign(HAlign_Center)
+			.ToolTipText(LOCTEXT("SelectAddNextSplinePointToolTip", "Add next spline point to current selection."))
+			.OnClicked(this, &FSplinePointDetails::OnSelectPrevNextSplinePoint, true, true)
+			.IsEnabled(this, &FSplinePointDetails::ArePointsSelected)
 		]
-	+ SHorizontalBox::Slot()
-		.HAlign(HAlign_Center)
+		+ SHorizontalBox::Slot()
+		.AutoWidth()
 		.VAlign(VAlign_Center)
+		.Padding(ButtonPadding)
 		[
 			SNew(SButton)
 			.ButtonStyle(FEditorStyle::Get(), "SplineComponentDetails.SelectLast")
-		.ContentPadding(2.f)
-		.VAlign(VAlign_Center)
-		.HAlign(HAlign_Center)
-		.ToolTipText(LOCTEXT("SelectLastSplinePointToolTip", "Select last spline point."))
-		.OnClicked(this, &FSplinePointDetails::OnSelectFirstLastSplinePoint, false)
+			.ContentPadding(2.f)
+			.VAlign(VAlign_Center)
+			.HAlign(HAlign_Center)
+			.ToolTipText(LOCTEXT("SelectLastSplinePointToolTip", "Select last spline point."))
+			.OnClicked(this, &FSplinePointDetails::OnSelectFirstLastSplinePoint, false)
 		]
-		];
+	];
 }
 
 void FSplinePointDetails::GenerateChildContent(IDetailChildrenBuilder& ChildrenBuilder)
diff --git a/Engine/Source/Editor/EditorStyle/Private/StarshipStyle.cpp b/Engine/Source/Editor/EditorStyle/Private/StarshipStyle.cpp
--- a/Engine/Source/Editor/EditorStyle/Private/StarshipStyle.cpp
+++ b/Engine/Source/Editor/EditorStyle/Private/StarshipStyle.cpp
@@ -1891,45 +1891,45 @@ void FStarshipEditorStyle::FStyle::SetupLevelGeneralStyles()
 	// Spline component controls
 	{
 		Set("SplineComponentDetails.SelectFirst", FButtonStyle(Button)
-			.SetNormal(IMAGE_BRUSH("/Sequencer/Transport_Bar/Go_To_Front_24x_OFF", Icon24x24))
-			.SetHovered(IMAGE_BRUSH("/Sequencer/Transport_Bar/Go_To_Front_24x_OFF", Icon24x24, SelectionColor))
-			.SetPressed(IMAGE_BRUSH("/Sequencer/Transport_Bar/Go_To_Front_24x", Icon24x24, SelectionColor_Pressed))
+			.SetNormal(IMAGE_BRUSH_SVG("Starship/Splines/Spline_SelectFirst", Icon20x20))
+			.SetHovered(IMAGE_BRUSH_SVG("Starship/Splines/Spline_SelectFirst", Icon20x20, SelectionColor))
+			.SetPressed(IMAGE_BRUSH_SVG("Starship/Splines/Spline_SelectFirst", Icon20x20, SelectionColor_Pressed))
 		);
 
 		Set("SplineComponentDetails.AddPrev", FButtonStyle(Button)
-			.SetNormal(IMAGE_BRUSH("/Sequencer/Transport_Bar/Step_Backwards_24x_OFF", Icon24x24))
-			.SetHovered(IMAGE_BRUSH("/Sequencer/Transport_Bar/Step_Backwards_24x_OFF", Icon24x24, SelectionColor))
-			.SetPressed(IMAGE_BRUSH("/Sequencer/Transport_Bar/Step_Backwards_24x", Icon24x24, SelectionColor_Pressed))
+			.SetNormal(IMAGE_BRUSH_SVG("Starship/Splines/Spline_AddPrevious", Icon20x20))
+			.SetHovered(IMAGE_BRUSH_SVG("Starship/Splines/Spline_AddPrevious", Icon20x20, SelectionColor))
+			.SetPressed(IMAGE_BRUSH_SVG("Starship/Splines/Spline_AddPrevious", Icon20x20, SelectionColor_Pressed))
 		);
 
 		Set("SplineComponentDetails.SelectPrev", FButtonStyle(Button)
-			.SetNormal(IMAGE_BRUSH("/Sequencer/Transport_Bar/Backwards_24x_OFF", Icon24x24))
-			.SetHovered(IMAGE_BRUSH("/Sequencer/Transport_Bar/Backwards_24x_OFF", Icon24x24, SelectionColor))
-			.SetPressed(IMAGE_BRUSH("/Sequencer/Transport_Bar/Backwards_24x", Icon24x24, SelectionColor_Pressed))
+			.SetNormal(IMAGE_BRUSH_SVG("Starship/Splines/Spline_SelectPrevious", Icon20x20))
+			.SetHovered(IMAGE_BRUSH_SVG("Starship/Splines/Spline_SelectPrevious", Icon20x20, SelectionColor))
+			.SetPressed(IMAGE_BRUSH_SVG("Starship/Splines/Spline_SelectPrevious", Icon20x20, SelectionColor_Pressed))
 		);
 
 		Set("SplineComponentDetails.SelectAll", FButtonStyle(Button)
-			.SetNormal(IMAGE_BRUSH("/Sequencer/Transport_Bar/Loop_24x_OFF", Icon24x24))
-			.SetHovered(IMAGE_BRUSH("/Sequencer/Transport_Bar/Loop_24x_OFF", Icon24x24, SelectionColor))
-			.SetPressed(IMAGE_BRUSH("/Sequencer/Transport_Bar/Loop_24x", Icon24x24, SelectionColor_Pressed))
+			.SetNormal(IMAGE_BRUSH_SVG("Starship/Splines/Spline_SelectAll", Icon20x20))
+			.SetHovered(IMAGE_BRUSH_SVG("Starship/Splines/Spline_SelectAll", Icon20x20, SelectionColor))
+			.SetPressed(IMAGE_BRUSH_SVG("Starship/Splines/Spline_SelectAll", Icon20x20, SelectionColor_Pressed))
 		);
 
 		Set("SplineComponentDetails.SelectNext", FButtonStyle(Button)
-			.SetNormal(IMAGE_BRUSH("/Sequencer/Transport_Bar/Play_24x_OFF", Icon24x24))
-			.SetHovered(IMAGE_BRUSH("/Sequencer/Transport_Bar/Play_24x_OFF", Icon24x24, SelectionColor))
-			.SetPressed(IMAGE_BRUSH("/Sequencer/Transport_Bar/Play_24x", Icon24x24, SelectionColor_Pressed))
+			.SetNormal(IMAGE_BRUSH_SVG("Starship/Splines/Spline_SelectNext", Icon20x20))
+			.SetHovered(IMAGE_BRUSH_SVG("Starship/Splines/Spline_SelectNext", Icon20x20, SelectionColor))
+			.SetPressed(IMAGE_BRUSH_SVG("Starship/Splines/Spline_SelectNext", Icon20x20, SelectionColor_Pressed))
 		);
 
 		Set("SplineComponentDetails.AddNext", FButtonStyle(Button)
-			.SetNormal(IMAGE_BRUSH("/Sequencer/Transport_Bar/Step_Forward_24x_OFF", Icon24x24))
-			.SetHovered(IMAGE_BRUSH("/Sequencer/Transport_Bar/Step_Forward_24x_OFF", Icon24x24, SelectionColor))
-			.SetPressed(IMAGE_BRUSH("/Sequencer/Transport_Bar/Step_Forward_24x", Icon24x24, SelectionColor_Pressed))
+			.SetNormal(IMAGE_BRUSH_SVG("Starship/Splines/Spline_AddNext", Icon20x20))
+			.SetHovered(IMAGE_BRUSH_SVG("Starship/Splines/Spline_AddNext", Icon20x20, SelectionColor))
+			.SetPressed(IMAGE_BRUSH_SVG("Starship/Splines/Spline_AddNext", Icon20x20, SelectionColor_Pressed))
 		);
 
 		Set("SplineComponentDetails.SelectLast", FButtonStyle(Button)
-			.SetNormal(IMAGE_BRUSH("/Sequencer/Transport_Bar/Go_To_End_24x_OFF", Icon24x24))
-			.SetHovered(IMAGE_BRUSH("/Sequencer/Transport_Bar/Go_To_End_24x_OFF", Icon24x24, SelectionColor))
-			.SetPressed(IMAGE_BRUSH("/Sequencer/Transport_Bar/Go_To_End_24x", Icon24x24, SelectionColor_Pressed))
+			.SetNormal(IMAGE_BRUSH_SVG("Starship/Splines/Spline_SelectLast", Icon20x20))
+			.SetHovered(IMAGE_BRUSH_SVG("Starship/Splines/Spline_SelectLast", Icon20x20, SelectionColor))
+			.SetPressed(IMAGE_BRUSH_SVG("Starship/Splines/Spline_SelectLast", Icon20x20, SelectionColor_Pressed))
 		);
 
 	}
