Fixed Experimental Class Warning

diff --git a/Engine/Source/Editor/DetailCustomizations/Private/ObjectDetails.cpp b/Engine/Source/Editor/DetailCustomizations/Private/ObjectDetails.cpp
--- a/Engine/Source/Editor/DetailCustomizations/Private/ObjectDetails.cpp
+++ b/Engine/Source/Editor/DetailCustomizations/Private/ObjectDetails.cpp
@@ -10,6 +10,7 @@
 #include "Widgets/Text/STextBlock.h"
 #include "Widgets/Images/SImage.h"
 #include "Widgets/Input/SButton.h"
+#include "SWarningOrErrorBox.h"
 #include "DetailLayoutBuilder.h"
 #include "DetailWidgetRow.h"
 #include "DetailCategoryBuilder.h"
@@ -47,41 +48,18 @@ void FObjectDetails::AddExperimentalWarningCategory(IDetailLayoutBuilder& Detail
 		const FText WarningText = bBaseClassIsExperimental ? FText::Format( LOCTEXT("ExperimentalClassWarning", "Uses experimental class: {0}") , FText::FromString(MostDerivedDevelopmentClassName) )
 			: FText::Format( LOCTEXT("EarlyAccessClassWarning", "Uses beta class {0}"), FText::FromString(MostDerivedDevelopmentClassName) );
 		const FText SearchString = WarningText;
-		const FText Tooltip = bBaseClassIsExperimental ? LOCTEXT("ExperimentalClassTooltip", "Here be dragons!  Uses one or more unsupported 'experimental' classes") : LOCTEXT("EarlyAccessClassTooltip", "Uses one or more 'beta' classes");
-		const FString ExcerptName = bBaseClassIsExperimental ? TEXT("ObjectUsesExperimentalClass") : TEXT("ObjectUsesEarlyAccessClass");
-		const FSlateBrush* WarningIcon = FEditorStyle::GetBrush(bBaseClassIsExperimental ? "PropertyEditor.ExperimentalClass" : "PropertyEditor.EarlyAccessClass");
 
 		IDetailCategoryBuilder& WarningCategory = DetailBuilder.EditCategory(CategoryName, CategoryDisplayName, ECategoryPriority::Transform);
 
 		FDetailWidgetRow& WarningRow = WarningCategory.AddCustomRow(SearchString)
 			.WholeRowContent()
 			[
-				SNew(SBorder)
-				.BorderImage(FEditorStyle::GetBrush("SettingsEditor.CheckoutWarningBorder"))
-				.BorderBackgroundColor(FColor (166,137,0))
+				SNew(SBox)
+				.Padding(FMargin(0.f, 4.f))
 				[
-					SNew(SHorizontalBox)
-					.ToolTip(IDocumentation::Get()->CreateToolTip(Tooltip, nullptr, TEXT("Shared/LevelEditor"), ExcerptName))
-					.Visibility(EVisibility::Visible)
-
-					+ SHorizontalBox::Slot()
-					.VAlign(VAlign_Center)
-					.AutoWidth()
-					.Padding(4.0f, 0.0f, 0.0f, 0.0f)
-					[
-						SNew(SImage)
-						.Image(WarningIcon)
-					]
-
-					+SHorizontalBox::Slot()
-					.VAlign(VAlign_Center)
-					.AutoWidth()
-					.Padding(4.0f, 0.0f, 0.0f, 0.0f)
-					[
-						SNew(STextBlock)
-						.Text(WarningText)
-						.Font(IDetailLayoutBuilder::GetDetailFont())
-					]
+					SNew(SWarningOrErrorBox)
+					.MessageStyle(EMessageStyle::Warning)
+					.Message(WarningText)
 				]
 			];
 	}
