MobilityCustomization has been refactored to be an FDetailCustomBuilder, which allows it to be favorited.

diff --git a/Engine/Plugins/Enterprise/DatasmithContent/Source/DatasmithContentEditor/Private/DatasmithAreaLightActorDetailsPanel.cpp b/Engine/Plugins/Enterprise/DatasmithContent/Source/DatasmithContentEditor/Private/DatasmithAreaLightActorDetailsPanel.cpp
--- a/Engine/Plugins/Enterprise/DatasmithContent/Source/DatasmithContentEditor/Private/DatasmithAreaLightActorDetailsPanel.cpp
+++ b/Engine/Plugins/Enterprise/DatasmithContent/Source/DatasmithContentEditor/Private/DatasmithAreaLightActorDetailsPanel.cpp
@@ -39,8 +39,7 @@ void FDatasmithAreaLightActorDetailsPanel::CustomizeDetails(IDetailLayoutBuilder
 	TSharedPtr<IPropertyHandle> MobilityHandle = DetailBuilder.GetProperty("Mobility");
 
 	uint8 RestrictedMobilityBits = 0u;
-	MobilityCustomization = MakeShareable(new FMobilityCustomization);
-	MobilityCustomization->CreateMobilityCustomization(TransformCategory, MobilityHandle, RestrictedMobilityBits, true /* bForLight */);
+	TransformCategory.AddCustomBuilder(MakeShared<FMobilityCustomization>(MobilityHandle, RestrictedMobilityBits, true /* bForLight */));
 }
 
 #undef LOCTEXT_NAMESPACE
diff --git a/Engine/Plugins/Enterprise/DatasmithContent/Source/DatasmithContentEditor/Private/DatasmithAreaLightActorDetailsPanel.h b/Engine/Plugins/Enterprise/DatasmithContent/Source/DatasmithContentEditor/Private/DatasmithAreaLightActorDetailsPanel.h
--- a/Engine/Plugins/Enterprise/DatasmithContent/Source/DatasmithContentEditor/Private/DatasmithAreaLightActorDetailsPanel.h
+++ b/Engine/Plugins/Enterprise/DatasmithContent/Source/DatasmithContentEditor/Private/DatasmithAreaLightActorDetailsPanel.h
@@ -6,7 +6,6 @@
 #include "Input/Reply.h"
 #include "IDetailCustomization.h"
 
-class FMobilityCustomization;
 class IDetailLayoutBuilder;
 
 class FDatasmithAreaLightActorDetailsPanel : public IDetailCustomization
@@ -20,7 +19,4 @@ public:
 	// IDetailCustomization interface
 	virtual void CustomizeDetails(IDetailLayoutBuilder& DetailLayout) override;
 	// End of IDetailCustomization interface
-
-private:
-	TSharedPtr<FMobilityCustomization> MobilityCustomization;
 };
diff --git a/Engine/Source/Editor/DetailCustomizations/Private/MobilityCustomization.cpp b/Engine/Source/Editor/DetailCustomizations/Private/MobilityCustomization.cpp
--- a/Engine/Source/Editor/DetailCustomizations/Private/MobilityCustomization.cpp
+++ b/Engine/Source/Editor/DetailCustomizations/Private/MobilityCustomization.cpp
@@ -9,23 +9,35 @@
 #include "EditorStyleSet.h"
 #include "DetailLayoutBuilder.h"
 #include "DetailWidgetRow.h"
-#include "IDetailPropertyRow.h"
-#include "DetailCategoryBuilder.h"
 #include "Widgets/Input/SSegmentedControl.h"
 
-#define LOCTEXT_NAMESPACE "MobilityCustomization"
+#define LOCTEXT_NAMESPACE "MobilityCustomization" 
 
-void FMobilityCustomization::CreateMobilityCustomization(IDetailCategoryBuilder& Category, TSharedPtr<IPropertyHandle> InMobilityHandle, uint8 RestrictedMobilityBits, bool bForLight)
+FMobilityCustomization::FMobilityCustomization(TSharedPtr<IPropertyHandle> InMobilityHandle, uint8 InRestrictedMobilityBits, bool InForLight)
 {
 	MobilityHandle = InMobilityHandle;
+	RestrictedMobilityBits = InRestrictedMobilityBits;
+	bForLight = InForLight;
+}
 
+FName FMobilityCustomization::GetName() const
+{
+	const FProperty* Property = MobilityHandle->GetProperty();
+	if (Property != nullptr)
+	{
+		return Property->GetFName();
+	}
+	return NAME_None;
+}
+
+void FMobilityCustomization::GenerateHeaderRowContent(FDetailWidgetRow& WidgetRow)
+{
 	TSharedRef<SSegmentedControl<EComponentMobility::Type>> ButtonOptionsPanel =
 		SNew(SSegmentedControl<EComponentMobility::Type>)
 		.Value(this, &FMobilityCustomization::GetActiveMobility)
 		.OnValueChanged(this, &FMobilityCustomization::OnMobilityChanged);
 		
-	IDetailPropertyRow& MobilityRow = Category.AddProperty(MobilityHandle);
-	MobilityRow.CustomWidget()
+	WidgetRow
 	.NameContent()
 	[
 		SNew(STextBlock)
diff --git a/Engine/Source/Editor/DetailCustomizations/Private/SceneComponentDetails.cpp b/Engine/Source/Editor/DetailCustomizations/Private/SceneComponentDetails.cpp
--- a/Engine/Source/Editor/DetailCustomizations/Private/SceneComponentDetails.cpp
+++ b/Engine/Source/Editor/DetailCustomizations/Private/SceneComponentDetails.cpp
@@ -212,8 +212,7 @@ void FSceneComponentDetails::CustomizeDetails( IDetailLayoutBuilder& DetailBuild
 		}
 	}
 
-	MobilityCustomization = MakeShareable(new FMobilityCustomization);
-	MobilityCustomization->CreateMobilityCustomization(TransformCategory, MobilityHandle, RestrictedMobilityBits, bAnySelectedIsLight);
+	TransformCategory.AddCustomBuilder(MakeShared<FMobilityCustomization>(MobilityHandle, RestrictedMobilityBits, bAnySelectedIsLight));
 
 	// Only display bHiddenInGame if the property is being flattened in to an Actor.
 	// Details panel for BP component will have the base class be the Actor due to how the SKismetInspector works, but in that case we
diff --git a/Engine/Source/Editor/DetailCustomizations/Private/SceneComponentDetails.h b/Engine/Source/Editor/DetailCustomizations/Private/SceneComponentDetails.h
--- a/Engine/Source/Editor/DetailCustomizations/Private/SceneComponentDetails.h
+++ b/Engine/Source/Editor/DetailCustomizations/Private/SceneComponentDetails.h
@@ -5,7 +5,6 @@
 #include "CoreMinimal.h"
 #include "IDetailCustomization.h"
 
-class FMobilityCustomization;
 class IDetailLayoutBuilder;
 
 class FSceneComponentDetails : public IDetailCustomization
@@ -19,6 +18,4 @@ public:
 	
 private:
 	void MakeTransformDetails( IDetailLayoutBuilder& DetailBuilder );
-
-	TSharedPtr<FMobilityCustomization> MobilityCustomization;
 };
diff --git a/Engine/Source/Editor/DetailCustomizations/Public/Customizations/MobilityCustomization.h b/Engine/Source/Editor/DetailCustomizations/Public/Customizations/MobilityCustomization.h
--- a/Engine/Source/Editor/DetailCustomizations/Public/Customizations/MobilityCustomization.h
+++ b/Engine/Source/Editor/DetailCustomizations/Public/Customizations/MobilityCustomization.h
@@ -4,13 +4,14 @@
 #include "CoreMinimal.h"
 #include "Engine/EngineTypes.h"
 #include "Styling/SlateColor.h"
+#include "IDetailCustomNodeBuilder.h"
 
 class IDetailCategoryBuilder;
 class IPropertyHandle;
 enum class ECheckBoxState : uint8;
 
 // Helper class to create a Mobility customization for the specified Property in the specified CategoryBuilder.
-class DETAILCUSTOMIZATIONS_API FMobilityCustomization : public TSharedFromThis<FMobilityCustomization>
+class DETAILCUSTOMIZATIONS_API FMobilityCustomization : public IDetailCustomNodeBuilder, public TSharedFromThis<FMobilityCustomization>
 {
 public:
 	enum
@@ -20,10 +21,11 @@ public:
 		MovableMobilityBitMask = (1u << EComponentMobility::Movable),
 	};
 
-	FMobilityCustomization()
-	{}
+	FMobilityCustomization(TSharedPtr<IPropertyHandle> InMobilityHandle, uint8 InRestrictedMobilityBits, bool InForLight);
 
-	void CreateMobilityCustomization(IDetailCategoryBuilder& InCategoryBuilder, TSharedPtr<IPropertyHandle> InMobilityHandle, uint8 InRestrictedMobilityBits, bool bForLight);
+	virtual void GenerateHeaderRowContent(FDetailWidgetRow& WidgetRow) override;
+	virtual FName GetName() const override;
+	virtual TSharedPtr<IPropertyHandle> GetPropertyHandle() const override { return MobilityHandle; }
 
 private:
 
@@ -33,4 +35,6 @@ private:
 	FText GetMobilityToolTip() const;
 	
 	TSharedPtr<IPropertyHandle> MobilityHandle;
+	bool bForLight;
+	uint8 RestrictedMobilityBits;
 };
diff --git a/Engine/Source/Editor/FoliageEdit/Private/FoliageTypePaintingCustomization.cpp b/Engine/Source/Editor/FoliageEdit/Private/FoliageTypePaintingCustomization.cpp
--- a/Engine/Source/Editor/FoliageEdit/Private/FoliageTypePaintingCustomization.cpp
+++ b/Engine/Source/Editor/FoliageEdit/Private/FoliageTypePaintingCustomization.cpp
@@ -90,8 +90,7 @@ void FFoliageTypePaintingCustomization::ShowFoliagePropertiesForCategory(IDetail
 		{
 			if (Property->GetFName() == MobilityName)
 			{
-				MobilityCustomization = MakeShareable(new FMobilityCustomization);
-				MobilityCustomization->CreateMobilityCustomization(CategoryBuilder, DetailLayoutBuilder.GetProperty(MobilityName), FMobilityCustomization::StationaryMobilityBitMask, false);
+				CategoryBuilder.AddCustomBuilder(MakeShared<FMobilityCustomization>(DetailLayoutBuilder.GetProperty(MobilityName), FMobilityCustomization::StationaryMobilityBitMask, false), false);
 			}
 			else
 			{
diff --git a/Engine/Source/Editor/FoliageEdit/Private/FoliageTypePaintingCustomization.h b/Engine/Source/Editor/FoliageEdit/Private/FoliageTypePaintingCustomization.h
--- a/Engine/Source/Editor/FoliageEdit/Private/FoliageTypePaintingCustomization.h
+++ b/Engine/Source/Editor/FoliageEdit/Private/FoliageTypePaintingCustomization.h
@@ -6,7 +6,6 @@
 #include "IDetailCustomization.h"
 
 class FEdModeFoliage;
-class FMobilityCustomization;
 class IDetailCategoryBuilder;
 class IDetailLayoutBuilder;
 class IDetailPropertyRow;
@@ -55,5 +54,4 @@ private:
 
 	TSharedPtr<IPropertyHandle> Scaling;
 	TSharedPtr<IPropertyHandle> ReapplyScaling;
-	TSharedPtr<FMobilityCustomization> MobilityCustomization;
 };
