Limit move snapshots to only the transform properties

diff --git a/Engine/Source/Editor/DetailCustomizations/Private/ComponentTransformDetails.cpp b/Engine/Source/Editor/DetailCustomizations/Private/ComponentTransformDetails.cpp
--- a/Engine/Source/Editor/DetailCustomizations/Private/ComponentTransformDetails.cpp
+++ b/Engine/Source/Editor/DetailCustomizations/Private/ComponentTransformDetails.cpp
@@ -1237,7 +1237,8 @@ void FComponentTransformDetails::OnSetTransform(ETransformField::Type TransformF
 
 				if (!bCommitted)
 				{
-					SnapshotTransactionBuffer(OldSceneComponent);
+					const FProperty* ConstValueProperty = ValueProperty;
+					SnapshotTransactionBuffer(OldSceneComponent, MakeArrayView(&ConstValueProperty, 1));
 				}
 
 				SceneComponent = FindObject<USceneComponent>(EditedActor, *SceneComponentPath);
@@ -1247,9 +1248,10 @@ void FComponentTransformDetails::OnSetTransform(ETransformField::Type TransformF
 					EditedActor->PostEditChangeChainProperty(PropertyChangedChainEvent);
 					SceneComponent = FindObject<USceneComponent>(EditedActor, *SceneComponentPath);
 
-					if (!bCommitted)
+					if (!bCommitted && OldSceneComponent != SceneComponent)
 					{
-						SnapshotTransactionBuffer(EditedActor);
+						const FProperty* ConstValueProperty = ValueProperty;
+						SnapshotTransactionBuffer(SceneComponent, MakeArrayView(&ConstValueProperty, 1));
 					}
 				}
 				
diff --git a/Engine/Source/Runtime/Engine/Private/ActorEditor.cpp b/Engine/Source/Runtime/Engine/Private/ActorEditor.cpp
--- a/Engine/Source/Runtime/Engine/Private/ActorEditor.cpp
+++ b/Engine/Source/Runtime/Engine/Private/ActorEditor.cpp
@@ -227,13 +227,6 @@ void AActor::PostEditMove(bool bFinished)
 	{
 		FNavigationSystem::OnPostEditActorMove(*this);
 	}
-
-	if (!bFinished)
-	{
-		// Snapshot the transaction buffer for this actor if we've not finished moving yet
-		// This allows listeners to be notified of intermediate changes of state
-		SnapshotTransactionBuffer(this);
-	}
 }
 
 bool AActor::ReregisterComponentsWhenModified() const
diff --git a/Engine/Source/Runtime/Engine/Private/Components/SceneComponent.cpp b/Engine/Source/Runtime/Engine/Private/Components/SceneComponent.cpp
--- a/Engine/Source/Runtime/Engine/Private/Components/SceneComponent.cpp
+++ b/Engine/Source/Runtime/Engine/Private/Components/SceneComponent.cpp
@@ -3325,7 +3325,12 @@ void USceneComponent::PostEditComponentMove(bool bFinished)
 	{
 		// Snapshot the transaction buffer for this component if we've not finished moving yet
 		// This allows listeners to be notified of intermediate changes of state
-		SnapshotTransactionBuffer(this);
+		static const FProperty* MovementProperties[] = {
+			USceneComponent::StaticClass()->FindPropertyByName(GET_MEMBER_NAME_CHECKED(USceneComponent, RelativeLocation)),
+			USceneComponent::StaticClass()->FindPropertyByName(GET_MEMBER_NAME_CHECKED(USceneComponent, RelativeRotation)),
+			USceneComponent::StaticClass()->FindPropertyByName(GET_MEMBER_NAME_CHECKED(USceneComponent, RelativeScale3D)),
+		};
+		SnapshotTransactionBuffer(this, MakeArrayView(MovementProperties));
 	}
 
 	{
