No-unity compile fix.

#rnx

diff --git a/Engine/Source/Editor/DetailCustomizations/Private/AnimStateAliasNodeDetails.cpp b/Engine/Source/Editor/DetailCustomizations/Private/AnimStateAliasNodeDetails.cpp
--- a/Engine/Source/Editor/DetailCustomizations/Private/AnimStateAliasNodeDetails.cpp
+++ b/Engine/Source/Editor/DetailCustomizations/Private/AnimStateAliasNodeDetails.cpp
@@ -2,6 +2,7 @@
 
 #include "AnimStateAliasNodeDetails.h"
 #include "Widgets/DeclarativeSyntaxSupport.h"
+#include "Widgets/Input/SCheckBox.h"
 #include "Widgets/Text/STextBlock.h"
 #include "DetailLayoutBuilder.h"
 #include "DetailWidgetRow.h"
diff --git a/Engine/Source/Editor/DetailCustomizations/Private/AnimStateAliasNodeDetails.h b/Engine/Source/Editor/DetailCustomizations/Private/AnimStateAliasNodeDetails.h
--- a/Engine/Source/Editor/DetailCustomizations/Private/AnimStateAliasNodeDetails.h
+++ b/Engine/Source/Editor/DetailCustomizations/Private/AnimStateAliasNodeDetails.h
@@ -4,6 +4,7 @@
 
 #include "CoreMinimal.h"
 #include "IDetailCustomization.h"
+#include "Styling/SlateStyles.h"
 
 class IDetailLayoutBuilder;
 
