Compile fix.

#rnx

diff --git a/Engine/Source/Editor/DetailCustomizations/Private/AnimStateAliasNodeDetails.h b/Engine/Source/Editor/DetailCustomizations/Private/AnimStateAliasNodeDetails.h
--- a/Engine/Source/Editor/DetailCustomizations/Private/AnimStateAliasNodeDetails.h
+++ b/Engine/Source/Editor/DetailCustomizations/Private/AnimStateAliasNodeDetails.h
@@ -4,7 +4,7 @@
 
 #include "CoreMinimal.h"
 #include "IDetailCustomization.h"
-#include "Styling/SlateStyles.h"
+#include "Styling/SlateTypes.h"
 
 class IDetailLayoutBuilder;
 
