SPropertyEditorAsset now requests additional width based on the number of additional buttons it has (if it doesn't have a thumbnail).

diff --git a/Engine/Source/Editor/DetailCustomizations/Private/SoftObjectPathCustomization.cpp b/Engine/Source/Editor/DetailCustomizations/Private/SoftObjectPathCustomization.cpp
--- a/Engine/Source/Editor/DetailCustomizations/Private/SoftObjectPathCustomization.cpp
+++ b/Engine/Source/Editor/DetailCustomizations/Private/SoftObjectPathCustomization.cpp
@@ -17,20 +17,25 @@ void FSoftObjectPathCustomization::CustomizeHeader( TSharedRef<IPropertyHandle>
 		? FEditorClassUtils::GetClassFromString(MetaClassName)
 		: UObject::StaticClass();
 
+	TSharedRef<SObjectPropertyEntryBox> ObjectPropertyEntryBox = SNew(SObjectPropertyEntryBox)
+		.AllowedClass(MetaClass)
+		.PropertyHandle(InStructPropertyHandle)
+		.ThumbnailPool(StructCustomizationUtils.GetThumbnailPool());
+
+	float MinDesiredWidth, MaxDesiredWidth;
+	ObjectPropertyEntryBox->GetDesiredWidth(MinDesiredWidth, MaxDesiredWidth);
+
 	HeaderRow
 	.NameContent()
 	[
 		InStructPropertyHandle->CreatePropertyNameWidget()
 	]
 	.ValueContent()
-	.MinDesiredWidth(250.0f)
-	.MaxDesiredWidth(0.0f)
+	.MinDesiredWidth(MinDesiredWidth)
+	.MaxDesiredWidth(MaxDesiredWidth)
 	[
 		// Add an object entry box.  Even though this isn't an object entry, we will simulate one
-		SNew(SObjectPropertyEntryBox)
-		.AllowedClass(MetaClass)
-		.PropertyHandle(InStructPropertyHandle)
-		.ThumbnailPool(StructCustomizationUtils.GetThumbnailPool())
+		ObjectPropertyEntryBox
 	];
 
 	// This avoids making duplicate reset boxes
diff --git a/Engine/Source/Editor/PropertyEditor/Private/SConstrainedBox.cpp b/Engine/Source/Editor/PropertyEditor/Private/SConstrainedBox.cpp
--- a/Engine/Source/Editor/PropertyEditor/Private/SConstrainedBox.cpp
+++ b/Engine/Source/Editor/PropertyEditor/Private/SConstrainedBox.cpp
@@ -27,7 +27,7 @@ FVector2D SConstrainedBox::ComputeDesiredSize(float LayoutScaleMultiplier) const
 		FVector2D ChildSize = ChildSlot.GetWidget()->GetDesiredSize();
 
 		float XVal = FMath::Max(MinWidthVal, ChildSize.X);
-		if (MaxWidthVal >= MinWidthVal)
+		if (MaxWidthVal > MinWidthVal)
 		{
 			XVal = FMath::Min(MaxWidthVal, XVal);
 		}
diff --git a/Engine/Source/Editor/PropertyEditor/Private/UserInterface/PropertyEditor/SPropertyEditorAsset.cpp b/Engine/Source/Editor/PropertyEditor/Private/UserInterface/PropertyEditor/SPropertyEditorAsset.cpp
--- a/Engine/Source/Editor/PropertyEditor/Private/UserInterface/PropertyEditor/SPropertyEditorAsset.cpp
+++ b/Engine/Source/Editor/PropertyEditor/Private/UserInterface/PropertyEditor/SPropertyEditorAsset.cpp
@@ -643,18 +643,28 @@ void SPropertyEditorAsset::Construct(const FArguments& InArgs, const TSharedPtr<
 			ActorPicker
 		];
 	}
+
+	NumButtons = ButtonBox->NumSlots();
 	
 	if (ButtonBoxWrapper.IsValid())
 	{
-		ButtonBoxWrapper->SetVisibility(ButtonBox->NumSlots() > 0 ? EVisibility::Visible : EVisibility::Collapsed);
+		ButtonBoxWrapper->SetVisibility(NumButtons > 0 ? EVisibility::Visible : EVisibility::Collapsed);
 	}
 }
 
 void SPropertyEditorAsset::GetDesiredWidth( float& OutMinDesiredWidth, float &OutMaxDesiredWidth )
 {
 	OutMinDesiredWidth = 250.f;
-	// No max width
 	OutMaxDesiredWidth = 350.f;
+
+	if (!AssetThumbnail.IsValid())
+	{
+		static const float ButtonWidth = 20.0f /* button width */ + 4.0f /* padding */;
+
+		const float AdditionalButtonSize = NumButtons * ButtonWidth + 8.0f /* button box padding */;
+		OutMinDesiredWidth += AdditionalButtonSize;
+		OutMaxDesiredWidth += AdditionalButtonSize;
+	}
 }
 
 const FSlateBrush* SPropertyEditorAsset::GetStatusIcon() const
diff --git a/Engine/Source/Editor/PropertyEditor/Private/UserInterface/PropertyEditor/SPropertyEditorAsset.h b/Engine/Source/Editor/PropertyEditor/Private/UserInterface/PropertyEditor/SPropertyEditorAsset.h
--- a/Engine/Source/Editor/PropertyEditor/Private/UserInterface/PropertyEditor/SPropertyEditorAsset.h
+++ b/Engine/Source/Editor/PropertyEditor/Private/UserInterface/PropertyEditor/SPropertyEditorAsset.h
@@ -376,6 +376,9 @@ private:
 	/** Whether the classes in the AllowedClassFilters list are exact or whether valid classes can be derived. */
 	bool bExactClass;
 
+	/** The number of additional buttons this picker has. */
+	int32 NumButtons;
+
 	/** Delegate to call when our object value is set */
 	FOnSetObject OnSetObject;
 
