Overrides ResetToDefault for HLOD settings in WorldSettings.

diff --git a/Engine/Source/Editor/DetailCustomizations/Private/DetailCustomizations.cpp b/Engine/Source/Editor/DetailCustomizations/Private/DetailCustomizations.cpp
--- a/Engine/Source/Editor/DetailCustomizations/Private/DetailCustomizations.cpp
+++ b/Engine/Source/Editor/DetailCustomizations/Private/DetailCustomizations.cpp
@@ -255,6 +255,7 @@ void FDetailCustomizationsModule::RegisterPropertyTypeCustomizations()
 	RegisterCustomPropertyTypeLayout("AutoReimportWildcard", FOnGetPropertyTypeCustomizationInstance::CreateStatic(&FAutoReimportWildcardCustomization::MakeInstance));
 	RegisterCustomPropertyTypeLayout("DistanceDatum", FOnGetPropertyTypeCustomizationInstance::CreateStatic(&FDistanceDatumStructCustomization::MakeInstance));
 	RegisterCustomPropertyTypeLayout("HierarchicalSimplification", FOnGetPropertyTypeCustomizationInstance::CreateStatic(&FHierarchicalSimplificationCustomizations::MakeInstance));
+	RegisterCustomPropertyTypeLayout("MeshMergingSettings", FOnGetPropertyTypeCustomizationInstance::CreateStatic(&FMeshMergingSettingsCustomization::MakeInstance));
 	RegisterCustomPropertyTypeLayout("MeshProxySettings", FOnGetPropertyTypeCustomizationInstance::CreateStatic(&FMeshProxySettingsCustomizations::MakeInstance));
 	RegisterCustomPropertyTypeLayout("PostProcessSettings", FOnGetPropertyTypeCustomizationInstance::CreateStatic(&FPostProcessSettingsCustomization::MakeInstance));
 	RegisterCustomPropertyTypeLayout("AssetImportInfo", FOnGetPropertyTypeCustomizationInstance::CreateStatic(&FAssetImportDataCustomization::MakeInstance));
diff --git a/Engine/Source/Editor/DetailCustomizations/Private/HierarchicalSimplificationCustomizations.cpp b/Engine/Source/Editor/DetailCustomizations/Private/HierarchicalSimplificationCustomizations.cpp
--- a/Engine/Source/Editor/DetailCustomizations/Private/HierarchicalSimplificationCustomizations.cpp
+++ b/Engine/Source/Editor/DetailCustomizations/Private/HierarchicalSimplificationCustomizations.cpp
@@ -66,29 +66,39 @@ void FHierarchicalSimplificationCustomizations::CustomizeChildren( TSharedRef<IP
 		if (Iter.Value() == SimplifyMeshPropertyHandle)
 		{
 			IDetailPropertyRow& SimplifyMeshRow = MergeGroup.AddPropertyRow(SimplifyMeshPropertyHandle.ToSharedRef());
+			AddResetToDefaultOverrides(SimplifyMeshRow);
+
 			SimplifyMeshRow.Visibility(TAttribute<EVisibility>(this, &FHierarchicalSimplificationCustomizations::IsSimplifyMeshVisible));
 		}
 		else if (Iter.Value() == ProxyMeshSettingPropertyHandle)
 		{
 			IDetailPropertyRow& SettingsRow = MergeGroup.AddPropertyRow(Iter.Value().ToSharedRef());
+			AddResetToDefaultOverrides(SettingsRow);
+
 			SettingsRow.Visibility(TAttribute<EVisibility>(this, &FHierarchicalSimplificationCustomizations::IsProxyMeshSettingVisible));
 		}
 		else if (Iter.Value() == MergeMeshSettingPropertyHandle)
 		{
 			IDetailPropertyRow& SettingsRow = MergeGroup.AddPropertyRow(Iter.Value().ToSharedRef());
+			AddResetToDefaultOverrides(SettingsRow);
+
 			SettingsRow.Visibility(TAttribute<EVisibility>(this, &FHierarchicalSimplificationCustomizations::IsMergeMeshSettingVisible));
 		}
 		else  if (Iter.Value() == TransitionScreenSizePropertyHandle)
 		{
 			IDetailPropertyRow& SettingsRow = MergeGroup.AddPropertyRow(Iter.Value().ToSharedRef());
+			AddResetToDefaultOverrides(SettingsRow);
 		}
 		else if (Iter.Value() == OverrideDrawDistancePropertyHandle)
 		{
 			IDetailPropertyRow& SettingsRow = MergeGroup.AddPropertyRow(Iter.Value().ToSharedRef());
+			AddResetToDefaultOverrides(SettingsRow);
 		}
 		else if (Iter.Value() == ReusePreviousLevelClustersPropertyHandle)
 		{
 			IDetailPropertyRow& SettingsRow = ClusterGroup.AddPropertyRow(Iter.Value().ToSharedRef());
+			AddResetToDefaultOverrides(SettingsRow);
+
 			uint32 Index = StructPropertyHandle->GetIndexInArray();
 			// Hide the property for HLOD level 0
 			SettingsRow.Visibility(TAttribute<EVisibility>::Create([Index]()
@@ -99,6 +109,7 @@ void FHierarchicalSimplificationCustomizations::CustomizeChildren( TSharedRef<IP
 		else
 		{
 			IDetailPropertyRow& SettingsRow = ClusterGroup.AddPropertyRow(Iter.Value().ToSharedRef());
+			AddResetToDefaultOverrides(SettingsRow);
 		}
 	}
 }
diff --git a/Engine/Source/Editor/DetailCustomizations/Private/HierarchicalSimplificationCustomizations.h b/Engine/Source/Editor/DetailCustomizations/Private/HierarchicalSimplificationCustomizations.h
--- a/Engine/Source/Editor/DetailCustomizations/Private/HierarchicalSimplificationCustomizations.h
+++ b/Engine/Source/Editor/DetailCustomizations/Private/HierarchicalSimplificationCustomizations.h
@@ -6,10 +6,12 @@
 #include "Layout/Visibility.h"
 #include "IPropertyTypeCustomization.h"
 #include "PropertyHandle.h"
+#include "GameFramework/WorldSettings.h"
+#include "OverrideResetToDefault.h"
 
 class FDetailWidgetRow;
 
-class FHierarchicalSimplificationCustomizations : public IPropertyTypeCustomization
+class FHierarchicalSimplificationCustomizations : public IPropertyTypeCustomization, public TOverrideResetToDefaultWithStaticUStruct<FHierarchicalSimplification>
 {
 public:
 	static TSharedRef<IPropertyTypeCustomization> MakeInstance();
diff --git a/Engine/Source/Editor/DetailCustomizations/Private/MeshMergingSettingsCustomization.cpp b/Engine/Source/Editor/DetailCustomizations/Private/MeshMergingSettingsCustomization.cpp
--- a/Engine/Source/Editor/DetailCustomizations/Private/MeshMergingSettingsCustomization.cpp
+++ b/Engine/Source/Editor/DetailCustomizations/Private/MeshMergingSettingsCustomization.cpp
@@ -6,9 +6,12 @@
 #include "UObject/UnrealType.h"
 #include "PropertyHandle.h"
 #include "DetailLayoutBuilder.h"
+#include "DetailWidgetRow.h"
 #include "IDetailPropertyRow.h"
+#include "IDetailChildrenBuilder.h"
 #include "DetailCategoryBuilder.h"
 #include "PropertyRestriction.h"
+#include "OverrideResetToDefault.h"
 
 #define LOCTEXT_NAMESPACE "FMeshMergingSettingCustomization"
 
@@ -92,4 +95,37 @@ bool FMeshMergingSettingsObjectCustomization::AreMaterialPropertiesEnabled() con
 	return !(CurrentEnumValue == (uint8)EMeshLODSelectionType::AllLODs);
 }
 
+TSharedRef<IPropertyTypeCustomization> FMeshMergingSettingsCustomization::MakeInstance()
+{
+	return MakeShareable(new FMeshMergingSettingsCustomization);
+}
+
+void FMeshMergingSettingsCustomization::CustomizeHeader(TSharedRef<IPropertyHandle> StructPropertyHandle, FDetailWidgetRow& HeaderRow, IPropertyTypeCustomizationUtils& StructCustomizationUtils)
+{
+	HeaderRow
+	.NameContent()
+	[
+			StructPropertyHandle->CreatePropertyNameWidget(StructPropertyHandle->GetPropertyDisplayName())
+	]
+	.ValueContent()
+	[
+			StructPropertyHandle->CreatePropertyValueWidget(false)
+	];
+}
+
+void FMeshMergingSettingsCustomization::CustomizeChildren(TSharedRef<IPropertyHandle> StructPropertyHandle, class IDetailChildrenBuilder& ChildBuilder, IPropertyTypeCustomizationUtils& StructCustomizationUtils)
+{
+	uint32 NumChildren = 0;
+	
+	StructPropertyHandle->GetNumChildren(NumChildren);
+
+	for (uint32 ChildIndex = 0; ChildIndex < NumChildren; ++ChildIndex)
+	{
+		TSharedRef<IPropertyHandle> ChildHandle = StructPropertyHandle->GetChildHandle(ChildIndex).ToSharedRef();
+		IDetailPropertyRow& NewRow = ChildBuilder.AddProperty(ChildHandle);
+
+		AddResetToDefaultOverrides(NewRow);
+	}
+}
+
 #undef LOCTEXT_NAMESPACE
diff --git a/Engine/Source/Editor/DetailCustomizations/Private/MeshMergingSettingsCustomization.h b/Engine/Source/Editor/DetailCustomizations/Private/MeshMergingSettingsCustomization.h
--- a/Engine/Source/Editor/DetailCustomizations/Private/MeshMergingSettingsCustomization.h
+++ b/Engine/Source/Editor/DetailCustomizations/Private/MeshMergingSettingsCustomization.h
@@ -3,8 +3,10 @@
 #pragma once
 
 #include "CoreMinimal.h"
+#include "Engine/MeshMerging.h"
 #include "Layout/Visibility.h"
 #include "IDetailCustomization.h"
+#include "OverrideResetToDefault.h"
 
 class IDetailLayoutBuilder;
 class IPropertyHandle;
@@ -26,3 +28,14 @@ protected:
 	TSharedPtr<IPropertyHandle> EnumProperty;
 private:
 };
+
+
+class FMeshMergingSettingsCustomization : public IPropertyTypeCustomization, public TOverrideResetToDefaultWithStaticUStruct<FMeshMergingSettings>
+{
+public:
+	static TSharedRef<IPropertyTypeCustomization> MakeInstance();
+
+	/** IPropertyTypeCustomization instance */
+	virtual void CustomizeHeader(TSharedRef<IPropertyHandle> StructPropertyHandle, FDetailWidgetRow& HeaderRow, IPropertyTypeCustomizationUtils& StructCustomizationUtils) override;
+	virtual void CustomizeChildren(TSharedRef<IPropertyHandle> StructPropertyHandle, class IDetailChildrenBuilder& ChildBuilder, IPropertyTypeCustomizationUtils& StructCustomizationUtils) override;
+};
diff --git a/Engine/Source/Editor/DetailCustomizations/Private/MeshProxySettingsCustomizations.cpp b/Engine/Source/Editor/DetailCustomizations/Private/MeshProxySettingsCustomizations.cpp
--- a/Engine/Source/Editor/DetailCustomizations/Private/MeshProxySettingsCustomizations.cpp
+++ b/Engine/Source/Editor/DetailCustomizations/Private/MeshProxySettingsCustomizations.cpp
@@ -79,8 +79,6 @@ void FMeshProxySettingsCustomizations::CustomizeChildren(TSharedRef<IPropertyHan
 	TSharedPtr< IPropertyHandle > MergeDistanceHandle                     = PropertyHandles.FindChecked(GET_MEMBER_NAME_CHECKED(FMeshProxySettings, MergeDistance));
 	TSharedPtr< IPropertyHandle > UnresolvedGeometryColorHandle           = PropertyHandles.FindChecked(GET_MEMBER_NAME_CHECKED(FMeshProxySettings, UnresolvedGeometryColor));
 	TSharedPtr< IPropertyHandle > VoxelSizeHandle                         = PropertyHandles.FindChecked(GET_MEMBER_NAME_CHECKED(FMeshProxySettings, VoxelSize));
-	
-
 
 	for (auto Iter(PropertyHandles.CreateConstIterator()); Iter; ++Iter)
 	{
@@ -88,27 +86,37 @@ void FMeshProxySettingsCustomizations::CustomizeChildren(TSharedRef<IPropertyHan
 		if (Iter.Value() == HardAngleThresholdPropertyHandle)
 		{
 			IDetailPropertyRow& MeshProxySettingsRow = AddPropertyToGroup(Iter.Value().ToSharedRef());
+			AddResetToDefaultOverrides(MeshProxySettingsRow);
+			
 			MeshProxySettingsRow.ToolTip(FText::FromString(FString("Angle at which a hard edge is introduced between faces.  Note: Increases vertex count and may introduce additional UV seams.  It is only recommended if not using normals maps")));
 			MeshProxySettingsRow.Visibility(TAttribute<EVisibility>(this, &FMeshProxySettingsCustomizations::IsHardAngleThresholdVisible));
 		}
 		else if (Iter.Value() == NormalCalcMethodPropertyHandle)
 		{
 			IDetailPropertyRow& MeshProxySettingsRow = AddPropertyToGroup(Iter.Value().ToSharedRef());
+			AddResetToDefaultOverrides(MeshProxySettingsRow);
+			
 			MeshProxySettingsRow.Visibility(TAttribute<EVisibility>(this, &FMeshProxySettingsCustomizations::IsNormalCalcMethodVisible));
 		}
 		else if (Iter.Value() == MaxRayCastDistdPropertyHandle)
 		{
 			IDetailPropertyRow& MeshProxySettingsRow = AddPropertyToGroup(Iter.Value().ToSharedRef());
+			AddResetToDefaultOverrides(MeshProxySettingsRow);
+
 			MeshProxySettingsRow.Visibility(TAttribute<EVisibility>(this, &FMeshProxySettingsCustomizations::IsSearchDistanceVisible));
 		}
 		else if (Iter.Value() == RecalculateNormalsPropertyHandle)
 		{
 			IDetailPropertyRow& MeshProxySettingsRow = AddPropertyToGroup(Iter.Value().ToSharedRef());
+			AddResetToDefaultOverrides(MeshProxySettingsRow);
+
 			MeshProxySettingsRow.Visibility(TAttribute<EVisibility>(this, &FMeshProxySettingsCustomizations::IsRecalculateNormalsVisible));
 		}
 		else if (Iter.Value() == UseLandscapeCullingPropertyHandle)
 		{
 			IDetailPropertyRow& MeshProxySettingsRow = AddPropertyToGroup(Iter.Value().ToSharedRef());
+			AddResetToDefaultOverrides(MeshProxySettingsRow);
+			
 			MeshProxySettingsRow.DisplayName(FText::FromString(FString("Enable Volume Culling")));
 			MeshProxySettingsRow.ToolTip(FText::FromString(FString("Allow culling volumes to exclude geometry.")));
 			MeshProxySettingsRow.Visibility(TAttribute<EVisibility>(this, &FMeshProxySettingsCustomizations::IsUseLandscapeCullingVisible));
@@ -116,26 +124,35 @@ void FMeshProxySettingsCustomizations::CustomizeChildren(TSharedRef<IPropertyHan
 		else if (Iter.Value() == LandscapeCullingPrecisionPropertyHandle)
 		{
 			IDetailPropertyRow& MeshProxySettingsRow = AddPropertyToGroup(Iter.Value().ToSharedRef());
+			AddResetToDefaultOverrides(MeshProxySettingsRow);
+			
 			MeshProxySettingsRow.Visibility(TAttribute<EVisibility>(this, &FMeshProxySettingsCustomizations::IsUseLandscapeCullingPrecisionVisible));
 		}
 		else if (Iter.Value() == MergeDistanceHandle)
 		{
 			IDetailPropertyRow& MeshProxySettingsRow = AddPropertyToGroup(Iter.Value().ToSharedRef());
+			AddResetToDefaultOverrides(MeshProxySettingsRow);
+			
 			MeshProxySettingsRow.Visibility(TAttribute<EVisibility>(this, &FMeshProxySettingsCustomizations::IsMergeDistanceVisible));
 		}
 		else if (Iter.Value() == UnresolvedGeometryColorHandle)
 		{
 			IDetailPropertyRow& MeshProxySettingsRow = AddPropertyToGroup(Iter.Value().ToSharedRef());
+			AddResetToDefaultOverrides(MeshProxySettingsRow);
+			
 			MeshProxySettingsRow.Visibility(TAttribute<EVisibility>(this, &FMeshProxySettingsCustomizations::IsUnresolvedGeometryColorVisible));
 		}
 		else if (Iter.Value() == VoxelSizeHandle)
 		{
 			IDetailPropertyRow& MeshProxySettingsRow = AddPropertyToGroup(Iter.Value().ToSharedRef());
+			AddResetToDefaultOverrides(MeshProxySettingsRow);
+			
 			MeshProxySettingsRow.Visibility(TAttribute<EVisibility>(this, &FMeshProxySettingsCustomizations::IsVoxelSizeVisible));
 		}
 		else
 		{
 			IDetailPropertyRow& SettingsRow = AddPropertyToGroup(Iter.Value().ToSharedRef());
+			AddResetToDefaultOverrides(SettingsRow);
 		}
 	}
 }
diff --git a/Engine/Source/Editor/DetailCustomizations/Private/MeshProxySettingsCustomizations.h b/Engine/Source/Editor/DetailCustomizations/Private/MeshProxySettingsCustomizations.h
--- a/Engine/Source/Editor/DetailCustomizations/Private/MeshProxySettingsCustomizations.h
+++ b/Engine/Source/Editor/DetailCustomizations/Private/MeshProxySettingsCustomizations.h
@@ -3,13 +3,15 @@
 #pragma once
 
 #include "CoreMinimal.h"
+#include "Engine/MeshMerging.h"
 #include "Layout/Visibility.h"
 #include "IPropertyTypeCustomization.h"
 #include "PropertyHandle.h"
+#include "OverrideResetToDefault.h"
 
 class FDetailWidgetRow;
 
-class FMeshProxySettingsCustomizations : public IPropertyTypeCustomization
+class FMeshProxySettingsCustomizations : public IPropertyTypeCustomization, public TOverrideResetToDefaultWithStaticUStruct<FMeshProxySettings>
 {
 
 public:
diff --git a/Engine/Source/Editor/DetailCustomizations/Private/OverrideResetToDefault.h b/Engine/Source/Editor/DetailCustomizations/Private/OverrideResetToDefault.h
new file mode 100644
--- /dev/null
+++ b/Engine/Source/Editor/DetailCustomizations/Private/OverrideResetToDefault.h
@@ -0,0 +1,79 @@
+// Copyright Epic Games, Inc. All Rights Reserved.
+
+#pragma once
+
+#include "CoreMinimal.h"
+#include "IDetailPropertyRow.h"
+#include "Layout/Visibility.h"
+#include "IPropertyTypeCustomization.h"
+#include "PropertyHandle.h"
+#include "UObject/PropertyAccessUtil.h"
+
+/**
+ * Class template aiming to override ResetToDefault methods by comparing properties with an internal static DefaultObject.
+ */
+template<typename UStructType>
+class TOverrideResetToDefaultWithStaticUStruct 
+{
+public:
+	/** Called by the UI to show/hide the reset widgets */
+	static bool IsResetToDefaultVisible(TSharedPtr<IPropertyHandle> InPropertyHandle);
+	
+	/** Reset to default triggered in UI */
+	static void OnResetToDefault(TSharedPtr<IPropertyHandle> InPropertyHandle);
+
+protected:
+	/** Adds callbacks used by the UI to determine if a property can be reset and reset it. */
+	static void AddResetToDefaultOverrides(IDetailPropertyRow& InDetailPropertyRow);
+
+	static const UStructType DefaultObject;
+};
+
+template<typename UStructType>
+const UStructType TOverrideResetToDefaultWithStaticUStruct<UStructType>::DefaultObject;
+
+template<typename UStructType>
+bool TOverrideResetToDefaultWithStaticUStruct<UStructType>::IsResetToDefaultVisible(TSharedPtr<IPropertyHandle> InPropertyHandle)
+{
+	FProperty*		Property = InPropertyHandle->GetProperty();
+	const void*		DefaultValuePtr = nullptr;
+	void*			ValuePtr = nullptr;
+
+	check(Property != nullptr);
+
+	DefaultValuePtr = Property->ContainerPtrToValuePtr<void>(&DefaultObject);
+	InPropertyHandle->GetValueData(ValuePtr);
+
+	if ((DefaultValuePtr != nullptr) && (ValuePtr != nullptr))
+	{
+		return !Property->Identical(DefaultValuePtr, ValuePtr);
+	}
+
+	return false;
+}
+
+template<typename UStructType>
+void TOverrideResetToDefaultWithStaticUStruct<UStructType>::OnResetToDefault(TSharedPtr<IPropertyHandle> InPropertyHandle)
+{
+	FProperty*	Property = InPropertyHandle->GetProperty();
+	const void* DefaultValuePtr = nullptr;
+	void*		ValuePtr = nullptr;
+	
+	check(Property != nullptr);
+
+	DefaultValuePtr = Property->ContainerPtrToValuePtr<void>(&DefaultObject);
+	InPropertyHandle->GetValueData(ValuePtr);
+
+	if ((DefaultValuePtr != nullptr) && (ValuePtr != nullptr))
+	{
+		PropertyAccessUtil::SetPropertyValue_DirectSingle(Property, DefaultValuePtr, Property, ValuePtr, 0, Property->HasAnyFlags(RF_ArchetypeObject), []() { return nullptr; });
+	}
+}
+
+template<typename UStructType>
+void TOverrideResetToDefaultWithStaticUStruct<UStructType>::AddResetToDefaultOverrides(IDetailPropertyRow& InDetailPropertyRow)
+{
+	InDetailPropertyRow.OverrideResetToDefault(FResetToDefaultOverride::Create(FIsResetToDefaultVisible::CreateStatic(&TOverrideResetToDefaultWithStaticUStruct<UStructType>::IsResetToDefaultVisible),
+																			   FResetToDefaultHandler::CreateStatic(&TOverrideResetToDefaultWithStaticUStruct<UStructType>::OnResetToDefault),
+																			   true));
+}
