Fix mobility showing up on some components that shouldnt display it

diff --git a/Engine/Source/Editor/DetailCustomizations/Private/SceneComponentDetails.cpp b/Engine/Source/Editor/DetailCustomizations/Private/SceneComponentDetails.cpp
--- a/Engine/Source/Editor/DetailCustomizations/Private/SceneComponentDetails.cpp
+++ b/Engine/Source/Editor/DetailCustomizations/Private/SceneComponentDetails.cpp
@@ -147,72 +147,75 @@ void FSceneComponentDetails::CustomizeDetails( IDetailLayoutBuilder& DetailBuild
 
 	// Put mobility property in Transform section
 	IDetailCategoryBuilder& TransformCategory = DetailBuilder.EditCategory( "TransformCommon", LOCTEXT("TransformCommonCategory", "Transform"), ECategoryPriority::Transform );
-	TSharedPtr<IPropertyHandle> MobilityHandle = DetailBuilder.GetProperty("Mobility");
+	TSharedRef<IPropertyHandle> MobilityHandle = DetailBuilder.GetProperty("Mobility");
 
-	uint8 RestrictedMobilityBits = 0u;
-	bool bAnySelectedIsLight = false;
+	if(MobilityHandle->IsValidHandle())
+	{
+		uint8 RestrictedMobilityBits = 0u;
+		bool bAnySelectedIsLight = false;
 
-	TArray< TWeakObjectPtr<UObject> > SelectedSceneComponents;
+		TArray< TWeakObjectPtr<UObject> > SelectedSceneComponents;
 
-	// see if any of the selected objects have mobility restrictions
-	DetailBuilder.GetObjectsBeingCustomized( SelectedSceneComponents );
+		// see if any of the selected objects have mobility restrictions
+		DetailBuilder.GetObjectsBeingCustomized(SelectedSceneComponents);
 
-	for (TArray<TWeakObjectPtr<UObject>>::TConstIterator ObjectIt(SelectedSceneComponents); ObjectIt; ++ObjectIt)
-	{
-		if (!ObjectIt->IsValid())
+		for (TArray<TWeakObjectPtr<UObject>>::TConstIterator ObjectIt(SelectedSceneComponents); ObjectIt; ++ObjectIt)
 		{
-			continue;
-		}
+			if (!ObjectIt->IsValid())
+			{
+				continue;
+			}
 
-		USceneComponent const* SceneComponent = Cast<USceneComponent>((*ObjectIt).Get());
-		if (SceneComponent == NULL)
-		{
-			continue;
-		}
+			USceneComponent const* SceneComponent = Cast<USceneComponent>((*ObjectIt).Get());
+			if (SceneComponent == NULL)
+			{
+				continue;
+			}
 
-		if (SceneComponent->IsA(ULightComponentBase::StaticClass()))
-		{
-			bAnySelectedIsLight = true;
-		}
+			if (SceneComponent->IsA(ULightComponentBase::StaticClass()))
+			{
+				bAnySelectedIsLight = true;
+			}
 
-		// if we haven't restricted the "Static" option yet
-		if (!(RestrictedMobilityBits & FMobilityCustomization::StaticMobilityBitMask))
-		{
-			FText RestrictReason;
-			if (IsMobilitySettingProhibited(EComponentMobility::Static, SceneComponent, RestrictReason))
+			// if we haven't restricted the "Static" option yet
+			if (!(RestrictedMobilityBits & FMobilityCustomization::StaticMobilityBitMask))
 			{
-				 TSharedPtr<FPropertyRestriction> StaticRestriction = MakeShareable(new FPropertyRestriction(MoveTemp(RestrictReason)));
-				 const UEnum* const ComponentMobilityEnum = StaticEnum<EComponentMobility::Type>();		
-				 StaticRestriction->AddDisabledValue(ComponentMobilityEnum->GetNameStringByValue((uint8)EComponentMobility::Static));
-				 MobilityHandle->AddRestriction(StaticRestriction.ToSharedRef());
+				FText RestrictReason;
+				if (IsMobilitySettingProhibited(EComponentMobility::Static, SceneComponent, RestrictReason))
+				{
+					TSharedPtr<FPropertyRestriction> StaticRestriction = MakeShareable(new FPropertyRestriction(MoveTemp(RestrictReason)));
+					const UEnum* const ComponentMobilityEnum = StaticEnum<EComponentMobility::Type>();
+					StaticRestriction->AddDisabledValue(ComponentMobilityEnum->GetNameStringByValue((uint8)EComponentMobility::Static));
+					MobilityHandle->AddRestriction(StaticRestriction.ToSharedRef());
 
-				 RestrictedMobilityBits |= FMobilityCustomization::StaticMobilityBitMask;
-			}			
-		}
+					RestrictedMobilityBits |= FMobilityCustomization::StaticMobilityBitMask;
+				}
+			}
 
-		// if we haven't restricted the "Stationary" option yet
-		if (!(RestrictedMobilityBits & FMobilityCustomization::StationaryMobilityBitMask))
-		{
-			FText RestrictReason;
-			if (IsMobilitySettingProhibited(EComponentMobility::Stationary, SceneComponent, RestrictReason))
+			// if we haven't restricted the "Stationary" option yet
+			if (!(RestrictedMobilityBits & FMobilityCustomization::StationaryMobilityBitMask))
 			{
-				TSharedPtr<FPropertyRestriction> StationaryRestriction = MakeShareable(new FPropertyRestriction(MoveTemp(RestrictReason)));
-				const UEnum* const ComponentMobilityEnum = StaticEnum<EComponentMobility::Type>();
-				StationaryRestriction->AddDisabledValue(ComponentMobilityEnum->GetNameStringByValue((uint8)EComponentMobility::Stationary));
-				MobilityHandle->AddRestriction(StationaryRestriction.ToSharedRef());
+				FText RestrictReason;
+				if (IsMobilitySettingProhibited(EComponentMobility::Stationary, SceneComponent, RestrictReason))
+				{
+					TSharedPtr<FPropertyRestriction> StationaryRestriction = MakeShareable(new FPropertyRestriction(MoveTemp(RestrictReason)));
+					const UEnum* const ComponentMobilityEnum = StaticEnum<EComponentMobility::Type>();
+					StationaryRestriction->AddDisabledValue(ComponentMobilityEnum->GetNameStringByValue((uint8)EComponentMobility::Stationary));
+					MobilityHandle->AddRestriction(StationaryRestriction.ToSharedRef());
 
-				RestrictedMobilityBits |= FMobilityCustomization::StationaryMobilityBitMask;
+					RestrictedMobilityBits |= FMobilityCustomization::StationaryMobilityBitMask;
+				}
 			}
-		}
 
-		// no need to go through all of them if we can't restrict any more
-		if ((RestrictedMobilityBits & FMobilityCustomization::StaticMobilityBitMask) && (RestrictedMobilityBits & FMobilityCustomization::StationaryMobilityBitMask))
-		{
-			break;
+			// no need to go through all of them if we can't restrict any more
+			if ((RestrictedMobilityBits & FMobilityCustomization::StaticMobilityBitMask) && (RestrictedMobilityBits & FMobilityCustomization::StationaryMobilityBitMask))
+			{
+				break;
+			}
 		}
-	}
 
-	TransformCategory.AddCustomBuilder(MakeShared<FMobilityCustomization>(MobilityHandle, RestrictedMobilityBits, bAnySelectedIsLight));
+		TransformCategory.AddCustomBuilder(MakeShared<FMobilityCustomization>(MobilityHandle, RestrictedMobilityBits, bAnySelectedIsLight));
+	}
 
 	// Only display bHiddenInGame if the property is being flattened in to an Actor.
 	// Details panel for BP component will have the base class be the Actor due to how the SKismetInspector works, but in that case we
