Fix icons

This removes a UE4 references and fixes icons by updating to the new name.

diff --git a/Engine/Source/Editor/DetailCustomizations/Private/MacTargetSettingsDetails.cpp b/Engine/Source/Editor/DetailCustomizations/Private/MacTargetSettingsDetails.cpp
--- a/Engine/Source/Editor/DetailCustomizations/Private/MacTargetSettingsDetails.cpp
+++ b/Engine/Source/Editor/DetailCustomizations/Private/MacTargetSettingsDetails.cpp
@@ -82,7 +82,7 @@ static FString GetIconFilename(EMacImageScope::Type Scope)
 
 	if (Scope == EMacImageScope::Engine)
 	{
-		FString Filename = FPaths::EngineDir() / FString(TEXT("Source/Runtime/Launch/Resources")) / PlatformName / FString("UE4.icns");
+		FString Filename = FPaths::EngineDir() / FString(TEXT("Source/Runtime/Launch/Resources")) / PlatformName / FString("UnrealEngine.icns");
 		return FPaths::ConvertRelativePathToFull(Filename);
 	}
 	else
