Fixed blueprint editor spline details crash and spline select keys not working in level editor before spline points are selected.

diff --git a/Engine/Source/Editor/DetailCustomizations/Private/SplineComponentDetails.cpp b/Engine/Source/Editor/DetailCustomizations/Private/SplineComponentDetails.cpp
--- a/Engine/Source/Editor/DetailCustomizations/Private/SplineComponentDetails.cpp
+++ b/Engine/Source/Editor/DetailCustomizations/Private/SplineComponentDetails.cpp
@@ -416,8 +416,6 @@ void FSplinePointDetails::GenerateSplinePointSelectionControls(IDetailChildrenBu
 
 void FSplinePointDetails::GenerateChildContent(IDetailChildrenBuilder& ChildrenBuilder)
 {
-	check(SplineComp);
-
 	// Select spline point buttons
 	GenerateSplinePointSelectionControls(ChildrenBuilder);
 
@@ -435,6 +433,11 @@ void FSplinePointDetails::GenerateChildContent(IDetailChildrenBuilder& ChildrenB
 		]
 		];
 
+	if (!SplineComp)
+	{
+		return;
+	}
+
 	// Input key
 	ChildrenBuilder.AddCustomRow(LOCTEXT("InputKey", "Input Key"))
 		.Visibility(TAttribute<EVisibility>(this, &FSplinePointDetails::IsEnabled))
@@ -682,12 +685,22 @@ void FSplinePointDetails::Tick(float DeltaTime)
 
 void FSplinePointDetails::UpdateValues()
 {
-	// Always update the spline component based on the spline component visualizer's currently edited component.
-	SplineComp = SplineVisualizer.IsValid() ? SplineVisualizer->GetEditedSplineComponent() : nullptr;
+	// If this is a blueprint spline, always update the spline component based on 
+	// the spline component visualizer's currently edited spline component.
+	if (SplineCompArchetype)
+	{
+		USplineComponent* EditedSplineComp = SplineVisualizer.IsValid() ? SplineVisualizer->GetEditedSplineComponent() : nullptr;
+
+		if (!EditedSplineComp || (EditedSplineComp->GetArchetype() != SplineCompArchetype))
+		{
+			return;
+		}
+
+		SplineComp = EditedSplineComp;
+	}
 
 	if (!SplineComp || !SplineVisualizer.IsValid())
 	{
-		SplineComp = nullptr;
 		return;
 	}
 
