Mobility handle is now marked hidden by customization to prevent duplicate Mobility entries.

diff --git a/Engine/Source/Editor/DetailCustomizations/Private/SceneComponentDetails.cpp b/Engine/Source/Editor/DetailCustomizations/Private/SceneComponentDetails.cpp
--- a/Engine/Source/Editor/DetailCustomizations/Private/SceneComponentDetails.cpp
+++ b/Engine/Source/Editor/DetailCustomizations/Private/SceneComponentDetails.cpp
@@ -214,6 +214,7 @@ void FSceneComponentDetails::CustomizeDetails( IDetailLayoutBuilder& DetailBuild
 			}
 		}
 
+		MobilityHandle->MarkHiddenByCustomization();
 		TransformCategory.AddCustomBuilder(MakeShared<FMobilityCustomization>(MobilityHandle, RestrictedMobilityBits, bAnySelectedIsLight));
 	}
 
