Display dynamic delegates in the actor details panel, allowing users to hit the green "Plus" button while editing a level blueprint or other blueprint. Previously this would only work on components, but this change will allow us to do it when we are editing the CDO of an actor as well.

diff --git a/Engine/Source/Editor/DetailCustomizations/Private/ActorDetails.cpp b/Engine/Source/Editor/DetailCustomizations/Private/ActorDetails.cpp
--- a/Engine/Source/Editor/DetailCustomizations/Private/ActorDetails.cpp
+++ b/Engine/Source/Editor/DetailCustomizations/Private/ActorDetails.cpp
@@ -56,6 +56,10 @@
 #include "ScopedTransaction.h"
 #include "WorldPartition/WorldPartitionSubsystem.h"
 
+#include "Kismet2/BlueprintEditorUtils.h"
+#include "K2Node_AddDelegate.h"
+#include "EdGraphSchema_K2_Actions.h"
+
 #define LOCTEXT_NAMESPACE "ActorDetails"
 
 FExtendActorDetails OnExtendActorDetails;
@@ -145,16 +149,24 @@ void FActorDetails::CustomizeDetails( IDetailLayoutBuilder& DetailLayout )
 	TSharedPtr<IPropertyHandle> PrimaryTickProperty = DetailLayout.GetProperty(GET_MEMBER_NAME_CHECKED(AActor, PrimaryActorTick));
 
 	// Defaults only show tick properties
-	if (DetailLayout.HasClassDefaultObject() && !HideCategories.Contains(TEXT("Tick")))
+	if (DetailLayout.HasClassDefaultObject())
 	{
-		// Note: the category is renamed to differentiate between 
-		IDetailCategoryBuilder& TickCategory = DetailLayout.EditCategory("Tick", LOCTEXT("TickCategoryName", "Actor Tick") );
-
-		TickCategory.AddProperty(PrimaryTickProperty->GetChildHandle(GET_MEMBER_NAME_CHECKED(FTickFunction, bStartWithTickEnabled)));
-		TickCategory.AddProperty(PrimaryTickProperty->GetChildHandle(GET_MEMBER_NAME_CHECKED(FTickFunction, TickInterval)));
-		TickCategory.AddProperty(PrimaryTickProperty->GetChildHandle(GET_MEMBER_NAME_CHECKED(FTickFunction, bTickEvenWhenPaused)), EPropertyLocation::Advanced);
-		TickCategory.AddProperty(PrimaryTickProperty->GetChildHandle(GET_MEMBER_NAME_CHECKED(FTickFunction, bAllowTickOnDedicatedServer)), EPropertyLocation::Advanced);
-		TickCategory.AddProperty(PrimaryTickProperty->GetChildHandle(GET_MEMBER_NAME_CHECKED(FTickFunction, TickGroup)), EPropertyLocation::Advanced);
+		if (!HideCategories.Contains(TEXT("Tick")))
+		{
+			// Note: the category is renamed to differentiate between 
+			IDetailCategoryBuilder& TickCategory = DetailLayout.EditCategory("Tick", LOCTEXT("TickCategoryName", "Actor Tick") );
+
+			TickCategory.AddProperty(PrimaryTickProperty->GetChildHandle(GET_MEMBER_NAME_CHECKED(FTickFunction, bStartWithTickEnabled)));
+			TickCategory.AddProperty(PrimaryTickProperty->GetChildHandle(GET_MEMBER_NAME_CHECKED(FTickFunction, TickInterval)));
+			TickCategory.AddProperty(PrimaryTickProperty->GetChildHandle(GET_MEMBER_NAME_CHECKED(FTickFunction, bTickEvenWhenPaused)), EPropertyLocation::Advanced);
+			TickCategory.AddProperty(PrimaryTickProperty->GetChildHandle(GET_MEMBER_NAME_CHECKED(FTickFunction, bAllowTickOnDedicatedServer)), EPropertyLocation::Advanced);
+			TickCategory.AddProperty(PrimaryTickProperty->GetChildHandle(GET_MEMBER_NAME_CHECKED(FTickFunction, TickGroup)), EPropertyLocation::Advanced);
+		}
+		
+		if (!HideCategories.Contains(TEXT("Events")))
+		{
+			AddEventsCategory(DetailLayout);
+		}	
 	}
 
 	PrimaryTickProperty->MarkHiddenByCustomization();
@@ -457,6 +469,100 @@ void FActorDetails::AddTransformCategory( IDetailLayoutBuilder& DetailBuilder )
 	TransformCategory.AddCustomBuilder( TransformDetails );
 }
 
+void FActorDetails::AddEventsCategory(IDetailLayoutBuilder& DetailBuilder)
+{
+	// Get the currently selected actor, which would be the "Default__Actor" 
+	const TArray<TWeakObjectPtr<UObject>>& Selected = DetailBuilder.GetSelectedObjects();
+	if(Selected.IsEmpty())
+	{
+		return;
+	}
+
+	AActor* Actor = Cast<AActor>(Selected[0].Get());
+	UBlueprint* Blueprint = Actor ? Cast<UBlueprint>(Actor->GetClass()->ClassGeneratedBy) : nullptr;
+
+	if(!Actor || !Blueprint || !FBlueprintEditorUtils::DoesSupportEventGraphs(Blueprint))
+	{
+		return;
+	}
+
+	IDetailCategoryBuilder& EventsCategory = DetailBuilder.EditCategory("Events", FText::GetEmpty(), ECategoryPriority::Uncommon);
+	static const FName HideInDetailPanelName("HideInDetailPanel");
+
+	// Find all the Multicast delegate properties and give a binding button for them
+	for (TFieldIterator<FMulticastDelegateProperty> PropertyIt(Actor->GetClass(), EFieldIteratorFlags::IncludeSuper); PropertyIt; ++PropertyIt)
+	{
+		FMulticastDelegateProperty* Property = *PropertyIt;
+		
+		// Only show BP assiangable, non-hidden delegates		
+		if (!Property->HasAnyPropertyFlags(CPF_Parm) && Property->HasAllPropertyFlags(CPF_BlueprintAssignable) && !Property->HasMetaData(HideInDetailPanelName))
+		{
+			const FName EventName = Property->GetFName();
+			FText EventText = Property->GetDisplayNameText();
+
+			EventsCategory.AddCustomRow(EventText)
+			.NameContent()
+			[
+				SNew(SHorizontalBox)
+				.ToolTipText(Property->GetToolTipText())
+
+				+ SHorizontalBox::Slot()
+				.AutoWidth()
+				.VAlign(VAlign_Center)
+				.Padding(0.0f, 0.0f, 5.0f, 0.0f)
+				[
+					SNew(SImage)
+					.Image(FEditorStyle::GetBrush("GraphEditor.Event_16x"))
+				]
+
+				+ SHorizontalBox::Slot()
+				.VAlign(VAlign_Center)
+				[
+					SNew(STextBlock)
+					.Font(IDetailLayoutBuilder::GetDetailFont())
+					.Text(EventText)
+				]
+			]
+			// A green "Plus" button to add a binding. For dynamic delegates on the CDO, you can always
+			// make a new binding, so always display the "Plus"
+			.ValueContent()
+			.MinDesiredWidth(150.0f)
+			.MaxDesiredWidth(200.0f)
+			[
+				SNew(SButton)
+				.ButtonStyle(FEditorStyle::Get(), "FlatButton.Success")
+				.HAlign(HAlign_Center)
+				.OnClicked(this, &FActorDetails::HandleAddOrViewEventForVariable, Blueprint, Property)
+				.ForegroundColor(FSlateColor::UseForeground())
+				[			
+					SNew(SImage)
+					.Image(FEditorStyle::GetBrush("Plus"))			
+				]
+			];
+		}
+	}
+}
+
+FReply FActorDetails::HandleAddOrViewEventForVariable(UBlueprint* BP, FMulticastDelegateProperty* Property)
+{
+	const UFunction* SignatureFunction = Property ? Property->SignatureFunction : nullptr;
+	UEdGraph* EventGraph = BP ? FBlueprintEditorUtils::FindEventGraph(BP) : nullptr;
+
+	if (EventGraph && SignatureFunction)
+	{
+		const FVector2D SpawnPos = EventGraph->GetGoodPlaceForNewNode();
+
+		// Adding a bound dynatic delegate from the Actor that is based off this BP will always be in a self context
+		UK2Node_AddDelegate* TemplateNode = NewObject<UK2Node_AddDelegate>();
+		TemplateNode->SetFromProperty(Property, /* bSelfContext */ true, Property->GetOwnerClass());
+
+		UEdGraphNode* SpawnedDelegate = FEdGraphSchemaAction_K2AssignDelegate::AssignDelegate(TemplateNode, EventGraph, nullptr, SpawnPos, true);
+		FKismetEditorUtilities::BringKismetToFocusAttentionOnObject(SpawnedDelegate, false);
+	}
+
+	return FReply::Handled();
+}
+
 namespace ActorDetailsUtil
 {
 	constexpr int32 MultipleValuesIndicator = 2;
diff --git a/Engine/Source/Editor/DetailCustomizations/Private/ActorDetails.h b/Engine/Source/Editor/DetailCustomizations/Private/ActorDetails.h
--- a/Engine/Source/Editor/DetailCustomizations/Private/ActorDetails.h
+++ b/Engine/Source/Editor/DetailCustomizations/Private/ActorDetails.h
@@ -85,6 +85,12 @@ private:
 
 	void AddTransformCategory( IDetailLayoutBuilder& DetailBuilder );
 
+	/** Display a category with all dynamic delegates on a CDO */
+	void AddEventsCategory(IDetailLayoutBuilder& DetailBuilder);
+
+	/** Handle the creation of a bound event from a dynamic delegate on a CDO */
+	FReply HandleAddOrViewEventForVariable(UBlueprint* BP, class FMulticastDelegateProperty* Property);
+
 	const TArray< TWeakObjectPtr<AActor> >& GetSelectedActors() const;
 
 	// Functions for actor loading strategy details
