Editor - Source file - importing data can now be cleared in bulk through some changes in the customization.

diff --git a/Engine/Source/Editor/DetailCustomizations/Private/AssetImportDataCustomization.cpp b/Engine/Source/Editor/DetailCustomizations/Private/AssetImportDataCustomization.cpp
--- a/Engine/Source/Editor/DetailCustomizations/Private/AssetImportDataCustomization.cpp
+++ b/Engine/Source/Editor/DetailCustomizations/Private/AssetImportDataCustomization.cpp
@@ -32,78 +32,29 @@ void FAssetImportDataCustomization::CustomizeChildren( TSharedRef<IPropertyHandl
 {
 	PropertyHandle = InPropertyHandle;
 
-	FAssetImportInfo* Info = GetEditStruct();
-	if (!Info)
-	{
-		return;
-	}
-
 	auto Font = IDetailLayoutBuilder::GetDetailFont();
 
-	const FText SourceFileText = LOCTEXT("SourceFile", "Source File");
-	int32 NumSourceFiles = Info->SourceFiles.Num() ? Info->SourceFiles.Num() : 1;
-
-	for (int32 Index = 0; Index < NumSourceFiles; ++Index)
+	FAssetImportInfo* Info = GetEditStruct();
+	if (Info)
 	{
-		FText SourceFileLabel = SourceFileText;
-		if (Info->SourceFiles.IsValidIndex(Index) && Info->SourceFiles[Index].DisplayLabelName.Len() > 0)
-		{
-			SourceFileLabel = FText::FromString(SourceFileText.ToString() + TEXT(" (") + Info->SourceFiles[Index].DisplayLabelName + TEXT(")"));
-		}
+		const FText SourceFileText = LOCTEXT("SourceFile", "Source File");
+		int32 NumSourceFiles = Info->SourceFiles.Num() ? Info->SourceFiles.Num() : 1;
 
-		ChildBuilder.AddCustomRow(SourceFileLabel)
-		.NameContent()
-		[
-			SNew(STextBlock)
-			.Text(SourceFileLabel)
-			.Font(Font)
-		]
-		.ValueContent()
-		.HAlign(HAlign_Fill)
-		.MaxDesiredWidth(TOptional<float>())
-		[
-			SNew(SHorizontalBox)
-			+ SHorizontalBox::Slot()
-			.VAlign(VAlign_Center)
+		for (int32 Index = 0; Index < NumSourceFiles; ++Index)
+		{
+			FText SourceFileLabel = SourceFileText;
+			if (Info->SourceFiles.IsValidIndex(Index) && Info->SourceFiles[Index].DisplayLabelName.Len() > 0)
+			{
+				SourceFileLabel = FText::FromString(SourceFileText.ToString() + TEXT(" (") + Info->SourceFiles[Index].DisplayLabelName + TEXT(")"));
+			}
+
+			ChildBuilder.AddCustomRow(SourceFileLabel)
+			.NameContent()
 			[
-				SNew(SEditableText)
-				.IsReadOnly(true)
-				.Text(this, &FAssetImportDataCustomization::GetFilenameText, Index)
-				.ToolTipText(this, &FAssetImportDataCustomization::GetFilenameText, Index)
+				SNew(STextBlock)
+				.Text(SourceFileLabel)
 				.Font(Font)
 			]
-
-			+SHorizontalBox::Slot()
-			.VAlign(VAlign_Center)
-			.AutoWidth()
-			[
-				SNew(SButton)
-				.OnClicked(this, &FAssetImportDataCustomization::OnChangePathClicked, Index)
-				.ToolTipText(LOCTEXT("ChangePath_Tooltip", "Browse for a new source file path"))
-				[
-					SNew(STextBlock)
-					.Text(LOCTEXT("...", "..."))
-					.Font(Font)
-				]
-			]
-
-			+SHorizontalBox::Slot()
-			.AutoWidth()
-			[
-				SNew(SButton)
-				.VAlign(VAlign_Center)
-				.HAlign(HAlign_Center)
-				.ButtonStyle(FEditorStyle::Get(), "HoverHintOnly")
-				.OnClicked(this, &FAssetImportDataCustomization::OnClearPathClicked, Index)
-				.ToolTipText(LOCTEXT("ClearPath_Tooltip", "Clear this source file information from the asset"))
-				[
-					SNew(SImage)
-					.Image(FEditorStyle::GetBrush("Cross"))
-				]
-			]
-		];
-
-		ChildBuilder.AddCustomRow(SourceFileText)
 			.ValueContent()
 			.HAlign(HAlign_Fill)
 			.MaxDesiredWidth(TOptional<float>())
@@ -111,46 +62,117 @@ void FAssetImportDataCustomization::CustomizeChildren( TSharedRef<IPropertyHandl
 				SNew(SHorizontalBox)
 				+ SHorizontalBox::Slot()
 				.VAlign(VAlign_Center)
-				.FillWidth(1.0f)
 				[
 					SNew(SEditableText)
 					.IsReadOnly(true)
-					.Text(this, &FAssetImportDataCustomization::GetTimestampText, Index)
+					.Text(this, &FAssetImportDataCustomization::GetFilenameText, Index)
+					.ToolTipText(this, &FAssetImportDataCustomization::GetFilenameText, Index)
 					.Font(Font)
 				]
-				+ SHorizontalBox::Slot()
+
+				+SHorizontalBox::Slot()
 				.VAlign(VAlign_Center)
-				.HAlign(HAlign_Center)
 				.AutoWidth()
 				[
 					SNew(SButton)
-					.VAlign(VAlign_Center)
-					.HAlign(HAlign_Center)
-					.ButtonStyle(FEditorStyle::Get(), "HoverHintOnly")
-					.IsEnabled(this, &FAssetImportDataCustomization::IsPropagateFromAbovePathEnable, Index)
-					.OnClicked(this, &FAssetImportDataCustomization::OnPropagateFromAbovePathClicked, Index)
-					.ToolTipText(LOCTEXT("PropagateFromAbovePath_Tooltip", "Use the above source path to set this path."))
+					.OnClicked(this, &FAssetImportDataCustomization::OnChangePathClicked, Index)
+					.ToolTipText(LOCTEXT("ChangePath_Tooltip", "Browse for a new source file path"))
 					[
-						SNew(SImage)
-						.Image(FEditorStyle::GetBrush("ArrowDown"))
+						SNew(STextBlock)
+						.Text(LOCTEXT("...", "..."))
+						.Font(Font)
 					]
 				]
-				+ SHorizontalBox::Slot()
-				.VAlign(VAlign_Center)
-				.HAlign(HAlign_Center)
+
+				+SHorizontalBox::Slot()
 				.AutoWidth()
 				[
 					SNew(SButton)
 					.VAlign(VAlign_Center)
 					.HAlign(HAlign_Center)
 					.ButtonStyle(FEditorStyle::Get(), "HoverHintOnly")
-					.IsEnabled(this, &FAssetImportDataCustomization::IsPropagateFromBelowPathEnable, Index)
-					.OnClicked(this, &FAssetImportDataCustomization::OnPropagateFromBelowPathClicked, Index)
-					.ToolTipText(LOCTEXT("PropagateFromBelowPath_Tooltip", "Use the below source path to set this path."))
+					.OnClicked(this, &FAssetImportDataCustomization::OnClearPathClicked, Index)
+					.ToolTipText(LOCTEXT("ClearPath_Tooltip", "Clear this source file information from the asset"))
 					[
 						SNew(SImage)
-						.Image(FEditorStyle::GetBrush("ArrowUp"))
+						.Image(FEditorStyle::GetBrush("Cross"))
+					]
+				]
+			];
+
+			ChildBuilder.AddCustomRow(SourceFileText)
+				.ValueContent()
+				.HAlign(HAlign_Fill)
+				.MaxDesiredWidth(TOptional<float>())
+				[
+					SNew(SHorizontalBox)
+					+ SHorizontalBox::Slot()
+					.VAlign(VAlign_Center)
+					.FillWidth(1.0f)
+					[
+						SNew(SEditableText)
+						.IsReadOnly(true)
+						.Text(this, &FAssetImportDataCustomization::GetTimestampText, Index)
+						.Font(Font)
+					]
+					+ SHorizontalBox::Slot()
+					.VAlign(VAlign_Center)
+					.HAlign(HAlign_Center)
+					.AutoWidth()
+					[
+						SNew(SButton)
+						.VAlign(VAlign_Center)
+						.HAlign(HAlign_Center)
+						.ButtonStyle(FEditorStyle::Get(), "HoverHintOnly")
+						.IsEnabled(this, &FAssetImportDataCustomization::IsPropagateFromAbovePathEnable, Index)
+						.OnClicked(this, &FAssetImportDataCustomization::OnPropagateFromAbovePathClicked, Index)
+						.ToolTipText(LOCTEXT("PropagateFromAbovePath_Tooltip", "Use the above source path to set this path."))
+						[
+							SNew(SImage)
+							.Image(FEditorStyle::GetBrush("ArrowDown"))
+						]
 					]
+					+ SHorizontalBox::Slot()
+					.VAlign(VAlign_Center)
+					.HAlign(HAlign_Center)
+					.AutoWidth()
+					[
+						SNew(SButton)
+						.VAlign(VAlign_Center)
+						.HAlign(HAlign_Center)
+						.ButtonStyle(FEditorStyle::Get(), "HoverHintOnly")
+						.IsEnabled(this, &FAssetImportDataCustomization::IsPropagateFromBelowPathEnable, Index)
+						.OnClicked(this, &FAssetImportDataCustomization::OnPropagateFromBelowPathClicked, Index)
+						.ToolTipText(LOCTEXT("PropagateFromBelowPath_Tooltip", "Use the below source path to set this path."))
+						[
+							SNew(SImage)
+							.Image(FEditorStyle::GetBrush("ArrowUp"))
+						]
+					]
+				];
+		}
+	}
+	else
+	{
+		ChildBuilder.AddCustomRow(FText::GetEmpty())
+			.NameContent()
+			[
+				SNew(STextBlock)
+				.Text(LOCTEXT("ClearAllSourceData", "Clear all the source file information from the asset."))
+				.Font(Font)
+			]
+			.ValueContent()
+			.HAlign(HAlign_Left)
+			[
+				SNew(SButton)
+				.VAlign(VAlign_Center)
+				.HAlign(HAlign_Center)
+				.ButtonStyle(FEditorStyle::Get(), "HoverHintOnly")
+				.OnClicked(this, &FAssetImportDataCustomization::OnClearAllPathsClicked)
+				.ToolTipText(LOCTEXT("ClearAllSourceData_Tooltip", "Clear all the source file information from the asset."))
+				[
+					SNew(SImage)
+					.Image(FEditorStyle::GetBrush("Cross"))
 				]
 			];
 	}
@@ -185,8 +207,7 @@ FText FAssetImportDataCustomization::GetTimestampText(int32 Index) const
 
 FAssetImportInfo* FAssetImportDataCustomization::GetEditStruct() const
 {
-	static TArray<FAssetImportInfo*> AssetImportInfo;
-	AssetImportInfo.Reset();
+	TArray<FAssetImportInfo*> AssetImportInfo;
 
 	if(PropertyHandle->IsValidHandle())
 	{
@@ -200,6 +221,18 @@ FAssetImportInfo* FAssetImportDataCustomization::GetEditStruct() const
 	return nullptr;
 }
 
+TArray<FAssetImportInfo*> FAssetImportDataCustomization::GetEditStructs() const
+{
+	TArray<FAssetImportInfo*> AssetImportInfo;
+
+	if (PropertyHandle->IsValidHandle())
+	{
+		PropertyHandle->AccessRawData(reinterpret_cast<TArray<void*>&>(AssetImportInfo));
+	}
+
+	return AssetImportInfo;
+}
+
 UAssetImportData* FAssetImportDataCustomization::GetOuterClass() const
 {
 	static TArray<UObject*> OuterObjects;
@@ -210,6 +243,20 @@ UAssetImportData* FAssetImportDataCustomization::GetOuterClass() const
 	return OuterObjects.Num() ? Cast<UAssetImportData>(OuterObjects[0]) : nullptr;
 }
 
+TArray<UAssetImportData*> FAssetImportDataCustomization::GetAllAssetImportData() const
+{
+	TArray<UObject*> OuterObjects;
+	PropertyHandle->GetOuterObjects(OuterObjects);
+
+	TArray<UAssetImportData*> AllAssetImportData;
+	for (UObject* Data : OuterObjects)
+	{
+		AllAssetImportData.Add(Cast<UAssetImportData>(Data));
+	}
+
+	return AllAssetImportData;
+}
+
 class FImportDataSourceFileTransactionScope
 {
 public:
@@ -285,6 +332,27 @@ FReply FAssetImportDataCustomization::OnClearPathClicked(int32 Index) const
 	return FReply::Handled();
 }
 
+FReply FAssetImportDataCustomization::OnClearAllPathsClicked() const
+{
+	TArray<UAssetImportData*> ImportDatas = GetAllAssetImportData();
+
+	FScopedTransaction Transaction(LOCTEXT("ClearAllSourceFilePaths", "Clear all source file paths"));
+	for (UAssetImportData* ImportData : ImportDatas)
+	{
+		FImportDataSourceFileTransactionScope TransactionScope(LOCTEXT("SourceReimportClearPath", "Clear Source file path"), ImportData);
+		
+		//Clear the filename, Hash and timestamp. Leave the Display label.
+		for (FAssetImportInfo::FSourceFile& SourceFile : ImportData->SourceData.SourceFiles)
+		{
+			SourceFile.RelativeFilename = FString();
+			SourceFile.FileHash = FMD5Hash();
+			SourceFile.Timestamp = 0;
+		}
+	}
+
+	return FReply::Handled();
+}
+
 bool FAssetImportDataCustomization::IsPropagateFromAbovePathEnable(int32 Index) const
 {
 	UAssetImportData* ImportData = GetOuterClass();
diff --git a/Engine/Source/Editor/DetailCustomizations/Private/AssetImportDataCustomization.h b/Engine/Source/Editor/DetailCustomizations/Private/AssetImportDataCustomization.h
--- a/Engine/Source/Editor/DetailCustomizations/Private/AssetImportDataCustomization.h
+++ b/Engine/Source/Editor/DetailCustomizations/Private/AssetImportDataCustomization.h
@@ -30,6 +30,8 @@ private:
 	/** Handle the user requesting that the specified index be cleared */
 	FReply OnClearPathClicked(int32 Index) const;
 
+	FReply OnClearAllPathsClicked() const;
+
 	/** Handle the user requesting that the Path use at Index - 1 will be replace the path at Index */
 	FReply OnPropagateFromAbovePathClicked(int32 Index) const;
 	bool IsPropagateFromAbovePathEnable(int32 Index) const;
@@ -38,12 +40,17 @@ private:
 	FReply OnPropagateFromBelowPathClicked(int32 Index) const;
 	bool IsPropagateFromBelowPathEnable(int32 Index) const;
 
-	/** Access the struct we are editing */
+	/** Access the struct we are editing - returns null if we have more than one. */
 	FAssetImportInfo* GetEditStruct() const;
 
+	/** Access all the structs we are editing */
+	TArray<FAssetImportInfo*> GetEditStructs() const;
+
 	/** Access the outer class that contains this struct */
 	UAssetImportData* GetOuterClass() const;
 
+	TArray<UAssetImportData*> GetAllAssetImportData() const;
+
 	/** Get text for the UI */
 	FText GetFilenameText(int32 Index) const;
 	FText GetTimestampText(int32 Index) const;
