Added FRuntimeVectorCurve and associated detail customization.

diff --git a/Engine/Source/Editor/DetailCustomizations/Private/CurveVectorCustomization.cpp b/Engine/Source/Editor/DetailCustomizations/Private/CurveVectorCustomization.cpp
new file mode 100644
--- /dev/null
+++ b/Engine/Source/Editor/DetailCustomizations/Private/CurveVectorCustomization.cpp
@@ -0,0 +1,472 @@
+// Copyright Epic Games, Inc. All Rights Reserved.
+
+#include "CurveVectorCustomization.h"
+#include "Curves/CurveVector.h"
+#include "Misc/PackageName.h"
+#include "Widgets/Text/STextBlock.h"
+#include "Editor.h"
+#include "Widgets/Layout/SBorder.h"
+#include "DetailWidgetRow.h"
+#include "Layout/WidgetPath.h"
+#include "Framework/Application/SlateApplication.h"
+#include "Widgets/Images/SImage.h"
+#include "Widgets/Input/SButton.h"
+#include "Dialogs/Dialogs.h"
+
+#include "Dialogs/DlgPickAssetPath.h"
+#include "PackageTools.h"
+#include "MiniCurveEditor.h"
+#include "AssetRegistryModule.h"
+#include "SCurveEditor.h"
+#include "Engine/Selection.h"
+#include "Subsystems/AssetEditorSubsystem.h"
+
+#define LOCTEXT_NAMESPACE "CurveVectorCustomization"
+
+const FVector2D FCurveVectorCustomization::DEFAULT_WINDOW_SIZE = FVector2D(800, 500);
+
+TSharedRef<IPropertyTypeCustomization> FCurveVectorCustomization::MakeInstance()
+{
+	return MakeShareable(new FCurveVectorCustomization);
+}
+
+FCurveVectorCustomization::~FCurveVectorCustomization()
+{
+	if (CurveWidget.IsValid() && CurveWidget->GetCurveOwner() == this)
+	{
+		CurveWidget->SetCurveOwner(nullptr, false);
+	}
+
+	DestroyPopOutWindow();
+}
+
+FCurveVectorCustomization::FCurveVectorCustomization()
+	: RuntimeCurve(NULL)
+	, Owner(NULL)
+	, ViewMinInput(0.0f)
+	, ViewMaxInput(5.0f)
+{
+}
+
+void FCurveVectorCustomization::CustomizeHeader(TSharedRef<IPropertyHandle> InStructPropertyHandle, FDetailWidgetRow& HeaderRow, IPropertyTypeCustomizationUtils& StructCustomizationUtils)
+{
+	this->StructPropertyHandle = InStructPropertyHandle;
+
+	TArray<UObject*> OuterObjects;
+	StructPropertyHandle->GetOuterObjects(OuterObjects);
+
+	TArray<void*> StructPtrs;
+	StructPropertyHandle->AccessRawData( StructPtrs );
+
+	if (StructPtrs.Num() == 1)
+	{
+		static const FName XAxisName(TEXT("XAxisName"));
+		static const FName YAxisName(TEXT("YAxisName"));
+
+		TOptional<FString> XAxisString;
+		if ( InStructPropertyHandle->HasMetaData(XAxisName) )
+		{
+			XAxisString = InStructPropertyHandle->GetMetaData(XAxisName);
+		}
+
+		TOptional<FString> YAxisString;
+		if ( InStructPropertyHandle->HasMetaData(YAxisName) )
+		{
+			YAxisString = InStructPropertyHandle->GetMetaData(YAxisName);
+		}
+
+		RuntimeCurve = reinterpret_cast<FRuntimeVectorCurve*>(StructPtrs[0]);
+
+		if (OuterObjects.Num() == 1)
+		{
+			Owner = OuterObjects[0];
+		}
+
+		HeaderRow
+			.NameContent()
+			[
+				InStructPropertyHandle->CreatePropertyNameWidget()
+			]
+			.ValueContent()
+			.HAlign(HAlign_Fill)
+			.MinDesiredWidth(200)
+			[
+				SNew(SBorder)
+				.VAlign(VAlign_Fill)
+				.OnMouseDoubleClick(this, &FCurveVectorCustomization::OnCurvePreviewDoubleClick)
+				[
+					SAssignNew(CurveWidget, SCurveEditor)
+					.ViewMinInput(this, &FCurveVectorCustomization::GetViewMinInput)
+					.ViewMaxInput(this, &FCurveVectorCustomization::GetViewMaxInput)
+					.TimelineLength(this, &FCurveVectorCustomization::GetTimelineLength)
+					.OnSetInputViewRange(this, &FCurveVectorCustomization::SetInputViewRange)
+					.XAxisName(XAxisString)
+					.YAxisName(YAxisString)
+					.HideUI(false)
+					.DesiredSize(FVector2D(300, 150))
+				]
+			];
+
+		check(CurveWidget.IsValid());
+		if (RuntimeCurve && RuntimeCurve->ExternalCurve)
+		{
+			CurveWidget->SetCurveOwner(RuntimeCurve->ExternalCurve, false);
+		}
+		else
+		{
+			CurveWidget->SetCurveOwner(this, InStructPropertyHandle->IsEditable());
+		}
+	}
+	else
+	{
+		HeaderRow
+			.NameContent()
+			[
+				InStructPropertyHandle->CreatePropertyNameWidget()
+			]
+			.ValueContent()
+			[
+				SNew(SBorder)
+				.VAlign(VAlign_Fill)
+				[
+					SNew(STextBlock)
+					.Text(StructPtrs.Num() == 0 ? LOCTEXT("NoCurves", "No Curves - unable to modify") : LOCTEXT("MultipleCurves", "Multiple Curves - unable to modify"))
+				]
+			];
+	}
+}
+
+void FCurveVectorCustomization::CustomizeChildren(TSharedRef<IPropertyHandle> InStructPropertyHandle, IDetailChildrenBuilder& StructBuilder, IPropertyTypeCustomizationUtils& StructCustomizationUtils)
+{
+	uint32 NumChildren = 0;
+	StructPropertyHandle->GetNumChildren(NumChildren);
+
+	for( uint32 ChildIndex = 0; ChildIndex < NumChildren; ++ChildIndex )
+	{
+		TSharedPtr<IPropertyHandle> Child = StructPropertyHandle->GetChildHandle( ChildIndex );
+
+		if( Child->GetProperty()->GetName() == TEXT("ExternalCurve") )
+		{
+			ExternalCurveHandle = Child;
+
+			FSimpleDelegate OnCurveChangedDelegate = FSimpleDelegate::CreateSP(this, &FCurveVectorCustomization::OnExternalCurveChanged, InStructPropertyHandle);
+			Child->SetOnPropertyValueChanged(OnCurveChangedDelegate);
+
+			StructBuilder.AddCustomRow(LOCTEXT("ExternalCurveLabel", "ExternalCurve"))
+				.NameContent()
+				[
+					Child->CreatePropertyNameWidget()
+				]
+				.ValueContent()
+				[
+					SNew(SVerticalBox)
+					+SVerticalBox::Slot()
+					.AutoHeight()
+					[
+						SNew(SHorizontalBox)
+						+SHorizontalBox::Slot()
+						.AutoWidth()
+						[
+							Child->CreatePropertyValueWidget()
+						]
+						+SHorizontalBox::Slot()
+						.AutoWidth()
+						.VAlign(VAlign_Center)
+						.Padding(1,0)
+						[
+							SNew(SButton)
+							.ButtonStyle( FEditorStyle::Get(), "NoBorder" )
+							.ContentPadding(1.f)
+							.ToolTipText(LOCTEXT("ConvertInternalCurveTooltip", "Convert to Internal Curve"))
+							.OnClicked(this, &FCurveVectorCustomization::OnConvertButtonClicked)
+							.IsEnabled(this, &FCurveVectorCustomization::IsConvertButtonEnabled)
+							[
+								SNew(SImage)
+								.Image( FEditorStyle::GetBrush(TEXT("PropertyWindow.Button_Clear")) )
+							]
+						]
+					]
+					+SVerticalBox::Slot()
+					.AutoHeight()
+					[
+						SNew(SHorizontalBox)
+						+SHorizontalBox::Slot()
+						.AutoWidth()
+						[
+							SNew(SButton)
+							.HAlign(HAlign_Center)
+							.Text( LOCTEXT( "CreateAssetButton", "Create External Curve" ) )
+							.ToolTipText(LOCTEXT( "CreateAssetTooltip", "Create a new CurveVector asset from this curve") )
+							.OnClicked(this, &FCurveVectorCustomization::OnCreateButtonClicked)
+							.IsEnabled(this, &FCurveVectorCustomization::IsCreateButtonEnabled)
+							.Visibility(Owner != nullptr ? EVisibility::Visible : EVisibility::Collapsed)
+						]
+					]
+				];
+		}
+		else
+		{
+			StructBuilder.AddProperty(Child.ToSharedRef());
+		}
+	}
+}
+
+static const FName XCurveName(TEXT("X"));
+static const FName YCurveName(TEXT("Y"));
+static const FName ZCurveName(TEXT("Z"));
+
+TArray<FRichCurveEditInfoConst> FCurveVectorCustomization::GetCurves() const
+{
+	TArray<FRichCurveEditInfoConst> Curves;
+	Curves.Add(FRichCurveEditInfoConst(&RuntimeCurve->VectorCurves[0], XCurveName));
+	Curves.Add(FRichCurveEditInfoConst(&RuntimeCurve->VectorCurves[1], YCurveName));
+	Curves.Add(FRichCurveEditInfoConst(&RuntimeCurve->VectorCurves[2], ZCurveName));
+	return Curves;
+}
+
+TArray<FRichCurveEditInfo> FCurveVectorCustomization::GetCurves()
+{
+	TArray<FRichCurveEditInfo> Curves;
+	Curves.Add(FRichCurveEditInfo(&RuntimeCurve->VectorCurves[0], XCurveName));
+	Curves.Add(FRichCurveEditInfo(&RuntimeCurve->VectorCurves[1], YCurveName));
+	Curves.Add(FRichCurveEditInfo(&RuntimeCurve->VectorCurves[2], ZCurveName));
+	return Curves;
+}
+
+void FCurveVectorCustomization::ModifyOwner()
+{
+	if (Owner)
+	{
+		Owner->Modify(true);
+	}
+}
+
+TArray<const UObject*> FCurveVectorCustomization::GetOwners() const
+{
+	TArray<const UObject*> Owners;
+	if (Owner)
+	{
+		Owners.Add(Owner);
+	}
+
+	return Owners;
+}
+
+void FCurveVectorCustomization::MakeTransactional()
+{
+	if (Owner)
+	{
+		Owner->SetFlags(Owner->GetFlags() | RF_Transactional);
+	}
+}
+
+void FCurveVectorCustomization::OnCurveChanged(const TArray<FRichCurveEditInfo>& ChangedCurveEditInfos)
+{
+	StructPropertyHandle->NotifyPostChange(EPropertyChangeType::Unspecified);
+}
+
+bool FCurveVectorCustomization::IsValidCurve(FRichCurveEditInfo CurveInfo)
+{
+	return CurveInfo.CurveToEdit == &RuntimeCurve->VectorCurves[0] ||
+		   CurveInfo.CurveToEdit == &RuntimeCurve->VectorCurves[1] ||
+		   CurveInfo.CurveToEdit == &RuntimeCurve->VectorCurves[2];
+}
+
+float FCurveVectorCustomization::GetTimelineLength() const
+{
+	return 0.f;
+}
+
+void FCurveVectorCustomization::SetInputViewRange(float InViewMinInput, float InViewMaxInput)
+{
+	ViewMaxInput = InViewMaxInput;
+	ViewMinInput = InViewMinInput;
+}
+
+void FCurveVectorCustomization::OnExternalCurveChanged(TSharedRef<IPropertyHandle> CurvePropertyHandle)
+{
+	if (RuntimeCurve)
+	{
+		if (RuntimeCurve->ExternalCurve)
+		{
+			CurveWidget->SetCurveOwner(RuntimeCurve->ExternalCurve, false);
+		}
+		else
+		{
+			CurveWidget->SetCurveOwner(this, CurvePropertyHandle->IsEditable());
+		}
+	}
+
+	CurvePropertyHandle->NotifyPostChange(EPropertyChangeType::ValueSet);
+}
+
+FReply FCurveVectorCustomization::OnCreateButtonClicked()
+{
+	if (CurveWidget.IsValid())
+	{
+		FString DefaultAsset = FPackageName::GetLongPackagePath(Owner->GetOutermost()->GetName()) + TEXT("/") + Owner->GetName() + TEXT("_ExternalCurve");
+
+		TSharedRef<SDlgPickAssetPath> NewCurveDlg =
+			SNew(SDlgPickAssetPath)
+			.Title(LOCTEXT("NewCurveDialogTitle", "Choose Location for External Curve Asset"))
+			.DefaultAssetPath(FText::FromString(DefaultAsset));
+
+		if (NewCurveDlg->ShowModal() != EAppReturnType::Cancel)
+		{
+			FString Package(NewCurveDlg->GetFullAssetPath().ToString());
+			FString Name(NewCurveDlg->GetAssetName().ToString());
+
+			// Find (or create!) the desired package for this object
+			UPackage* Pkg = CreatePackage( *Package);
+			UPackage* OutermostPkg = Pkg->GetOutermost();
+
+			TArray<UPackage*> TopLevelPackages;
+			TopLevelPackages.Add( OutermostPkg );
+			if (!UPackageTools::HandleFullyLoadingPackages(TopLevelPackages, LOCTEXT("CreateANewObject", "Create a new object")))
+			{
+				// User aborted.
+				return FReply::Handled();
+			}
+
+			if (!PromptUserIfExistingObject(Name, Package, Pkg))
+			{
+				return FReply::Handled();
+			}
+
+			// PromptUserIfExistingObject may have GCed and recreated our outermost package - re-acquire it here.
+			OutermostPkg = Pkg->GetOutermost();
+
+			// Create a new asset and set it as the external curve
+			FName AssetName = *Name;
+			UCurveVector* NewCurve = Cast<UCurveVector>(CurveWidget->CreateCurveObject(UCurveVector::StaticClass(), Pkg, AssetName));
+			if (NewCurve)
+			{
+				// run through points of editor data and add to external curve
+				for (int32 Index = 0; Index < 3; Index++)
+				{
+					CopyCurveData(&RuntimeCurve->VectorCurves[Index], &NewCurve->FloatCurves[Index]);
+				}
+
+				// Set the new object as the sole selection.
+				USelection* SelectionSet = GEditor->GetSelectedObjects();
+				SelectionSet->DeselectAll();
+				SelectionSet->Select( NewCurve );
+
+				// Notify the asset registry
+				FAssetRegistryModule::AssetCreated(NewCurve);
+
+				// Mark the package dirty...
+				OutermostPkg->MarkPackageDirty();
+
+				// Make sure expected type of pointer passed to SetValue, so that it's not interpreted as a bool
+				ExternalCurveHandle->SetValue(NewCurve);
+			}
+		}
+	}
+	return FReply::Handled();
+}
+
+bool FCurveVectorCustomization::IsCreateButtonEnabled() const
+{
+	return CurveWidget.IsValid() && RuntimeCurve != NULL && RuntimeCurve->ExternalCurve == NULL;
+}
+
+FReply FCurveVectorCustomization::OnConvertButtonClicked()
+{
+	if (RuntimeCurve && RuntimeCurve->ExternalCurve)
+	{
+		// clear points of editor data
+		for (int32 Index = 0; Index < 3; Index++)
+		{
+			RuntimeCurve->VectorCurves[Index].Reset();
+		}
+
+		// run through points of external curve and add to editor data
+		for (int32 Index = 0; Index < 3; Index++)
+		{
+			CopyCurveData(&RuntimeCurve->ExternalCurve->FloatCurves[Index], &RuntimeCurve->VectorCurves[Index]);
+		}
+
+		// null out external curve
+		const UObject* NullObject = NULL;
+		ExternalCurveHandle->SetValue(NullObject);
+	}
+	return FReply::Handled();
+}
+
+bool FCurveVectorCustomization::IsConvertButtonEnabled() const
+{
+	return RuntimeCurve != NULL && RuntimeCurve->ExternalCurve != NULL;
+}
+
+FReply FCurveVectorCustomization::OnCurvePreviewDoubleClick(const FGeometry& InMyGeometry, const FPointerEvent& InMouseEvent)
+{
+	if (InMouseEvent.GetEffectingButton() == EKeys::LeftMouseButton)
+	{
+		if (RuntimeCurve->ExternalCurve)
+		{
+			GEditor->GetEditorSubsystem<UAssetEditorSubsystem>()->OpenEditorForAsset(RuntimeCurve->ExternalCurve);
+		}
+		else
+		{
+			DestroyPopOutWindow();
+
+			// Determine the position of the window so that it will spawn near the mouse, but not go off the screen.
+			const FVector2D CursorPos = FSlateApplication::Get().GetCursorPos();
+			FSlateRect Anchor(CursorPos.X, CursorPos.Y, CursorPos.X, CursorPos.Y);
+
+			FVector2D AdjustedSummonLocation = FSlateApplication::Get().CalculatePopupWindowPosition(Anchor, FCurveVectorCustomization::DEFAULT_WINDOW_SIZE, true, FVector2D::ZeroVector, Orient_Horizontal);
+
+			TSharedPtr<SWindow> Window = SNew(SWindow)
+				.Title(FText::Format(LOCTEXT("WindowHeader", "{0} - Internal Vector Curve Editor"), StructPropertyHandle->GetPropertyDisplayName()))
+				.ClientSize(FCurveVectorCustomization::DEFAULT_WINDOW_SIZE)
+				.ScreenPosition(AdjustedSummonLocation)
+				.AutoCenter(EAutoCenter::None)
+				.SupportsMaximize(false)
+				.SupportsMinimize(false)
+				.SizingRule( ESizingRule::FixedSize );
+
+			// init the mini curve editor widget
+			TSharedRef<SMiniCurveEditor> MiniCurveEditor =
+				SNew(SMiniCurveEditor)
+				.CurveOwner(this)
+				.OwnerObject(Owner)
+				.ParentWindow(Window);
+
+			Window->SetContent( MiniCurveEditor );
+
+			// Find the window of the parent widget
+			FWidgetPath WidgetPath;
+			FSlateApplication::Get().GeneratePathToWidgetChecked( CurveWidget.ToSharedRef(), WidgetPath );
+			Window = FSlateApplication::Get().AddWindowAsNativeChild( Window.ToSharedRef(), WidgetPath.GetWindow() );
+
+			//hold on to the window created for external use...
+			CurveEditorWindow = Window;
+		}
+	}
+	return FReply::Handled();
+}
+
+void FCurveVectorCustomization::CopyCurveData(const FRichCurve* SrcCurve, FRichCurve* DestCurve)
+{
+	if( SrcCurve && DestCurve )
+	{
+		for (auto It(SrcCurve->GetKeyIterator()); It; ++It)
+		{
+			const FRichCurveKey& Key = *It;
+			FKeyHandle KeyHandle = DestCurve->AddKey(Key.Time, Key.Value);
+			DestCurve->GetKey(KeyHandle) = Key;
+		}
+	}
+}
+
+void FCurveVectorCustomization::DestroyPopOutWindow()
+{
+	if (CurveEditorWindow.IsValid())
+	{
+		CurveEditorWindow.Pin()->RequestDestroyWindow();
+		CurveEditorWindow.Reset();
+	}
+}
+
+#undef LOCTEXT_NAMESPACE
diff --git a/Engine/Source/Editor/DetailCustomizations/Private/CurveVectorCustomization.h b/Engine/Source/Editor/DetailCustomizations/Private/CurveVectorCustomization.h
new file mode 100644
--- /dev/null
+++ b/Engine/Source/Editor/DetailCustomizations/Private/CurveVectorCustomization.h
@@ -0,0 +1,132 @@
+// Copyright Epic Games, Inc. All Rights Reserved.
+
+#pragma once
+
+#include "CoreMinimal.h"
+#include "Input/Reply.h"
+#include "IPropertyTypeCustomization.h"
+#include "PropertyHandle.h"
+#include "Curves/CurveOwnerInterface.h"
+#include "IDetailChildrenBuilder.h"
+
+class FDetailWidgetRow;
+class SCurveEditor;
+struct FRuntimeVectorCurve;
+
+/**
+ * Customizes a RuntimeVectorCurve struct to display a Curve Editor
+ */
+class FCurveVectorCustomization : public IPropertyTypeCustomization, public FCurveOwnerInterface
+{
+public:
+	static TSharedRef<IPropertyTypeCustomization> MakeInstance();
+
+	virtual ~FCurveVectorCustomization();
+
+	/** IPropertyTypeCustomization interface */
+	virtual void CustomizeHeader( TSharedRef<IPropertyHandle> InStructPropertyHandle, FDetailWidgetRow& HeaderRow, IPropertyTypeCustomizationUtils& StructCustomizationUtils ) override;
+	virtual void CustomizeChildren( TSharedRef<IPropertyHandle> InStructPropertyHandle, IDetailChildrenBuilder& StructBuilder, IPropertyTypeCustomizationUtils& StructCustomizationUtils ) override;
+
+	/** FCurveOwnerInterface interface */
+	virtual TArray<FRichCurveEditInfoConst> GetCurves() const override;
+	virtual TArray<FRichCurveEditInfo> GetCurves() override;
+	virtual void ModifyOwner() override;
+	virtual TArray<const UObject*> GetOwners() const override;
+	virtual void MakeTransactional() override;
+	virtual void OnCurveChanged(const TArray<FRichCurveEditInfo>& ChangedCurveEditInfos) override;
+	virtual bool IsValidCurve( FRichCurveEditInfo CurveInfo ) override;
+
+private:
+	/**
+	 * Constructor
+	 */
+	FCurveVectorCustomization();
+
+	/**
+	 * Get View Min Input for the Curve Editor
+	 */
+	float GetViewMinInput() const { return ViewMinInput; }
+
+	/**
+	 * Get View Min Input for the Curve Editor
+	 */
+	float GetViewMaxInput() const { return ViewMaxInput; }
+
+	/**
+	 * Get Timeline Length for the Curve Editor
+	 */
+	float GetTimelineLength() const;
+
+	/**
+	 * Set View Min and Max Inputs for the Curve Editor
+	 */
+	void SetInputViewRange(float InViewMinInput, float InViewMaxInput);
+
+	/**
+	 * Called when RuntimeVectorCurve's External Curve is changed
+	 */
+	void OnExternalCurveChanged(TSharedRef<IPropertyHandle> CurvePropertyHandle);
+
+	/**
+	 * Called when button clicked to create an External Vector Curve
+	 */
+	FReply OnCreateButtonClicked();
+
+	/**
+	 * Whether Create button is enabled
+	 */
+	bool IsCreateButtonEnabled() const;
+
+	/**
+	 * Called when button clicked to convert External Curve to internal curve
+	 */
+	FReply OnConvertButtonClicked();
+
+	/**
+	 * Whether Convert button is enabled
+	 */
+	bool IsConvertButtonEnabled() const;
+
+	/**
+	 * Called when user double clicks on the curve preview to open a full size editor
+	 */
+	FReply OnCurvePreviewDoubleClick(const FGeometry& InMyGeometry, const FPointerEvent& InMouseEvent);
+
+	/**
+	 * Copies data from one Rich Curve to another
+	 */
+	static void CopyCurveData(const FRichCurve* SrcCurve, FRichCurve* DestCurve);
+
+	/**
+	 * Destroys the pop out window used for editing internal curves.
+	 */
+	void DestroyPopOutWindow();
+
+private:
+	/** Cached RuntimeCurveLinearColor struct handle */
+	TSharedPtr<IPropertyHandle> StructPropertyHandle;
+
+	/** Cached External Curve handle */
+	TSharedPtr<IPropertyHandle> ExternalCurveHandle;
+
+	/** Small preview Curve Editor */
+	TSharedPtr<SCurveEditor> CurveWidget;
+
+	/** Window for pop out Curve Editor */
+	TWeakPtr<SWindow> CurveEditorWindow;
+
+	/** Cached pointer to the actual RuntimeVectorCurve struct */
+	FRuntimeVectorCurve* RuntimeCurve;
+
+	/** Object that owns the RuntimeFloatCurve */
+	UObject* Owner;
+
+	/** View Min Input for the Curve Editor */
+	float ViewMinInput;
+
+	/** View Max Input for the Curve Editor */
+	float ViewMaxInput;
+
+	/** Size of the pop out Curve Editor window */
+	static const FVector2D DEFAULT_WINDOW_SIZE;
+};
diff --git a/Engine/Source/Editor/DetailCustomizations/Private/DetailCustomizations.cpp b/Engine/Source/Editor/DetailCustomizations/Private/DetailCustomizations.cpp
--- a/Engine/Source/Editor/DetailCustomizations/Private/DetailCustomizations.cpp
+++ b/Engine/Source/Editor/DetailCustomizations/Private/DetailCustomizations.cpp
@@ -60,6 +60,7 @@
 #include "Customizations/ColorStructCustomization.h"
 #include "SlateColorCustomization.h"
 #include "CurveStructCustomization.h"
+#include "CurveVectorCustomization.h"
 #include "NavLinkStructCustomization.h"
 #include "NavAgentSelectorCustomization.h"
 #include "DirectoryPathStructCustomization.h"
@@ -234,6 +235,7 @@ void FDetailCustomizationsModule::RegisterPropertyTypeCustomizations()
 	RegisterCustomPropertyTypeLayout("AIDataProviderFloatValue", FOnGetPropertyTypeCustomizationInstance::CreateStatic(&FAIDataProviderValueDetails::MakeInstance));
 	RegisterCustomPropertyTypeLayout("AIDataProviderBoolValue", FOnGetPropertyTypeCustomizationInstance::CreateStatic(&FAIDataProviderValueDetails::MakeInstance));
 	RegisterCustomPropertyTypeLayout("RuntimeFloatCurve", FOnGetPropertyTypeCustomizationInstance::CreateStatic(&FCurveStructCustomization::MakeInstance));
+	RegisterCustomPropertyTypeLayout("RuntimeVectorCurve", FOnGetPropertyTypeCustomizationInstance::CreateStatic(&FCurveVectorCustomization::MakeInstance));
 	RegisterCustomPropertyTypeLayout("EnvNamedValue", FOnGetPropertyTypeCustomizationInstance::CreateStatic(&FEnvQueryParamInstanceCustomization::MakeInstance));
 	RegisterCustomPropertyTypeLayout("NavigationLink", FOnGetPropertyTypeCustomizationInstance::CreateStatic(&FNavLinkStructCustomization::MakeInstance));
 	RegisterCustomPropertyTypeLayout("NavigationSegmentLink", FOnGetPropertyTypeCustomizationInstance::CreateStatic(&FNavLinkStructCustomization::MakeInstance));
diff --git a/Engine/Source/Runtime/Engine/Classes/Curves/CurveVector.h b/Engine/Source/Runtime/Engine/Classes/Curves/CurveVector.h
--- a/Engine/Source/Runtime/Engine/Classes/Curves/CurveVector.h
+++ b/Engine/Source/Runtime/Engine/Classes/Curves/CurveVector.h
@@ -2,13 +2,29 @@
 
 #pragma once
 
-
 #include "CoreMinimal.h"
 #include "UObject/ObjectMacros.h"
 #include "Curves/RichCurve.h"
 #include "Curves/CurveBase.h"
 #include "CurveVector.generated.h"
 
+USTRUCT(BlueprintType)
+struct ENGINE_API FRuntimeVectorCurve
+{
+	GENERATED_USTRUCT_BODY()
+
+	UPROPERTY()
+	FRichCurve VectorCurves[3];
+
+	UPROPERTY(EditAnywhere, Category = RuntimeFloatCurve)
+	TObjectPtr<class UCurveVector> ExternalCurve = nullptr;
+
+	FVector GetValue(float InTime) const;
+	
+	/** Get the current curve struct */
+    FRichCurve* GetRichCurve(int32 Index);
+	const FRichCurve* GetRichCurveConst(int32 Index) const;
+};
 
 UCLASS(BlueprintType, MinimalAPI)
 class UCurveVector : public UCurveBase
diff --git a/Engine/Source/Runtime/Engine/Private/Curves/CurveVector.cpp b/Engine/Source/Runtime/Engine/Private/Curves/CurveVector.cpp
--- a/Engine/Source/Runtime/Engine/Private/Curves/CurveVector.cpp
+++ b/Engine/Source/Runtime/Engine/Private/Curves/CurveVector.cpp
@@ -6,6 +6,53 @@
 
 #include "Curves/CurveVector.h"
 
+FVector FRuntimeVectorCurve::GetValue(float InTime) const
+{
+	if (ExternalCurve)
+	{
+		return ExternalCurve->GetVectorValue(InTime);
+	}
+
+	FVector Result;
+	Result.X = VectorCurves[0].Eval(InTime);
+	Result.Y = VectorCurves[1].Eval(InTime);
+	Result.Z = VectorCurves[2].Eval(InTime);
+	return Result;
+}
+
+FRichCurve* FRuntimeVectorCurve::GetRichCurve(int32 Index)
+{
+	if (Index < 0 || Index >= 3)
+	{
+		return nullptr; 
+	}
+
+	if (ExternalCurve != nullptr)
+	{
+		return &(ExternalCurve->FloatCurves[Index]);
+	}
+	else
+	{
+		return &(VectorCurves[Index]);
+	}
+}
+
+const FRichCurve* FRuntimeVectorCurve::GetRichCurveConst(int32 Index) const
+{
+	if (Index < 0 || Index >= 3)
+	{
+		return nullptr; 
+	}
+	
+	if (ExternalCurve != nullptr)
+	{
+		return &(ExternalCurve->FloatCurves[Index]);
+	}
+	else
+	{
+		return &(VectorCurves[Index]);
+	}
+}
 
 UCurveVector::UCurveVector(const FObjectInitializer& ObjectInitializer)
 	: Super(ObjectInitializer)
