Fix component picker. Only show the actors that have a valid component.

diff --git a/Engine/Source/Editor/DetailCustomizations/Private/ComponentReferenceCustomization.cpp b/Engine/Source/Editor/DetailCustomizations/Private/ComponentReferenceCustomization.cpp
--- a/Engine/Source/Editor/DetailCustomizations/Private/ComponentReferenceCustomization.cpp
+++ b/Engine/Source/Editor/DetailCustomizations/Private/ComponentReferenceCustomization.cpp
@@ -57,7 +57,23 @@ void FComponentReferenceCustomization::CustomizeHeader(TSharedRef<IPropertyHandl
 		BuildComboBox();
 
 		InPropertyHandle->SetOnPropertyValueChanged(FSimpleDelegate::CreateSP(this, &FComponentReferenceCustomization::OnPropertyValueChanged));
-		OnPropertyValueChanged();
+
+		// set cached values
+		{
+			CachedComponent.Reset();
+			CachedFirstOuterActor = GetFirstOuterActor();
+
+			FComponentReference TmpComponentReference;
+			CachedPropertyAccess = GetValue(TmpComponentReference);
+			if (CachedPropertyAccess == FPropertyAccess::Success)
+			{
+				CachedComponent = TmpComponentReference.GetComponent(CachedFirstOuterActor.Get());
+				if (!IsComponentReferenceValid(TmpComponentReference))
+				{
+					CachedComponent.Reset();
+				}
+			}
+		}
 
 		HeaderRow.NameContent()
 		[
@@ -558,7 +574,7 @@ void FComponentReferenceCustomization::OnMenuOpenChanged(bool bOpen)
 
 bool FComponentReferenceCustomization::IsFilteredActor(const AActor* const Actor) const
 {
-	return false;
+	return bAllowAnyActor || Actor == CachedFirstOuterActor.Get();
 }
 
 bool FComponentReferenceCustomization::IsFilteredComponent(const UActorComponent* const Component) const
diff --git a/Engine/Source/Editor/SceneOutliner/Private/ActorHierarchy.cpp b/Engine/Source/Editor/SceneOutliner/Private/ActorHierarchy.cpp
--- a/Engine/Source/Editor/SceneOutliner/Private/ActorHierarchy.cpp
+++ b/Engine/Source/Editor/SceneOutliner/Private/ActorHierarchy.cpp
@@ -278,10 +278,27 @@ void FActorHierarchy::CreateWorldChildren(UWorld* World, TArray<FSceneOutlinerTr
 
 		if (FSceneOutlinerTreeItemPtr ActorItem = Mode->CreateItemFor<FActorTreeItem>(Actor))
 		{
-			OutItems.Add(ActorItem);
+			if (bShowingOnlyActorWithValidComponents)
+			{
+				int32 InsertLocation = OutItems.Num();
+
+				// Create all component items
+				CreateComponentItems(Actor, OutItems);
+
+				if (OutItems.Num() != InsertLocation)
+				{
+					// Add the actor before the components
+					OutItems.Insert(ActorItem, InsertLocation);
+				}
+			}
+			else
+			{
+				OutItems.Add(ActorItem);
+
+				// Create all component items
+				CreateComponentItems(Actor, OutItems);
+			}
 
-			// Create all component items
-			CreateComponentItems(Actor, OutItems);
 		}
 	}
 
diff --git a/Engine/Source/Editor/SceneOutliner/Private/ActorMode.cpp b/Engine/Source/Editor/SceneOutliner/Private/ActorMode.cpp
--- a/Engine/Source/Editor/SceneOutliner/Private/ActorMode.cpp
+++ b/Engine/Source/Editor/SceneOutliner/Private/ActorMode.cpp
@@ -78,6 +78,7 @@ FActorMode::FActorMode(const FActorModeParams& Params)
 	: ISceneOutlinerMode(Params.SceneOutliner)
 	, SpecifiedWorldToDisplay(Params.SpecifiedWorldToDisplay)
 	, bHideComponents(Params.bHideComponents)
+	, bHideActorWithNoComponent(Params.bHideActorWithNoComponent)
 	, bHideLevelInstanceHierarchy(Params.bHideLevelInstanceHierarchy)
 	, bHideUnloadedActors(Params.bHideUnloadedActors)
 {
@@ -95,6 +96,7 @@ TUniquePtr<ISceneOutlinerHierarchy> FActorMode::CreateHierarchy()
 {
 	TUniquePtr<FActorHierarchy> ActorHierarchy = FActorHierarchy::Create(this, RepresentingWorld);
 	ActorHierarchy->SetShowingComponents(!bHideComponents);
+	ActorHierarchy->SetShowingOnlyActorWithValidComponents(!bHideComponents && bHideActorWithNoComponent);
 	ActorHierarchy->SetShowingLevelInstances(!bHideLevelInstanceHierarchy);
 	ActorHierarchy->SetShowingUnloadedActors(!bHideUnloadedActors);
 
diff --git a/Engine/Source/Editor/SceneOutliner/Private/SceneOutlinerModule.cpp b/Engine/Source/Editor/SceneOutliner/Private/SceneOutlinerModule.cpp
--- a/Engine/Source/Editor/SceneOutliner/Private/SceneOutlinerModule.cpp
+++ b/Engine/Source/Editor/SceneOutliner/Private/SceneOutlinerModule.cpp
@@ -104,6 +104,7 @@ TSharedRef<ISceneOutliner> FSceneOutlinerModule::CreateComponentPicker(const FSc
 			Params.SceneOutliner = Outliner;
 			Params.SpecifiedWorldToDisplay = SpecifiedWorld;
 			Params.bHideComponents = false;
+			Params.bHideActorWithNoComponent = true;
 			Params.bHideLevelInstanceHierarchy = true;
 			Params.bHideUnloadedActors = true;
 			return new FActorPickingMode(Params, OnItemPicked);
diff --git a/Engine/Source/Editor/SceneOutliner/Public/ActorHierarchy.h b/Engine/Source/Editor/SceneOutliner/Public/ActorHierarchy.h
--- a/Engine/Source/Editor/SceneOutliner/Public/ActorHierarchy.h
+++ b/Engine/Source/Editor/SceneOutliner/Public/ActorHierarchy.h
@@ -25,6 +25,7 @@ public:
 	virtual FSceneOutlinerTreeItemPtr CreateParentItem(const FSceneOutlinerTreeItemPtr& Item) const override;
 
 	void SetShowingComponents(bool bInShowingComponents) { bShowingComponents = bInShowingComponents; }
+	void SetShowingOnlyActorWithValidComponents(bool bInShowingOnlyActorWithValidComponents) { bShowingOnlyActorWithValidComponents = bInShowingOnlyActorWithValidComponents; }
 	void SetShowingLevelInstances(bool bInShowingLevelInstances) { bShowingLevelInstances = bInShowingLevelInstances; }
 	void SetShowingUnloadedActors(bool bInShowingUnloadedActors) { bShowingUnloadedActors = bInShowingUnloadedActors; }
 private:
@@ -72,6 +73,7 @@ private:
 	TWeakObjectPtr<UWorld> RepresentingWorld;
 	
 	bool bShowingComponents = false;
+	bool bShowingOnlyActorWithValidComponents = false;
 	bool bShowingLevelInstances = false;
 	bool bShowingUnloadedActors = false;
 private:
diff --git a/Engine/Source/Editor/SceneOutliner/Public/ActorMode.h b/Engine/Source/Editor/SceneOutliner/Public/ActorMode.h
--- a/Engine/Source/Editor/SceneOutliner/Public/ActorMode.h
+++ b/Engine/Source/Editor/SceneOutliner/Public/ActorMode.h
@@ -34,6 +34,7 @@ struct SCENEOUTLINER_API FActorModeParams
 	TWeakObjectPtr<UWorld> SpecifiedWorldToDisplay = nullptr;
 	SSceneOutliner* SceneOutliner = nullptr;
 	bool bHideComponents = true;
+	bool bHideActorWithNoComponent = false;
 	bool bHideLevelInstanceHierarchy = true;
 	bool bHideUnloadedActors = true;
 };
@@ -84,6 +85,8 @@ protected:
 
 	/** Should components be hidden */
 	bool bHideComponents;
+	/** Should actor with no component be hidden. */
+	bool bHideActorWithNoComponent;
 	/** Should the level instance hierarchy be hidden */
 	bool bHideLevelInstanceHierarchy;
 	/** Should unloaded actors be hidden */
