Make one file per actor support experimental and off by default.

diff --git a/Engine/Source/Editor/DetailCustomizations/Private/ActorDetails.cpp b/Engine/Source/Editor/DetailCustomizations/Private/ActorDetails.cpp
--- a/Engine/Source/Editor/DetailCustomizations/Private/ActorDetails.cpp
+++ b/Engine/Source/Editor/DetailCustomizations/Private/ActorDetails.cpp
@@ -55,6 +55,7 @@
 #include "Misc/MessageDialog.h"
 #include "ScopedTransaction.h"
 #include "WorldPartition/WorldPartitionSubsystem.h"
+#include "Settings/EditorExperimentalSettings.h"
 
 #include "Kismet2/BlueprintEditorUtils.h"
 #include "K2Node_AddDelegate.h"
@@ -659,41 +660,44 @@ void FActorDetails::AddActorCategory( IDetailLayoutBuilder& DetailBuilder, const
 		];
 	}
 
-	// WorldSettings should not be packaged externally
-	if (SelectedActorInfo.SelectionClass != AWorldSettings::StaticClass())
+	if (GetDefault<UEditorExperimentalSettings>()->bEnableOneFilePerActorSupport)
 	{
-		// Actor Packaging Mode
-		const bool bIsPartitionedWorld = UWorld::HasSubsystem<UWorldPartitionSubsystem>(SelectedActorInfo.SharedWorld);
-
-		auto OnGetMenuContent = [=]() -> TSharedRef<SWidget> {
-			FMenuBuilder MenuBuilder(true, nullptr);
-			MenuBuilder.AddMenuEntry(ActorDetailsUtil::GetActorPackagingModeText(0), FText(), FSlateIcon(), FExecuteAction::CreateSP(this, &FActorDetails::OnActorPackagingModeChanged, false));
-			MenuBuilder.AddMenuEntry(ActorDetailsUtil::GetActorPackagingModeText(1), FText(), FSlateIcon(), FExecuteAction::CreateSP(this, &FActorDetails::OnActorPackagingModeChanged, true));
-			return MenuBuilder.MakeWidget();
-		};
-
-		ActorCategory.AddCustomRow( LOCTEXT("ActorPackagingModeRow", "ActorPackagingMode"), true)
-			.NameContent()
-			[
-				SNew(STextBlock)
-				.Text(LOCTEXT("ActorPackagingMode", "Packaging Mode"))
-				.ToolTipText(LOCTEXT("ActorPackagingMode_ToolTip", "Change the actor packaging mode. This will indicate if the actor is packaged alongside its level or in an external package."))
-				.Font(IDetailLayoutBuilder::GetDetailFont())
-			]
-			.ValueContent()
-			[
-				SNew(SComboButton)
-				.ContentPadding(2)
-				.OnGetMenuContent_Lambda(OnGetMenuContent)
-				.IsEnabled(this, &FActorDetails::IsActorPackagingModeEditable)
-				.ButtonContent()
+		// WorldSettings should not be packaged externally
+		if (SelectedActorInfo.SelectionClass != AWorldSettings::StaticClass())
+		{
+			// Actor Packaging Mode
+			const bool bIsPartitionedWorld = UWorld::HasSubsystem<UWorldPartitionSubsystem>(SelectedActorInfo.SharedWorld);
+
+			auto OnGetMenuContent = [=]() -> TSharedRef<SWidget> {
+				FMenuBuilder MenuBuilder(true, nullptr);
+				MenuBuilder.AddMenuEntry(ActorDetailsUtil::GetActorPackagingModeText(0), FText(), FSlateIcon(), FExecuteAction::CreateSP(this, &FActorDetails::OnActorPackagingModeChanged, false));
+				MenuBuilder.AddMenuEntry(ActorDetailsUtil::GetActorPackagingModeText(1), FText(), FSlateIcon(), FExecuteAction::CreateSP(this, &FActorDetails::OnActorPackagingModeChanged, true));
+				return MenuBuilder.MakeWidget();
+			};
+
+			ActorCategory.AddCustomRow( LOCTEXT("ActorPackagingModeRow", "ActorPackagingMode"), true)
+				.NameContent()
 				[
 					SNew(STextBlock)
-					.Text(this, &FActorDetails::GetCurrentActorPackagingMode)
+					.Text(LOCTEXT("ActorPackagingMode", "Packaging Mode"))
+					.ToolTipText(LOCTEXT("ActorPackagingMode_ToolTip", "Change the actor packaging mode. This will indicate if the actor is packaged alongside its level or in an external package."))
 					.Font(IDetailLayoutBuilder::GetDetailFont())
 				]
-				.IsEnabled(!bIsPartitionedWorld)
-			];
+				.ValueContent()
+				[
+					SNew(SComboButton)
+					.ContentPadding(2)
+					.OnGetMenuContent_Lambda(OnGetMenuContent)
+					.IsEnabled(this, &FActorDetails::IsActorPackagingModeEditable)
+					.ButtonContent()
+					[
+						SNew(STextBlock)
+						.Text(this, &FActorDetails::GetCurrentActorPackagingMode)
+						.Font(IDetailLayoutBuilder::GetDetailFont())
+					]
+					.IsEnabled(!bIsPartitionedWorld)
+				];
+		}
 	}
 }
 
diff --git a/Engine/Source/Editor/DetailCustomizations/Private/WorldSettingsDetails.cpp b/Engine/Source/Editor/DetailCustomizations/Private/WorldSettingsDetails.cpp
--- a/Engine/Source/Editor/DetailCustomizations/Private/WorldSettingsDetails.cpp
+++ b/Engine/Source/Editor/DetailCustomizations/Private/WorldSettingsDetails.cpp
@@ -44,7 +44,10 @@ void FWorldSettingsDetails::CustomizeDetails( IDetailLayoutBuilder& DetailBuilde
 
 	AddLightmapCustomization(DetailBuilder);
 
-	AddLevelExternalActorsCustomization(DetailBuilder);
+	if (GetDefault<UEditorExperimentalSettings>()->bEnableOneFilePerActorSupport)
+	{
+		AddLevelExternalActorsCustomization(DetailBuilder);
+	}
 
 	DetailBuilder.HideProperty(AActor::GetHiddenPropertyName(), AActor::StaticClass());
 }
diff --git a/Engine/Source/Editor/UnrealEd/Classes/Settings/EditorExperimentalSettings.h b/Engine/Source/Editor/UnrealEd/Classes/Settings/EditorExperimentalSettings.h
--- a/Engine/Source/Editor/UnrealEd/Classes/Settings/EditorExperimentalSettings.h
+++ b/Engine/Source/Editor/UnrealEd/Classes/Settings/EditorExperimentalSettings.h
@@ -63,7 +63,7 @@ public:
 	/** Break on Exceptions allows you to trap Access Nones and other exceptional events in Blueprints. */
 	UPROPERTY(EditAnywhere, config, Category=Blueprints, meta=(DisplayName="Blueprint Break on Exceptions"))
 	bool bBreakOnExceptions;
-	
+
 protected:
 	/** Any blueprint deriving from one of these base classes will be allowed to recompile during Play-in-Editor */
 	UPROPERTY(EditAnywhere, config, Category = Blueprints, meta=(AllowAbstract))
@@ -121,6 +121,10 @@ public:
 	UPROPERTY(EditAnywhere, config, Category = Tools, meta = (ConfigRestartRequired = true))
 	bool bFacialAnimationImporter;
 
+	/** Whether to enable one file per actor support in the engine. */
+	UPROPERTY(EditAnywhere, config, Category = Tools, meta=(DisplayName="Enable One File Per Actor"))
+	bool bEnableOneFilePerActorSupport;
+
 	/** Enable experimental PIE preview device launch */
 	UPROPERTY(EditAnywhere, config, Category = PIE, meta = (DisplayName = "Enable mobile PIE with preview device launch options."))
 	bool bMobilePIEPreviewDeviceLaunch;
