Mobility is now a dropdown instead of a segmented control.

diff --git a/Engine/Source/Editor/DetailCustomizations/Private/MobilityCustomization.cpp b/Engine/Source/Editor/DetailCustomizations/Private/MobilityCustomization.cpp
--- a/Engine/Source/Editor/DetailCustomizations/Private/MobilityCustomization.cpp
+++ b/Engine/Source/Editor/DetailCustomizations/Private/MobilityCustomization.cpp
@@ -1,18 +1,93 @@
 // Copyright Epic Games, Inc. All Rights Reserved.
 
 #include "Customizations/MobilityCustomization.h"
-#include "Widgets/DeclarativeSyntaxSupport.h"
+
+#include "DetailLayoutBuilder.h"
+#include "DetailWidgetRow.h"
+#include "EditorStyleSet.h"
 #include "Styling/SlateTypes.h"
+#include "Widgets/DeclarativeSyntaxSupport.h"
 #include "Widgets/Images/SImage.h"
-#include "Widgets/Text/STextBlock.h"
 #include "Widgets/Input/SCheckBox.h"
-#include "EditorStyleSet.h"
-#include "DetailLayoutBuilder.h"
-#include "DetailWidgetRow.h"
-#include "Widgets/Input/SSegmentedControl.h"
+#include "Widgets/Input/SComboBox.h"
+#include "Widgets/Text/STextBlock.h"
 
 #define LOCTEXT_NAMESPACE "MobilityCustomization" 
 
+static FName GetOptionName(EComponentMobility::Type InMobility)
+{
+	switch (InMobility)
+	{
+		case EComponentMobility::Static: return FName("Static");
+		case EComponentMobility::Stationary: return FName("Stationary");
+		case EComponentMobility::Movable: return FName("Movable");
+	}
+
+	check(false);
+	return NAME_None;
+}
+
+static EComponentMobility::Type GetOptionValue(FName InMobilityName)
+{
+	if (InMobilityName.IsNone())
+	{
+		return EComponentMobility::Static;
+	}
+
+	if (InMobilityName == GetOptionName(EComponentMobility::Static))
+	{
+		return EComponentMobility::Static;
+	}
+	else if (InMobilityName == GetOptionName(EComponentMobility::Movable))
+	{
+		return EComponentMobility::Stationary;
+	}
+	else if (InMobilityName == GetOptionName(EComponentMobility::Stationary))
+	{
+		return EComponentMobility::Movable;
+	}
+
+	check(false);
+	return EComponentMobility::Static;
+}
+
+static FText GetOptionText(EComponentMobility::Type InMobility)
+{
+	switch (InMobility)
+	{
+	case EComponentMobility::Static:
+		return LOCTEXT("Static", "Static");
+	case EComponentMobility::Movable:
+		return LOCTEXT("Movable", "Movable");
+	case EComponentMobility::Stationary:
+		return LOCTEXT("Stationary", "Stationary");
+	}
+
+	check(false);
+	return FText::GetEmpty();
+}
+
+static FText GetOptionToolTip(EComponentMobility::Type InMobility, const bool bForLight)
+{
+	switch (InMobility)
+	{
+	case EComponentMobility::Static:
+		return bForLight
+			? LOCTEXT("Mobility_Static_Light_Tooltip", "A static light can't be changed in game.\n* Fully Baked Lighting\n* Fastest Rendering")
+			: LOCTEXT("Mobility_Static_Tooltip", "A static object can't be changed in game.\n* Allows Baked Lighting\n* Fastest Rendering");;
+	case EComponentMobility::Movable:
+		return bForLight
+			? LOCTEXT("Mobility_Movable_Light_Tooltip", "Movable lights can be moved and changed in game.\n* Totally Dynamic\n* Whole Scene Dynamic Shadows\n* Slowest Rendering")
+			: LOCTEXT("Mobility_Movable_Tooltip", "Movable objects can be moved and changed in game.\n* Totally Dynamic\n* Casts a Dynamic Shadow \n* Slowest Rendering");
+	case EComponentMobility::Stationary:
+		return bForLight
+			? LOCTEXT("Mobility_Stationary_Tooltip", "A stationary light will only have its shadowing and bounced lighting from static geometry baked by Lightmass, all other lighting will be dynamic.  It can change color and intensity in game.\n* Can't Move\n* Allows Partially Baked Lighting\n* Dynamic Shadows from Movable objects")
+			: LOCTEXT("Mobility_Stationary_Object_Tooltip", "A stationary object can be changed in game but not moved, and enables cached lighting methods. \n* Cached Dynamic Shadows.");
+	}
+	check(false);
+	return FText::GetEmpty();
+}
+
 FMobilityCustomization::FMobilityCustomization(TSharedPtr<IPropertyHandle> InMobilityHandle, uint8 InRestrictedMobilityBits, bool InForLight)
 {
 	MobilityHandle = InMobilityHandle;
@@ -32,121 +107,73 @@ FName FMobilityCustomization::GetName() const
 
 void FMobilityCustomization::GenerateHeaderRowContent(FDetailWidgetRow& WidgetRow)
 {
-	TSharedRef<SSegmentedControl<EComponentMobility::Type>> ButtonOptionsPanel =
-		SNew(SSegmentedControl<EComponentMobility::Type>)
-		.Value(this, &FMobilityCustomization::GetActiveMobility)
-		.OnValueChanged(this, &FMobilityCustomization::OnMobilityChanged);
+	AllowedOptions.Reset();
+
+	const bool bShowStatic = !(RestrictedMobilityBits & StaticMobilityBitMask);
+	if (bShowStatic)
+	{
+		AllowedOptions.Add(GetOptionName(EComponentMobility::Static));
+	}
+
+	const bool bShowStationary = !(RestrictedMobilityBits & StationaryMobilityBitMask);
+	if (bShowStationary)
+	{
+		AllowedOptions.Add(GetOptionName(EComponentMobility::Stationary));
+	}
+
+	const bool bShowMovable = !(RestrictedMobilityBits & MovableMobilityBitMask);
+	if (bShowMovable)
+	{
+		AllowedOptions.Add(GetOptionName(EComponentMobility::Movable));
+	}
+
+	TSharedRef<SComboBox<FName>> ComboBox =
+		SNew(SComboBox<FName>)
+		.OptionsSource(&AllowedOptions)
+		.InitiallySelectedItem(GetOptionName(GetActiveMobility()))
+		.OnSelectionChanged(this, &FMobilityCustomization::OnMobilityChanged)
+		.OnGenerateWidget(this, &FMobilityCustomization::OnGenerateWidget)
+		.Content()
+		[
+			SNew(STextBlock)
+			.Font(FAppStyle::Get().GetFontStyle("PropertyWindow.MobilityFont"))
+			.Text(this, &FMobilityCustomization::GetActiveMobilityText)
+			.ToolTipText(this, &FMobilityCustomization::GetActiveMobilityToolTip)
+		];
 		
 	WidgetRow
 	.NameContent()
 	[
 		SNew(STextBlock)
+		.Font(IDetailLayoutBuilder::GetDetailFont())
 		.Text(LOCTEXT("Mobility", "Mobility"))
 		.ToolTipText(this, &FMobilityCustomization::GetMobilityToolTip)
-		.Font(IDetailLayoutBuilder::GetDetailFont())
 	]
 	.ValueContent()
 	.MaxDesiredWidth(0)
 	[
-		ButtonOptionsPanel
+		ComboBox
 	];
+}
 
-	bool bShowStatic = !( RestrictedMobilityBits & StaticMobilityBitMask );
-	bool bShowStationary = !( RestrictedMobilityBits & StationaryMobilityBitMask );
-
-	int32 ColumnIndex = 0;
-
-	if ( bShowStatic )
-	{
-		FText StaticTooltip = bForLight
-			? LOCTEXT("Mobility_Static_Light_Tooltip", "A static light can't be changed in game.\n* Fully Baked Lighting\n* Fastest Rendering")
-			: LOCTEXT("Mobility_Static_Tooltip", "A static object can't be changed in game.\n* Allows Baked Lighting\n* Fastest Rendering");
-
-		// Static Mobility
-		ButtonOptionsPanel->AddSlot(EComponentMobility::Static)
-		.HAlign(HAlign_Center)
-		[
-			SNew(SHorizontalBox)
-			+ SHorizontalBox::Slot()
-			.AutoWidth()
-			.VAlign(VAlign_Center)
-			.Padding(0.0f, 0.0f, 3.0f, 0.0f)
-			[
-				SNew(SImage)
-				.Image(FAppStyle::Get().GetBrush("PropertyWindow.MobilityStatic"))
-				.ColorAndOpacity(FSlateColor::UseForeground())
-			]
-			+ SHorizontalBox::Slot()
-			.VAlign(VAlign_Center)
-			[
-				SNew(STextBlock)
-				.Font(FAppStyle::Get().GetFontStyle("PropertyWindow.MobilityFont"))
-				.Text(LOCTEXT("Static", "Static"))
-			]
-		]
-		.ToolTip(StaticTooltip);
-	}
+TSharedRef<SWidget> FMobilityCustomization::OnGenerateWidget(const FName InMobilityName) const
+{
+	const EComponentMobility::Type Mobility = GetOptionValue(InMobilityName);
+	const FText Text = GetOptionText(Mobility);
+	const FText ToolTip = GetOptionToolTip(Mobility, bForLight);
+
+	return SNew(STextBlock)
+		.Font(FAppStyle::Get().GetFontStyle("PropertyWindow.MobilityFont"))
+		.Text(Text)
+		.ToolTipText(ToolTip);
+}
 
-	// Stationary Mobility
-	if ( bShowStationary )
+void FMobilityCustomization::OnMobilityChanged(FName InMobilityName, ESelectInfo::Type)
+{
+	if (MobilityHandle.IsValid() && !InMobilityName.IsNone())
 	{
-		FText StationaryTooltip = bForLight
-			? LOCTEXT("Mobility_Stationary_Tooltip", "A stationary light will only have its shadowing and bounced lighting from static geometry baked by Lightmass, all other lighting will be dynamic.  It can change color and intensity in game.\n* Can't Move\n* Allows Partially Baked Lighting\n* Dynamic Shadows from Movable objects")
-			: LOCTEXT("Mobility_Stationary_Object_Tooltip", "A stationary object can be changed in game but not moved, and enables cached lighting methods. \n* Cached Dynamic Shadows.");
-
-		ButtonOptionsPanel->AddSlot(EComponentMobility::Stationary)
-		.HAlign(HAlign_Center)
-		[
-			SNew(SHorizontalBox)
-			+ SHorizontalBox::Slot()
-			.AutoWidth()
-			.VAlign(VAlign_Center)
-			.Padding(0.0f, 0.0f, 3.0f, 0.0f)
-			[
-				SNew(SImage)
-				.Image(FAppStyle::Get().GetBrush("PropertyWindow.MobilityStationary"))
-				.ColorAndOpacity(FSlateColor::UseForeground())
-			]
-			+ SHorizontalBox::Slot()
-			.VAlign(VAlign_Center)
-			[
-				SNew(STextBlock)
-				.Font(FAppStyle::Get().GetFontStyle("PropertyWindow.MobilityFont"))
-				.Text(LOCTEXT("Stationary", "Stationary"))
-			]
-		]
-		.ToolTip(StationaryTooltip);
+		MobilityHandle->SetValue((uint8) GetOptionValue(InMobilityName));
 	}
-
-	FText MovableTooltip = bForLight
-			? LOCTEXT("Mobility_Movable_Light_Tooltip", "Movable lights can be moved and changed in game.\n* Totally Dynamic\n* Whole Scene Dynamic Shadows\n* Slowest Rendering")
-			: LOCTEXT("Mobility_Movable_Tooltip", "Movable objects can be moved and changed in game.\n* Totally Dynamic\n* Casts a Dynamic Shadow \n* Slowest Rendering");
-
-	// Movable Mobility
-	ButtonOptionsPanel->AddSlot(EComponentMobility::Movable)
-	.HAlign(HAlign_Center)
-	[
-		SNew(SHorizontalBox)
-		+ SHorizontalBox::Slot()
-		.AutoWidth()
-		.VAlign(VAlign_Center)
-		.Padding(0.0f, 0.0f, 3.0f, 0.0f)
-		[
-			SNew(SImage)
-			.Image(FAppStyle::Get().GetBrush("PropertyWindow.MobilityMoveable"))
-			.ColorAndOpacity(FSlateColor::UseForeground())
-		]
-		+ SHorizontalBox::Slot()
-		.VAlign(VAlign_Center)
-		[
-			SNew(STextBlock)
-			.Font(FAppStyle::Get().GetFontStyle("PropertyWindow.MobilityFont"))
-			.Text(LOCTEXT("Movable", "Movable"))
-		]
-	]
-	.ToolTip(MovableTooltip);
-
-	ButtonOptionsPanel->RebuildChildren();
 }
 
 EComponentMobility::Type FMobilityCustomization::GetActiveMobility() const
@@ -156,31 +183,20 @@ EComponentMobility::Type FMobilityCustomization::GetActiveMobility() const
 		uint8 MobilityByte;
 		MobilityHandle->GetValue(MobilityByte);
 
-		return (EComponentMobility::Type)MobilityByte;
+		return (EComponentMobility::Type) MobilityByte;
 	}
 
 	return EComponentMobility::Static;
 }
 
-FSlateColor FMobilityCustomization::GetMobilityTextColor(EComponentMobility::Type InMobility) const
+FText FMobilityCustomization::GetActiveMobilityText() const
 {
-	if (MobilityHandle.IsValid())
-	{
-		uint8 MobilityByte;
-		MobilityHandle->GetValue(MobilityByte);
-
-		return MobilityByte == InMobility ? FSlateColor(FLinearColor(0, 0, 0)) : FSlateColor(FLinearColor(0.72f, 0.72f, 0.72f, 1.f));
-	}
-
-	return FSlateColor(FLinearColor(0.72f, 0.72f, 0.72f, 1.f));
+	return GetOptionText(GetActiveMobility());
 }
 
-void FMobilityCustomization::OnMobilityChanged(EComponentMobility::Type InMobility)
+FText FMobilityCustomization::GetActiveMobilityToolTip() const
 {
-	if (MobilityHandle.IsValid())
-	{
-		MobilityHandle->SetValue((uint8)InMobility);
-	}
+	return GetOptionToolTip(GetActiveMobility(), bForLight);
 }
 
 FText FMobilityCustomization::GetMobilityToolTip() const
diff --git a/Engine/Source/Editor/DetailCustomizations/Public/Customizations/MobilityCustomization.h b/Engine/Source/Editor/DetailCustomizations/Public/Customizations/MobilityCustomization.h
--- a/Engine/Source/Editor/DetailCustomizations/Public/Customizations/MobilityCustomization.h
+++ b/Engine/Source/Editor/DetailCustomizations/Public/Customizations/MobilityCustomization.h
@@ -2,9 +2,10 @@
 #pragma once
 
 #include "CoreMinimal.h"
+#include "IDetailCustomNodeBuilder.h"
 #include "Engine/EngineTypes.h"
+#include "Framework/SlateDelegates.h"
 #include "Styling/SlateColor.h"
-#include "IDetailCustomNodeBuilder.h"
 
 class IDetailCategoryBuilder;
 class IPropertyHandle;
@@ -29,11 +30,14 @@ public:
 
 private:
 
-	EComponentMobility::Type GetActiveMobility() const;
-	FSlateColor GetMobilityTextColor(EComponentMobility::Type InMobility) const;
-	void OnMobilityChanged(EComponentMobility::Type InMobility);
+	TSharedRef<class SWidget> OnGenerateWidget(FName InMobilityName) const;
+	void OnMobilityChanged(FName InMobilityName, ESelectInfo::Type);
 	FText GetMobilityToolTip() const;
-	
+	EComponentMobility::Type GetActiveMobility() const;
+	FText GetActiveMobilityText() const;
+	FText GetActiveMobilityToolTip() const;
+
+	TArray<FName> AllowedOptions;
 	TSharedPtr<IPropertyHandle> MobilityHandle;
 	bool bForLight;
 	uint8 RestrictedMobilityBits;
