Fix mobility dropdown not accounting for multiple values and crashing instead

diff --git a/Engine/Source/Editor/DetailCustomizations/Private/MobilityCustomization.cpp b/Engine/Source/Editor/DetailCustomizations/Private/MobilityCustomization.cpp
--- a/Engine/Source/Editor/DetailCustomizations/Private/MobilityCustomization.cpp
+++ b/Engine/Source/Editor/DetailCustomizations/Private/MobilityCustomization.cpp
@@ -14,78 +14,85 @@
 
 #define LOCTEXT_NAMESPACE "MobilityCustomization" 
 
-static FName GetOptionName(EComponentMobility::Type InMobility)
+static FName GetOptionName(TOptional<EComponentMobility::Type> InMobility)
 {
-	switch (InMobility)
+	if (InMobility.IsSet())
 	{
+		EComponentMobility::Type Mobility = InMobility.GetValue();
+		switch (Mobility)
+		{
 		case EComponentMobility::Static: return FName("Static");
 		case EComponentMobility::Stationary: return FName("Stationary");
 		case EComponentMobility::Movable: return FName("Movable");
+		}
 	}
-
-	check(false);
 	return NAME_None;
 }
 
-static EComponentMobility::Type GetOptionValue(FName InMobilityName)
-{
-	if (InMobilityName.IsNone())
-	{
-		return EComponentMobility::Static;
-	}
 
-	if (InMobilityName == GetOptionName(EComponentMobility::Static))
-	{
-		return EComponentMobility::Static;
-	}
-	else if (InMobilityName == GetOptionName(EComponentMobility::Movable))
-	{
-		return EComponentMobility::Stationary;
-	}
-	else if (InMobilityName == GetOptionName(EComponentMobility::Stationary))
+static TOptional<EComponentMobility::Type> GetOptionValue(FName InMobilityName)
+{
+	if (!InMobilityName.IsNone())
 	{
-		return EComponentMobility::Movable;
+		if (InMobilityName == GetOptionName(EComponentMobility::Static))
+		{
+			return EComponentMobility::Static;
+		}
+		else if (InMobilityName == GetOptionName(EComponentMobility::Stationary))
+		{
+			return EComponentMobility::Stationary;
+		}
+		else if (InMobilityName == GetOptionName(EComponentMobility::Movable))
+		{
+			return EComponentMobility::Movable;
+		}
 	}
 
-	check(false);
-	return EComponentMobility::Static;
+	return TOptional<EComponentMobility::Type>();
 }
 
-static FText GetOptionText(EComponentMobility::Type InMobility)
+static FText GetOptionText(TOptional<EComponentMobility::Type> InMobility)
 {
-	switch (InMobility)
+	if (InMobility.IsSet())
 	{
-	case EComponentMobility::Static:
-		return LOCTEXT("Static", "Static");
-	case EComponentMobility::Movable:
-		return LOCTEXT("Movable", "Movable");
-	case EComponentMobility::Stationary:
-		return LOCTEXT("Stationary", "Stationary");
+		EComponentMobility::Type Mobility = InMobility.GetValue();
+		switch (Mobility)
+		{
+		case EComponentMobility::Static:
+			return LOCTEXT("Static", "Static");
+		case EComponentMobility::Movable:
+			return LOCTEXT("Movable", "Movable");
+		case EComponentMobility::Stationary:
+			return LOCTEXT("Stationary", "Stationary");
+		}
 	}
 
-	check(false);
-	return FText::GetEmpty();
+	return LOCTEXT("MultipleValues", "Multiple Values");
 }
 
-static FText GetOptionToolTip(EComponentMobility::Type InMobility, const bool bForLight)
+static FText GetOptionToolTip(TOptional<EComponentMobility::Type> InMobility, const bool bForLight)
 {
-	switch (InMobility)
-	{
-	case EComponentMobility::Static:
-		return bForLight
-			? LOCTEXT("Mobility_Static_Light_Tooltip", "A static light can't be changed in game.\n* Fully Baked Lighting\n* Fastest Rendering")
-			: LOCTEXT("Mobility_Static_Tooltip", "A static object can't be changed in game.\n* Allows Baked Lighting\n* Fastest Rendering");;
-	case EComponentMobility::Movable:
-		return bForLight
-			? LOCTEXT("Mobility_Movable_Light_Tooltip", "Movable lights can be moved and changed in game.\n* Totally Dynamic\n* Whole Scene Dynamic Shadows\n* Slowest Rendering")
-			: LOCTEXT("Mobility_Movable_Tooltip", "Movable objects can be moved and changed in game.\n* Totally Dynamic\n* Casts a Dynamic Shadow \n* Slowest Rendering");
-	case EComponentMobility::Stationary:
-		return bForLight
-			? LOCTEXT("Mobility_Stationary_Tooltip", "A stationary light will only have its shadowing and bounced lighting from static geometry baked by Lightmass, all other lighting will be dynamic.  It can change color and intensity in game.\n* Can't Move\n* Allows Partially Baked Lighting\n* Dynamic Shadows from Movable objects")
-			: LOCTEXT("Mobility_Stationary_Object_Tooltip", "A stationary object can be changed in game but not moved, and enables cached lighting methods. \n* Cached Dynamic Shadows.");
+	if(InMobility.IsSet())
+	{ 
+		EComponentMobility::Type Mobility = InMobility.GetValue();
+		switch (Mobility)
+		{
+		case EComponentMobility::Static:
+			return bForLight
+				? LOCTEXT("Mobility_Static_Light_Tooltip", "A static light can't be changed in game.\n* Fully Baked Lighting\n* Fastest Rendering")
+				: LOCTEXT("Mobility_Static_Tooltip", "A static object can't be changed in game.\n* Allows Baked Lighting\n* Fastest Rendering");;
+		case EComponentMobility::Movable:
+			return bForLight
+				? LOCTEXT("Mobility_Movable_Light_Tooltip", "Movable lights can be moved and changed in game.\n* Totally Dynamic\n* Whole Scene Dynamic Shadows\n* Slowest Rendering")
+				: LOCTEXT("Mobility_Movable_Tooltip", "Movable objects can be moved and changed in game.\n* Totally Dynamic\n* Casts a Dynamic Shadow \n* Slowest Rendering");
+		case EComponentMobility::Stationary:
+			return bForLight
+				? LOCTEXT("Mobility_Stationary_Tooltip", "A stationary light will only have its shadowing and bounced lighting from static geometry baked by Lightmass, all other lighting will be dynamic.  It can change color and intensity in game.\n* Can't Move\n* Allows Partially Baked Lighting\n* Dynamic Shadows from Movable objects")
+				: LOCTEXT("Mobility_Stationary_Object_Tooltip", "A stationary object can be changed in game but not moved, and enables cached lighting methods. \n* Cached Dynamic Shadows.");
+		}
 	}
-	check(false);
-	return FText::GetEmpty();
+	
+	return LOCTEXT("MobilityMultipleValues_Tooltip", "Multiple components are selected with different mobilities");
 }
 
 FMobilityCustomization::FMobilityCustomization(TSharedPtr<IPropertyHandle> InMobilityHandle, uint8 InRestrictedMobilityBits, bool InForLight)
@@ -159,7 +166,7 @@ void FMobilityCustomization::GenerateHeaderRowContent(FDetailWidgetRow& WidgetRo
 
 TSharedRef<SWidget> FMobilityCustomization::OnGenerateWidget(const FName InMobilityName) const
 {
-	const EComponentMobility::Type Mobility = GetOptionValue(InMobilityName);
+	const TOptional<EComponentMobility::Type> Mobility = GetOptionValue(InMobilityName);
 	const FText Text = GetOptionText(Mobility);
 	const FText ToolTip = GetOptionToolTip(Mobility, bForLight);
 
@@ -173,21 +180,22 @@ void FMobilityCustomization::OnMobilityChanged(FName InMobilityName, ESelectInfo
 {
 	if (MobilityHandle.IsValid() && !InMobilityName.IsNone())
 	{
-		MobilityHandle->SetValue((uint8) GetOptionValue(InMobilityName));
+		MobilityHandle->SetValue((uint8) GetOptionValue(InMobilityName).GetValue());
 	}
 }
 
-EComponentMobility::Type FMobilityCustomization::GetActiveMobility() const
+TOptional<EComponentMobility::Type> FMobilityCustomization::GetActiveMobility() const
 {
 	if (MobilityHandle.IsValid())
 	{
 		uint8 MobilityByte;
-		MobilityHandle->GetValue(MobilityByte);
-
-		return (EComponentMobility::Type) MobilityByte;
+		if(MobilityHandle->GetValue(MobilityByte) == FPropertyAccess::Success)
+		{
+			return (EComponentMobility::Type)MobilityByte;
+		}
 	}
 
-	return EComponentMobility::Static;
+	return TOptional<EComponentMobility::Type>();
 }
 
 FText FMobilityCustomization::GetActiveMobilityText() const
diff --git a/Engine/Source/Editor/DetailCustomizations/Public/Customizations/MobilityCustomization.h b/Engine/Source/Editor/DetailCustomizations/Public/Customizations/MobilityCustomization.h
--- a/Engine/Source/Editor/DetailCustomizations/Public/Customizations/MobilityCustomization.h
+++ b/Engine/Source/Editor/DetailCustomizations/Public/Customizations/MobilityCustomization.h
@@ -33,7 +33,7 @@ private:
 	TSharedRef<class SWidget> OnGenerateWidget(FName InMobilityName) const;
 	void OnMobilityChanged(FName InMobilityName, ESelectInfo::Type);
 	FText GetMobilityToolTip() const;
-	EComponentMobility::Type GetActiveMobility() const;
+	TOptional<EComponentMobility::Type> GetActiveMobility() const;
 	FText GetActiveMobilityText() const;
 	FText GetActiveMobilityToolTip() const;
 
