The Rendering > Advanced > TextureStreaming/LOD/etc. categories now use the correct display name.

diff --git a/Engine/Source/Editor/DetailCustomizations/Private/PrimitiveComponentDetails.cpp b/Engine/Source/Editor/DetailCustomizations/Private/PrimitiveComponentDetails.cpp
--- a/Engine/Source/Editor/DetailCustomizations/Private/PrimitiveComponentDetails.cpp
+++ b/Engine/Source/Editor/DetailCustomizations/Private/PrimitiveComponentDetails.cpp
@@ -127,7 +127,7 @@ void FPrimitiveComponentDetails::AddAdvancedSubCategory( IDetailLayoutBuilder& D
 		IDetailCategoryBuilder& MainCategory = DetailBuilder.EditCategory(MainCategoryName);
 
 		const bool bForAdvanced = true;
-		IDetailGroup& Group = MainCategory.AddGroup( SubCategoryName, FText::FromName(SubCategoryName), bForAdvanced );
+		IDetailGroup& Group = MainCategory.AddGroup( SubCategoryName, SubCategory.GetDisplayName(), bForAdvanced );
 
 		for( int32 PropertyIndex = 0; PropertyIndex < SubCategoryProperties.Num(); ++PropertyIndex )
 		{
