Math Struct Customization: Improve workflow for multi-edit

diff --git a/Engine/Source/Editor/DetailCustomizations/Private/MathStructProxyCustomizations.cpp b/Engine/Source/Editor/DetailCustomizations/Private/MathStructProxyCustomizations.cpp
--- a/Engine/Source/Editor/DetailCustomizations/Private/MathStructProxyCustomizations.cpp
+++ b/Engine/Source/Editor/DetailCustomizations/Private/MathStructProxyCustomizations.cpp
@@ -363,18 +363,36 @@ bool FMatrixStructCustomization::CacheValues( TWeakPtr<IPropertyHandle> Property
 	TArray<void*> RawData;
 	PropertyHandle->AccessRawData(RawData);
 
-	if (RawData.Num() == 1)
+	const FMatrix* FirstMatrixValue = nullptr;
+	for(void* RawDataPtr : RawData)
 	{
-		FMatrix* MatrixValue = reinterpret_cast<FMatrix*>(RawData[0]);
-		if (MatrixValue != NULL)
+		FMatrix* MatrixValue = reinterpret_cast<FMatrix*>(RawDataPtr);
+		if (MatrixValue == nullptr)
 		{
-			CachedTranslation->Set(MatrixValue->GetOrigin());
-			CachedRotation->Set(MatrixValue->Rotator());
-			CachedScale->Set(MatrixValue->GetScaleVector());
-			return true;
+			return false;
+		}
+
+		if(FirstMatrixValue)
+		{
+			if(!FirstMatrixValue->Equals(*MatrixValue, 0.0001f))
+			{
+				return false;
+			}
+		}
+		else
+		{
+			FirstMatrixValue = MatrixValue;
 		}
 	}
 
+	if(FirstMatrixValue)
+	{
+		CachedTranslation->Set(FirstMatrixValue->GetOrigin());
+		CachedRotation->Set(FirstMatrixValue->Rotator());
+		CachedScale->Set(FirstMatrixValue->GetScaleVector());
+		return true;
+	}
+
 	return false;
 }
 
@@ -501,16 +519,34 @@ bool FTransformStructCustomization::CacheValues( TWeakPtr<IPropertyHandle> Prope
 	TArray<void*> RawData;
 	PropertyHandle->AccessRawData(RawData);
 
-	if (RawData.Num() == 1)
+	const FTransform* FirstTransformValue = nullptr;
+	for(void* RawDataPtr : RawData)
 	{
-		FTransform* TransformValue = reinterpret_cast<FTransform*>(RawData[0]);
-		if (TransformValue != NULL)
+		FTransform* TransformValue = reinterpret_cast<FTransform*>(RawDataPtr);
+		if (TransformValue == nullptr)
 		{
-			CachedTranslation->Set(TransformValue->GetTranslation());
-			CachedRotation->Set(TransformValue->GetRotation().Rotator());
-			CachedScale->Set(TransformValue->GetScale3D());
-			return true;
+			return false;
 		}
+
+		if(FirstTransformValue)
+		{
+			if(!FirstTransformValue->Equals(*TransformValue, 0.0001f))
+			{
+				return false;
+			}
+		}
+		else
+		{
+			FirstTransformValue = TransformValue;
+		}
+	}
+
+	if(FirstTransformValue)
+	{
+		CachedTranslation->Set(FirstTransformValue->GetTranslation());
+		CachedRotation->Set(FirstTransformValue->GetRotation().Rotator());
+		CachedScale->Set(FirstTransformValue->GetScale3D());
+		return true;
 	}
 
 	return false;
