ADDED: Custom Primitive Data Customization
ADDED: Set/Get Custom Primitive Data by [REDACTED_PERSON]
ADDED: Primitive Index to Material Parameter Metadata for runtime use

diff --git a/Engine/Source/Editor/DetailCustomizations/DetailCustomizations.Build.cs b/Engine/Source/Editor/DetailCustomizations/DetailCustomizations.Build.cs
--- a/Engine/Source/Editor/DetailCustomizations/DetailCustomizations.Build.cs
+++ b/Engine/Source/Editor/DetailCustomizations/DetailCustomizations.Build.cs
@@ -63,7 +63,8 @@ public class DetailCustomizations : ModuleRules
                 "DataTableEditor",
 				"ToolMenus",
 				"PhysicsCore",
-				"RenderCore"
+				"RenderCore",
+				"MaterialEditor",
 			}
 		);
 
diff --git a/Engine/Source/Editor/DetailCustomizations/Private/CustomPrimitiveDataCustomization.cpp b/Engine/Source/Editor/DetailCustomizations/Private/CustomPrimitiveDataCustomization.cpp
new file mode 100644
--- /dev/null
+++ b/Engine/Source/Editor/DetailCustomizations/Private/CustomPrimitiveDataCustomization.cpp
@@ -0,0 +1,767 @@
+// Copyright Epic Games, Inc. All Rights Reserved.
+
+#include "CustomPrimitiveDataCustomization.h"
+
+#include "Widgets/Layout/SBox.h"
+#include "Widgets/Text/STextBlock.h"
+#include "Widgets/Input/SHyperlink.h"
+#include "Widgets/Colors/SColorPicker.h"
+#include "Widgets/Colors/SColorBlock.h"
+#include "SlateOptMacros.h"
+
+#include "EditorStyleSet.h"
+
+#include "IDetailGroup.h"
+#include "IDetailPropertyRow.h"
+#include "IPropertyTypeCustomization.h"
+#include "IPropertyUtilities.h"
+#include "PropertyCustomizationHelpers.h"
+
+#include "Materials/MaterialInstance.h"
+#include "Materials/MaterialExpressionScalarParameter.h"
+#include "Materials/MaterialExpressionVectorParameter.h"
+#include "IMaterialEditor.h"
+
+#define LOCTEXT_NAMESPACE "CustomPrimitiveDataCustomization"
+
+TSharedRef<IPropertyTypeCustomization> FCustomPrimitiveDataCustomization::MakeInstance()
+{
+	return MakeShareable(new FCustomPrimitiveDataCustomization);
+}
+
+FCustomPrimitiveDataCustomization::~FCustomPrimitiveDataCustomization()
+{
+	Cleanup();
+}
+
+void FCustomPrimitiveDataCustomization::CustomizeHeader(TSharedRef<IPropertyHandle> PropertyHandle, FDetailWidgetRow& HeaderRow, IPropertyTypeCustomizationUtils& CustomizationUtils)
+{
+	TSharedPtr<IPropertyHandle> DataProperty = PropertyHandle->GetChildHandle("Data");
+
+	// Move the data array to be the outer, so we don't have to expand the struct
+	HeaderRow.NameContent()
+	[
+		PropertyHandle->CreatePropertyNameWidget()
+	]
+	.ValueContent()
+	[
+		DataProperty->CreatePropertyValueWidget()
+	];
+}
+
+void FCustomPrimitiveDataCustomization::CustomizeChildren(TSharedRef<IPropertyHandle> PropertyHandle, IDetailChildrenBuilder& ChildBuilder, IPropertyTypeCustomizationUtils& CustomizationUtils)
+{
+	Cleanup();
+
+	PropertyUtils = CustomizationUtils.GetPropertyUtilities();
+
+	TSharedPtr<IPropertyHandle> DataProperty = PropertyHandle->GetChildHandle("Data");
+	DataArrayHandle = DataProperty->AsArray();
+
+	int32 NumSelectedComponents = 0;
+	int32 MaxPrimitiveDataIndex = INDEX_NONE;
+
+	ForEachSelectedComponent([&](UPrimitiveComponent* Component)
+	{
+		PopulateParameterData(Component, MaxPrimitiveDataIndex);
+		NumSelectedComponents++;
+		ComponentsToWatch.Add(Component);
+	});
+
+	uint32 NumElements;
+	DataArrayHandle->GetNumElements(NumElements);
+
+	FSimpleDelegate OnElemsChanged = FSimpleDelegate::CreateRaw(this, &FCustomPrimitiveDataCustomization::OnUpdated);
+	PropertyHandle->SetOnPropertyValueChanged(OnElemsChanged);
+	DataProperty->SetOnPropertyValueChanged(OnElemsChanged);
+	DataArrayHandle->SetOnNumElementsChanged(OnElemsChanged);
+
+	FCoreUObjectDelegates::OnObjectPropertyChanged.AddRaw(this, &FCustomPrimitiveDataCustomization::OnObjectPropertyChanged);
+	UMaterial::OnMaterialCompilationFinished().AddRaw(this, &FCustomPrimitiveDataCustomization::OnMaterialCompiled);
+
+	const int32 NumPrimitiveIndices = FMath::Max(MaxPrimitiveDataIndex + 1, (int32)NumElements);
+
+	if (NumPrimitiveIndices == 0)
+	{
+		return;
+	}
+
+	FString ArrayString;
+	const bool bDataEditable = DataProperty.IsValid() && DataProperty->IsEditable() && DataProperty->GetValueAsDisplayString(ArrayString) == FPropertyAccess::Success;
+
+	uint8 VectorGroupPrimIdx = 0;
+	IDetailGroup* VectorGroup = NULL;
+
+	for (uint8 PrimIdx = 0; PrimIdx < NumPrimitiveIndices; ++PrimIdx)
+	{
+		TSharedPtr<IPropertyHandle> ElementHandle = PrimIdx < NumElements ? TSharedPtr<IPropertyHandle>(DataArrayHandle->GetElement(PrimIdx)) : NULL;
+
+		if (VectorGroup && (PrimIdx - VectorGroupPrimIdx) > 3)
+		{
+			// We're no longer in a vector group
+			VectorGroup = NULL;
+		}
+
+		// Always prioritize the first vector found
+		if (VectorGroup == NULL && VectorParameterData.Contains(PrimIdx))
+		{
+			// Create a collapsing group that contains our color picker, so we can quickly assign colors to our vector
+			VectorGroupPrimIdx = PrimIdx;
+			VectorGroup = CreateVectorGroup(ChildBuilder, PrimIdx, bDataEditable, NumElements);			
+		}
+
+		if (ScalarParameterData.Contains(PrimIdx) || VectorParameterData.Contains(PrimIdx))
+		{
+			CreateParameterRow(ChildBuilder, PrimIdx, ElementHandle, NumSelectedComponents, bDataEditable, VectorGroup, CustomizationUtils);
+		}
+		else
+		{
+			// We've encountered a gap in defined custom primitive data, mark it undefined
+			TSharedRef<SWidget> NameWidget = GetUndefinedParameterWidget(PrimIdx, CustomizationUtils);
+
+			if (ElementHandle.IsValid())
+			{
+				ChildBuilder.AddProperty(ElementHandle.ToSharedRef())
+				.CustomWidget()
+				.NameContent()
+				[
+					NameWidget
+				]
+				.ValueContent()
+				[
+					ElementHandle->CreatePropertyValueWidget(false)
+				]
+				.IsEnabled(bDataEditable);;
+			}
+			else
+			{
+				ChildBuilder.AddCustomRow(FText::AsNumber(PrimIdx))
+				.NameContent()
+				[
+					NameWidget
+				]
+				.IsEnabled(bDataEditable);
+			}
+		}
+	}
+}
+
+BEGIN_SLATE_FUNCTION_BUILD_OPTIMIZATION
+IDetailGroup* FCustomPrimitiveDataCustomization::CreateVectorGroup(IDetailChildrenBuilder& ChildBuilder, uint8 PrimIdx, bool bDataEditable, int32 NumElements)
+{
+	IDetailGroup* VectorGroup = &ChildBuilder.AddGroup(VectorParameterData[PrimIdx][0].Info.Name, FText::FromName(VectorParameterData[PrimIdx][0].Info.Name));
+
+	TSharedPtr<SColorBlock> ColorBlock;
+	TSharedPtr<SVerticalBox> VectorGroupNameBox = SNew(SVerticalBox);
+
+	// Use this to make sure we don't make duplicate parameters for the group header
+	TSet<FGuid> AddedParametersForThisGroup;
+
+	for (const FParameterData& ParameterData : VectorParameterData[PrimIdx])
+	{
+		if (AddedParametersForThisGroup.Contains(ParameterData.ExpressionID))
+		{
+			continue;
+		}
+
+		AddedParametersForThisGroup.Add(ParameterData.ExpressionID);
+
+		VectorGroupNameBox->AddSlot()
+		.Padding(2.f)
+		[
+			CreateHyperlink(FText::FromName(ParameterData.Info.Name), ParameterData.Material, ParameterData.ExpressionID)
+		];
+	}
+
+	TSharedPtr<SWidget> VectorGroupNameWidget;
+	if (VectorGroupNameBox->NumSlots() > 1)
+	{
+		// We have multiple overlapping parameter names, make sure to put a border around it to contain it
+		// TODO: Maybe a UI designer needs to take a look at this?
+		VectorGroupNameWidget = SNew(SBorder)
+		.BorderImage(FEditorStyle::Get().GetBrush("FilledBorder"))
+		.Padding(2.f)
+		[
+			VectorGroupNameBox.ToSharedRef()
+		];
+	}
+	else
+	{
+		VectorGroupNameWidget = VectorGroupNameBox;
+	}
+
+	VectorGroup->HeaderRow()
+	.NameContent()
+	[
+		VectorGroupNameWidget.ToSharedRef()
+	]
+	.ValueContent()
+	[
+		SNew(SHorizontalBox)
+		.IsEnabled(bDataEditable)
+		+ SHorizontalBox::Slot()
+		.VAlign(VAlign_Center)
+		.Padding(0.f, 2.f)
+		[
+			SAssignNew(ColorBlock, SColorBlock)
+			.Color(this, &FCustomPrimitiveDataCustomization::GetVectorColor, PrimIdx)
+			.ShowBackgroundForAlpha(true)
+			.AlphaDisplayMode(EColorBlockAlphaDisplayMode::Separate)
+			.OnMouseButtonDown(this, &FCustomPrimitiveDataCustomization::OnMouseButtonDownColorBlock, PrimIdx)
+			.Size(FVector2D(35.0f, 12.0f))
+			.Visibility(PrimIdx < NumElements ? EVisibility::Visible : EVisibility::Collapsed)
+		]
+		+ SHorizontalBox::Slot()
+		.Padding(2.0f)
+		.HAlign(HAlign_Center)
+		.VAlign(VAlign_Center)
+		.AutoWidth()
+		[
+			PropertyCustomizationHelpers::MakeAddButton(FSimpleDelegate::CreateSP(this, &FCustomPrimitiveDataCustomization::OnAddedDesiredPrimitiveData, (uint8)(PrimIdx + 3)), FText(), (int32)NumElements < PrimIdx + 4)
+		]
+		+ SHorizontalBox::Slot()
+		.Padding(2.0f)
+		.HAlign(HAlign_Center)
+		.VAlign(VAlign_Center)
+		.AutoWidth()
+		[
+			PropertyCustomizationHelpers::MakeEmptyButton(FSimpleDelegate::CreateSP(this, &FCustomPrimitiveDataCustomization::OnRemovedPrimitiveData, PrimIdx), LOCTEXT("RemoveVector", "Removes this vector (and anything after)"), (int32)PrimIdx < NumElements)
+		]
+		+ SHorizontalBox::Slot()
+		.Padding(2.0f)
+		.HAlign(HAlign_Center)
+		.VAlign(VAlign_Center)
+		.AutoWidth()
+		[
+			PropertyCustomizationHelpers::MakeResetButton(FSimpleDelegate::CreateSP(this, &FCustomPrimitiveDataCustomization::SetDefaultVectorValue, PrimIdx), FText(), (int32)PrimIdx < NumElements)
+		]
+	];
+
+	ColorBlocks.Add(PrimIdx, ColorBlock);
+	
+	return VectorGroup;
+}
+END_SLATE_FUNCTION_BUILD_OPTIMIZATION
+
+BEGIN_SLATE_FUNCTION_BUILD_OPTIMIZATION
+void FCustomPrimitiveDataCustomization::CreateParameterRow(IDetailChildrenBuilder& ChildBuilder, uint8 PrimIdx, TSharedPtr<IPropertyHandle> ElementHandle, int32 NumSelectedComponents, bool bDataEditable, IDetailGroup* VectorGroup, IPropertyTypeCustomizationUtils& CustomizationUtils)
+{
+	// Use this to make sure we don't make duplicate parameters in each row
+	TSet<FGuid> AddedParametersForThisRow;
+
+	TArray<FText> SearchText;
+	TSharedPtr<SVerticalBox> VerticalBox = SNew(SVerticalBox);
+
+	TSet<TWeakObjectPtr<UPrimitiveComponent>> Components;
+
+	if (VectorParameterData.Contains(PrimIdx))
+	{
+		for (const FParameterData& ParameterData : VectorParameterData[PrimIdx])
+		{
+			Components.Add(ParameterData.Component);
+
+			if (AddedParametersForThisRow.Contains(ParameterData.ExpressionID))
+			{
+				continue;
+			}
+
+			AddedParametersForThisRow.Add(ParameterData.ExpressionID);
+
+			FMaterialParameterMetadata ParameterMetadata;
+			ParameterData.Material->GetParameterDefaultValue(EMaterialParameterType::Vector, ParameterData.Info, ParameterMetadata);
+
+			FText ChannelName;
+			switch (ParameterData.IndexOffset)
+			{
+			case 0:
+				ChannelName = ParameterMetadata.ChannelNames.R.IsEmpty() ? LOCTEXT("DefaultVectorChannelRed", "R") : ParameterMetadata.ChannelNames.R;
+			case 1:
+				ChannelName = ParameterMetadata.ChannelNames.G.IsEmpty() ? LOCTEXT("DefaultVectorChannelGreen", "G") : ParameterMetadata.ChannelNames.G;
+			case 2:
+				ChannelName = ParameterMetadata.ChannelNames.B.IsEmpty() ? LOCTEXT("DefaultVectorChannelBlue", "B") : ParameterMetadata.ChannelNames.B;
+			case 3:
+				ChannelName = ParameterMetadata.ChannelNames.A.IsEmpty() ? LOCTEXT("DefaultVectorChannelAlpha", "A") : ParameterMetadata.ChannelNames.A;
+			default:
+				checkNoEntry();
+				break;
+			}
+
+			const FText ParameterName = FText::Format(
+				LOCTEXT("VectorParameterName", "{0}.{1}"),
+				FText::FromName(ParameterData.Info.Name),
+				ChannelName
+			);
+
+			VerticalBox->AddSlot()
+			.Padding(2.f)
+			[
+				CreateHyperlink(ParameterName, ParameterData.Material, ParameterData.ExpressionID)
+			];
+
+			SearchText.Add(ParameterName);
+		}
+	}
+
+	if (ScalarParameterData.Contains(PrimIdx))
+	{
+		for (const FParameterData& ParameterData : ScalarParameterData[PrimIdx])
+		{
+			Components.Add(ParameterData.Component);
+
+			if (AddedParametersForThisRow.Contains(ParameterData.ExpressionID))
+			{
+				continue;
+			}
+
+			AddedParametersForThisRow.Add(ParameterData.ExpressionID);
+
+			const FText ParameterName = FText::FromName(ParameterData.Info.Name);
+
+			VerticalBox->AddSlot()
+			.Padding(2.f)
+			[
+				CreateHyperlink(ParameterName, ParameterData.Material, ParameterData.ExpressionID)
+			];
+
+			SearchText.Add(ParameterName);
+		}
+	}
+
+	if (Components.Num() != NumSelectedComponents)
+	{
+		// Some components aren't defining parameters at this index, add the undefined parameter widget in case this was user error
+		VerticalBox->AddSlot()
+		.Padding(2.f)
+		[
+			GetUndefinedParameterWidget(PrimIdx, CustomizationUtils)
+		];
+	}
+
+	TSharedPtr<SWidget> NameWidget;
+	if (VerticalBox->NumSlots() > 1)
+	{
+		// We have multiple overlapping parameter names, make sure to put a border around it to contain it
+		// TODO: Maybe a UI designer needs to take a look at this?
+		NameWidget = SNew(SBorder)
+		.BorderImage(FEditorStyle::Get().GetBrush("FilledBorder"))
+		.Padding(2.f)
+		[
+			VerticalBox.ToSharedRef()
+		];
+	}
+	else
+	{
+		NameWidget = VerticalBox;
+	}
+
+	if (ElementHandle.IsValid())
+	{
+		// We already have data for this row, be sure to use it
+		TSharedRef<IPropertyHandle> ElementHandleRef = ElementHandle.ToSharedRef();
+		IDetailPropertyRow& Row = VectorGroup ? VectorGroup->AddPropertyRow(ElementHandleRef) : ChildBuilder.AddProperty(ElementHandleRef);
+
+		TSharedRef<SWidget> ValueWidget = ElementHandle->CreatePropertyValueWidget(false);
+		ValueWidget->SetEnabled(bDataEditable);
+		ElementHandleRef->SetOnPropertyResetToDefault(FSimpleDelegate::CreateSP(this, &FCustomPrimitiveDataCustomization::SetDefaultValue, ElementHandle, PrimIdx));
+
+		Row.CustomWidget()
+		.NameContent()
+		[
+			NameWidget.ToSharedRef()
+		]
+		.ValueContent()
+		[
+			ValueWidget
+		];
+	}
+	else
+	{
+		// We don't have data for this row, add an empty row that contains the parameter names and the ability to add data up until this point
+		FDetailWidgetRow& Row = VectorGroup ? VectorGroup->AddWidgetRow() : ChildBuilder.AddCustomRow(FText::Join(LOCTEXT("SearchTextDelimiter", " "), SearchText));
+
+		Row.NameContent()
+		[
+			NameWidget.ToSharedRef()
+		]
+		.ValueContent()
+		[
+			SNew(SHorizontalBox)
+			.IsEnabled(bDataEditable)
+			+ SHorizontalBox::Slot()
+			.Padding(2.0f)
+			.HAlign(HAlign_Center)
+			.VAlign(VAlign_Center)
+			.AutoWidth()
+			[
+				PropertyCustomizationHelpers::MakeAddButton(FSimpleDelegate::CreateSP(this, &FCustomPrimitiveDataCustomization::OnAddedDesiredPrimitiveData, PrimIdx))
+			]
+		];
+	}
+}
+END_SLATE_FUNCTION_BUILD_OPTIMIZATION
+
+template<typename Predicate>
+void FCustomPrimitiveDataCustomization::ForEachSelectedComponent(Predicate Pred)
+{
+	if (PropertyUtils.IsValid())
+	{
+		for (TWeakObjectPtr<UObject> Object : PropertyUtils->GetSelectedObjects())
+		{
+			if (UPrimitiveComponent* Component = Cast<UPrimitiveComponent>(Object.Get()))
+			{
+				Pred(Component);
+			}
+			else if (AActor* Actor = Cast<AActor>(Object.Get()))
+			{
+				for (UActorComponent* ActorComponent : Actor->GetComponents())
+				{
+					if (UPrimitiveComponent* PrimitiveComponent = Cast<UPrimitiveComponent>(ActorComponent))
+					{
+						Pred(PrimitiveComponent);
+					}
+				}
+			}
+		}
+	}
+}
+
+void FCustomPrimitiveDataCustomization::Cleanup()
+{
+	FCoreUObjectDelegates::OnObjectPropertyChanged.RemoveAll(this);
+	UMaterial::OnMaterialCompilationFinished().RemoveAll(this);
+
+	PropertyUtils = NULL;
+	DataArrayHandle = NULL;
+
+	ComponentsToWatch.Empty();
+	MaterialsToWatch.Empty();
+	VectorParameterData.Empty();
+	ScalarParameterData.Empty();
+	ColorBlocks.Empty();
+}
+
+void FCustomPrimitiveDataCustomization::PopulateParameterData(UPrimitiveComponent* PrimitiveComponent, int32& MaxPrimitiveDataIndex)
+{
+	const int32 NumMaterials = PrimitiveComponent->GetNumMaterials();
+
+	for (int32 i = 0; i < NumMaterials; ++i)
+	{
+		UMaterialInterface* MaterialInterface = PrimitiveComponent->GetMaterial(i);
+		UMaterial* Material = MaterialInterface ? MaterialInterface->GetBaseMaterial() : NULL;
+
+		if (Material == NULL)
+		{
+			continue;
+		}
+
+		MaterialsToWatch.Add(Material);
+
+		TMap<FMaterialParameterInfo, FMaterialParameterMetadata> Parameters;
+
+		MaterialInterface->GetAllParametersOfType(EMaterialParameterType::Vector, Parameters);
+
+		for (const TPair<FMaterialParameterInfo, FMaterialParameterMetadata>& Parameter : Parameters)
+		{
+			const FMaterialParameterInfo& Info = Parameter.Key;
+			const FMaterialParameterMetadata& ParameterMetadata = Parameter.Value;
+			
+			if (ParameterMetadata.PrimitiveDataIndex > INDEX_NONE)
+			{
+				// Add each element individually, so that we can overlap vector parameter names
+				VectorParameterData.FindOrAdd(ParameterMetadata.PrimitiveDataIndex + 0).Add(
+				{
+					PrimitiveComponent,
+					MaterialInterface,
+					Info,
+					ParameterMetadata.ExpressionGuid,
+					0
+				});
+				VectorParameterData.FindOrAdd(ParameterMetadata.PrimitiveDataIndex + 1).Add(
+				{
+					PrimitiveComponent,
+					MaterialInterface,
+					Info,
+					ParameterMetadata.ExpressionGuid,
+					1
+				});
+				VectorParameterData.FindOrAdd(ParameterMetadata.PrimitiveDataIndex + 2).Add(
+				{
+					PrimitiveComponent,
+					MaterialInterface,
+					Info,
+					ParameterMetadata.ExpressionGuid,
+					2
+				});
+				VectorParameterData.FindOrAdd(ParameterMetadata.PrimitiveDataIndex + 3).Add(
+				{
+					PrimitiveComponent,
+					MaterialInterface,
+					Info,
+					ParameterMetadata.ExpressionGuid,
+					3
+				});
+				MaxPrimitiveDataIndex = FMath::Max(MaxPrimitiveDataIndex, (int32)(ParameterMetadata.PrimitiveDataIndex + 3));
+			}
+		}
+
+		Parameters.Reset();
+
+		MaterialInterface->GetAllParametersOfType(EMaterialParameterType::Scalar, Parameters);
+
+		for (const TPair<FMaterialParameterInfo, FMaterialParameterMetadata>& Parameter : Parameters)
+		{
+			const FMaterialParameterInfo& Info = Parameter.Key;
+			const FMaterialParameterMetadata& ParameterMetadata = Parameter.Value;
+
+			if (ParameterMetadata.PrimitiveDataIndex > INDEX_NONE)
+			{
+				ScalarParameterData.FindOrAdd(ParameterMetadata.PrimitiveDataIndex).Add(
+				{
+					PrimitiveComponent,
+					MaterialInterface,
+					Info,
+					ParameterMetadata.ExpressionGuid,
+					0
+				});
+				MaxPrimitiveDataIndex = FMath::Max(MaxPrimitiveDataIndex, (int32)ParameterMetadata.PrimitiveDataIndex);
+			}
+		}
+	}
+}
+
+void FCustomPrimitiveDataCustomization::OnUpdated()
+{
+	if (PropertyUtils.IsValid())
+	{
+		PropertyUtils->ForceRefresh();
+	}
+}
+	
+void FCustomPrimitiveDataCustomization::OnObjectPropertyChanged(UObject* Object, FPropertyChangedEvent& PropertyChangedEvent)
+{
+	const bool bIsCustomPrimitiveDataProperty = PropertyChangedEvent.GetPropertyName() == "CustomPrimitiveData"
+		|| (PropertyChangedEvent.MemberProperty != NULL && PropertyChangedEvent.MemberProperty->GetFName() == "CustomPrimitiveData");
+	if (ComponentsToWatch.Contains(Object) && !bIsCustomPrimitiveDataProperty)
+	{
+		OnUpdated();
+	}
+}
+
+void FCustomPrimitiveDataCustomization::OnMaterialCompiled(UMaterialInterface* Material)
+{
+	// NOTE: We use a soft object ptr here as the old material object will be stale on compile
+	if (MaterialsToWatch.Contains(Material))
+	{
+		OnUpdated();
+	}
+}
+
+void FCustomPrimitiveDataCustomization::OnNavigate(TWeakObjectPtr<UMaterialInterface> MaterialInterface, FGuid ExpressionID)
+{
+	UMaterial* Material = MaterialInterface.IsValid() ? MaterialInterface->GetMaterial() : NULL;
+	
+	if (UMaterialExpression* Expression = Material ? Material->FindExpressionByGUID<UMaterialExpression>(ExpressionID) : NULL)
+	{
+		UAssetEditorSubsystem* AssetEditorSubsystem = GEditor->GetEditorSubsystem<UAssetEditorSubsystem>();
+		AssetEditorSubsystem->OpenEditorForAsset(Material);
+		IMaterialEditor* MaterialEditor = (IMaterialEditor*)AssetEditorSubsystem->FindEditorForAsset(Material, true);
+		MaterialEditor->JumpToExpression(Expression);
+	}
+}
+
+void FCustomPrimitiveDataCustomization::OnAddedDesiredPrimitiveData(uint8 PrimIdx)
+{
+	uint32 NumElements;
+	if (DataArrayHandle.IsValid() && DataArrayHandle->GetNumElements(NumElements) == FPropertyAccess::Success)
+	{
+		GEditor->BeginTransaction(LOCTEXT("OnAddedDesiredPrimitiveData", "Added Items"));
+
+		for (int32 i = NumElements; i <= PrimIdx; ++i)
+		{
+			DataArrayHandle->AddItem();
+
+			SetDefaultValue(DataArrayHandle->GetElement(i), i);
+		}
+
+		GEditor->EndTransaction();
+	}
+}
+
+void FCustomPrimitiveDataCustomization::OnRemovedPrimitiveData(uint8 PrimIdx)
+{
+	uint32 NumElements;
+	if (DataArrayHandle.IsValid() && DataArrayHandle->GetNumElements(NumElements) == FPropertyAccess::Success)
+	{
+		GEditor->BeginTransaction(LOCTEXT("OnRemovedPrimitiveData", "Removed Items"));
+
+		for (int32 i = NumElements - 1; i >= PrimIdx; --i)
+		{
+			DataArrayHandle->DeleteItem(i);
+		}
+
+		GEditor->EndTransaction();
+	}
+}
+
+FLinearColor FCustomPrimitiveDataCustomization::GetVectorColor(uint8 PrimIdx) const
+{
+	FVector4 Color(ForceInitToZero);
+
+	uint32 NumElems;
+	if (DataArrayHandle.IsValid() && DataArrayHandle->GetNumElements(NumElems) == FPropertyAccess::Success)
+	{
+		const int32 MaxElems = FMath::Min((int32)NumElems, PrimIdx + 4);
+
+		for (int32 i = PrimIdx; i < MaxElems; ++i)
+		{
+			DataArrayHandle->GetElement(i)->GetValue(Color[i - PrimIdx]);
+		}
+	}
+
+	return FLinearColor(Color);
+}
+
+void FCustomPrimitiveDataCustomization::SetVectorColor(FLinearColor NewColor, uint8 PrimIdx)
+{
+	FVector4 Color(NewColor);
+
+	uint32 NumElems;
+	if (DataArrayHandle.IsValid() && DataArrayHandle->GetNumElements(NumElems) == FPropertyAccess::Success)
+	{
+		const int32 MaxElems = FMath::Min((int32)NumElems, PrimIdx + 4);
+
+		for (int32 i = PrimIdx; i < MaxElems; ++i)
+		{
+			DataArrayHandle->GetElement(i)->SetValue(Color[i - PrimIdx]);
+		}
+	}
+}
+
+void FCustomPrimitiveDataCustomization::SetDefaultValue(TSharedPtr<IPropertyHandle> Handle, uint8 PrimIdx)
+{
+	if (Handle.IsValid())
+	{
+		TSet<TWeakObjectPtr<UPrimitiveComponent>> ChangedComponents;
+
+		// Prioritize vector data since we have a color picker
+		if (VectorParameterData.Contains(PrimIdx))
+		{
+			for (const FParameterData& ParameterData : VectorParameterData[PrimIdx])
+			{
+				if (ParameterData.Component.IsValid() && !ChangedComponents.Contains(ParameterData.Component))
+				{
+					FLinearColor Color(ForceInitToZero);
+					if (!ParameterData.Material.IsValid() || ParameterData.Material->GetVectorParameterDefaultValue(ParameterData.Info, Color))
+					{
+						float* ColorPtr = reinterpret_cast<float*>(&Color);
+						ParameterData.Component->SetDefaultCustomPrimitiveDataFloat(PrimIdx, ColorPtr[ParameterData.IndexOffset]);
+
+						FPropertyChangedEvent PropertyChangedEvent(Handle->GetParentHandle()->GetParentHandle()->GetProperty());
+						PropertyChangedEvent.SetActiveMemberProperty(Handle->GetParentHandle()->GetProperty());
+						ParameterData.Component->PostEditChangeProperty(PropertyChangedEvent);
+
+						ChangedComponents.Add(ParameterData.Component);
+					}
+				}
+			}
+		}
+
+		if (ScalarParameterData.Contains(PrimIdx))
+		{
+			for (const FParameterData& ParameterData : VectorParameterData[PrimIdx])
+			{
+				if (ParameterData.Component.IsValid() && !ChangedComponents.Contains(ParameterData.Component))
+				{
+					float Value = 0.f;
+					if (!ParameterData.Material.IsValid() || ParameterData.Material->GetScalarParameterDefaultValue(ParameterData.Info, Value))
+					{
+						ParameterData.Component->SetDefaultCustomPrimitiveDataFloat(PrimIdx, Value);
+
+						FPropertyChangedEvent PropertyChangedEvent(Handle->GetParentHandle()->GetParentHandle()->GetProperty());
+						PropertyChangedEvent.SetActiveMemberProperty(Handle->GetParentHandle()->GetProperty());
+						ParameterData.Component->PostEditChangeProperty(PropertyChangedEvent);
+
+						ChangedComponents.Add(ParameterData.Component);
+					}
+				}
+			}
+		}
+	}
+}
+
+void FCustomPrimitiveDataCustomization::SetDefaultVectorValue(uint8 PrimIdx)
+{
+	uint32 NumElems;
+	if (DataArrayHandle.IsValid() && DataArrayHandle->GetNumElements(NumElems) == FPropertyAccess::Success)
+	{
+		GEditor->BeginTransaction(LOCTEXT("SetDefaultVectorValue", "Reset Vector To Default"));
+
+		const int32 MaxElems = FMath::Min((int32)NumElems, PrimIdx + 4);
+		for (int32 i = PrimIdx; i < MaxElems; ++i)
+		{
+			DataArrayHandle->GetElement(i)->ResetToDefault();
+		}
+
+		GEditor->EndTransaction();
+	}
+}
+
+FReply FCustomPrimitiveDataCustomization::OnMouseButtonDownColorBlock(const FGeometry& MyGeometry, const FPointerEvent& MouseEvent, uint8 PrimIdx)
+{
+	if (MouseEvent.GetEffectingButton() != EKeys::LeftMouseButton)
+	{
+		return FReply::Unhandled();
+	}
+
+	GEditor->BeginTransaction(FText::Format(LOCTEXT("SetVectorColor", "Edit Primitive Data Vector: {0}"), FText::AsNumber(PrimIdx)));
+
+	FColorPickerArgs PickerArgs;
+	PickerArgs.bUseAlpha = true;
+	PickerArgs.InitialColorOverride = GetVectorColor(PrimIdx);
+	PickerArgs.ParentWidget = ColorBlocks[PrimIdx];
+	PickerArgs.DisplayGamma = TAttribute<float>::Create(TAttribute<float>::FGetter::CreateUObject(GEngine, &UEngine::GetDisplayGamma));
+	PickerArgs.OnColorCommitted = FOnLinearColorValueChanged::CreateSP(this, &FCustomPrimitiveDataCustomization::SetVectorColor, PrimIdx);
+	PickerArgs.OnColorPickerCancelled = FOnColorPickerCancelled::CreateSP(this, &FCustomPrimitiveDataCustomization::OnColorPickerCancelled, PrimIdx);
+	PickerArgs.OnColorPickerWindowClosed = FOnWindowClosed::CreateSP(this, &FCustomPrimitiveDataCustomization::OnColorPickerWindowClosed);
+
+	OpenColorPicker(PickerArgs);
+
+	return FReply::Handled();
+}
+
+void FCustomPrimitiveDataCustomization::OnColorPickerCancelled(FLinearColor OriginalColor, uint8 PrimIdx)
+{
+	SetVectorColor(OriginalColor, PrimIdx);
+	GEditor->CancelTransaction(0);
+}
+
+void FCustomPrimitiveDataCustomization::OnColorPickerWindowClosed(const TSharedRef<SWindow>& Window)
+{
+	GEditor->EndTransaction();
+}
+
+BEGIN_SLATE_FUNCTION_BUILD_OPTIMIZATION
+TSharedRef<SHyperlink> FCustomPrimitiveDataCustomization::CreateHyperlink(FText Text, TWeakObjectPtr<UMaterialInterface> Material, const FGuid& ExpressionID)
+{
+	return SNew(SHyperlink)
+		.Text(Text)
+		.OnNavigate(this, &FCustomPrimitiveDataCustomization::OnNavigate, Material, ExpressionID)
+		.Style(FEditorStyle::Get(), "HoverOnlyHyperlink")
+		.TextStyle(FEditorStyle::Get(), "DetailsView.HyperlinkStyle");
+}
+END_SLATE_FUNCTION_BUILD_OPTIMIZATION
+
+BEGIN_SLATE_FUNCTION_BUILD_OPTIMIZATION
+TSharedRef<SWidget> FCustomPrimitiveDataCustomization::GetUndefinedParameterWidget(int32 PrimIdx, IPropertyTypeCustomizationUtils& CustomizationUtils) const
+{
+	FText PrimIdxText = FText::AsNumber(PrimIdx);
+	TSharedRef<SWidget> UndefinedParamWidget = SNew(STextBlock)
+		.Text(FText::Format(LOCTEXT("UndefinedParameter", "{0} (Undefined)"), PrimIdxText))
+		.Font(CustomizationUtils.GetRegularFont());
+	UndefinedParamWidget->SetToolTipText(FText::Format(LOCTEXT("UndefinedParameterTooltip", "A component is selected that doesn't define a parameter for primitive index {0}"), PrimIdxText));
+		
+	return UndefinedParamWidget;
+}
+END_SLATE_FUNCTION_BUILD_OPTIMIZATION
+
+#undef LOCTEXT_NAMESPACE
diff --git a/Engine/Source/Editor/DetailCustomizations/Private/CustomPrimitiveDataCustomization.h b/Engine/Source/Editor/DetailCustomizations/Private/CustomPrimitiveDataCustomization.h
new file mode 100644
--- /dev/null
+++ b/Engine/Source/Editor/DetailCustomizations/Private/CustomPrimitiveDataCustomization.h
@@ -0,0 +1,105 @@
+// Copyright Epic Games, Inc. All Rights Reserved.
+
+#pragma once
+
+#include "CoreMinimal.h"
+#include "IPropertyTypeCustomization.h"
+#include "IDetailChildrenBuilder.h"
+
+class UMaterialExpressionVectorParameter;
+class UMaterialInterface;
+
+class FDetailWidgetRow;
+class SColorBlock;
+
+class IPropertyHandleArray;
+
+class FCustomPrimitiveDataCustomization : public IPropertyTypeCustomization
+{
+public:
+	static TSharedRef<IPropertyTypeCustomization> MakeInstance();
+
+	~FCustomPrimitiveDataCustomization();
+
+	/** IPropertyTypeCustomization interface */
+	void CustomizeHeader(TSharedRef<IPropertyHandle> PropertyHandle, FDetailWidgetRow& HeaderRow, IPropertyTypeCustomizationUtils& CustomizationUtils) override;
+	void CustomizeChildren(TSharedRef<IPropertyHandle> PropertyHandle, IDetailChildrenBuilder& ChildBuilder, IPropertyTypeCustomizationUtils& CustomizationUtils) override;
+
+private:
+	struct FParameterData
+	{
+	public:
+		TWeakObjectPtr<UPrimitiveComponent> Component;
+		TWeakObjectPtr<UMaterialInterface> Material;
+		FMaterialParameterInfo Info;
+		FGuid ExpressionID;
+		uint8 IndexOffset;
+	};
+
+	TSharedPtr<IPropertyUtilities> PropertyUtils;
+	TSharedPtr<IPropertyHandleArray> DataArrayHandle;
+
+	TSet<TWeakObjectPtr<UObject>> ComponentsToWatch;
+	TSet<TSoftObjectPtr<UMaterialInterface>> MaterialsToWatch;
+
+	TMap<uint8, TArray<FParameterData>> VectorParameterData;
+	TMap<uint8, TArray<FParameterData>> ScalarParameterData;
+	TMap<uint8, TSharedPtr<SColorBlock>> ColorBlocks;
+
+	/**
+	 * Creates a group for the specified vector parameter, with a header that contains the vector parameter name and
+	 * a color block that opens a color picker when clicked. Also contains add, delete and reset to default buttons to
+	 * quickly populate a vector parameter
+	 */
+	IDetailGroup* CreateVectorGroup(IDetailChildrenBuilder& ChildBuilder, uint8 PrimIdx, bool bDataEditable, int32 NumElements);
+
+	/**
+	 * Creates a row for the specified primitive index. Will create hyperlinks for each parameter, to take the user to the parameter's node.
+	 * If there are multiple unique parameters for the one row, it will fill the name content with each individual parameter name, or 
+	 * mention that the row is undefined for a component.
+	 * 
+	 * If the parameter doesn't have an element yet, provides an add button for the user to quickly add elements up until that parameter.
+	 */
+	void CreateParameterRow(IDetailChildrenBuilder& ChildBuilder, uint8 PrimIdx, TSharedPtr<IPropertyHandle> ElementHandle, int32 NumSelectedComponents, bool bDataEditable, IDetailGroup* VectorGroup, IPropertyTypeCustomizationUtils& CustomizationUtils);
+
+	template<typename Predicate>
+	void ForEachSelectedComponent(Predicate Pred);
+
+	void Cleanup();
+
+	/** Populates our parameter data for scalar and vector parameters. Updates the max primitive index if necessary. */
+	void PopulateParameterData(UPrimitiveComponent* PrimitiveComponent, int32& MaxPrimitiveDataIndex);
+
+	void OnUpdated();
+	void OnObjectPropertyChanged(UObject* Object, FPropertyChangedEvent& PropertyChangedEvent);
+	void OnMaterialCompiled(UMaterialInterface* Material);
+
+	/** Will open the material editor to the specified parameter expression when a user clicks a hyperlink */
+	void OnNavigate(TWeakObjectPtr<UMaterialInterface> Material, FGuid ExpressionID);
+
+	/** Adds elements to the primitive data array up until and including the specified PrimIdx. Also initializes elements with defaults from material interface */
+	void OnAddedDesiredPrimitiveData(uint8 PrimIdx);
+
+	/** Remove all elements up until and including this primitive index */
+	void OnRemovedPrimitiveData(uint8 PrimIdx);
+
+	FLinearColor GetVectorColor( uint8 PrimIdx) const;
+	void SetVectorColor(FLinearColor NewColor, uint8 PrimIdx);
+
+	/** Iterates through each parameter, and sets their values to the material interface's defaults on a per-component basis */
+	void SetDefaultValue(TSharedPtr<IPropertyHandle> Handle, uint8 PrimIdx);
+
+	/** Helper method to SetDefaultValue for vectors, but makes one transaction for all edits */
+	void SetDefaultVectorValue(uint8 PrimIdx);
+
+	FReply OnMouseButtonDownColorBlock(const FGeometry& MyGeometry, const FPointerEvent& MouseEvent, uint8 PrimIdx);
+	void OnColorPickerCancelled(FLinearColor OriginalColor, uint8 PrimIdx);
+	void OnColorPickerWindowClosed(const TSharedRef<SWindow>& Window);
+
+	/** Hyperlink that's used for navigating to parameters */
+	TSharedRef<class SHyperlink> CreateHyperlink(FText Text, TWeakObjectPtr<UMaterialInterface> Material, const FGuid& ExpressionId);
+
+	/** Default widget that's used when a component is missing defined parameters for the specified primitive index */
+	TSharedRef<class SWidget> GetUndefinedParameterWidget(int32 PrimIdx, IPropertyTypeCustomizationUtils& CustomizationUtils) const;
+};
+
diff --git a/Engine/Source/Editor/DetailCustomizations/Private/DetailCustomizations.cpp b/Engine/Source/Editor/DetailCustomizations/Private/DetailCustomizations.cpp
--- a/Engine/Source/Editor/DetailCustomizations/Private/DetailCustomizations.cpp
+++ b/Engine/Source/Editor/DetailCustomizations/Private/DetailCustomizations.cpp
@@ -146,6 +146,7 @@
 #include "BoundsCopyComponentDetails.h"
 #include "SupportedRangeTypes.h"	// StructsSupportingRangeVisibility
 #include "MaterialInstanceDynamicDetails.h"
+#include "CustomPrimitiveDataCustomization.h"
 
 IMPLEMENT_MODULE( FDetailCustomizationsModule, DetailCustomizations );
 
@@ -289,6 +290,7 @@ void FDetailCustomizationsModule::RegisterPropertyTypeCustomizations()
 	RegisterCustomPropertyTypeLayout("ComponentReference", FOnGetPropertyTypeCustomizationInstance::CreateStatic(&FComponentReferenceCustomization::MakeInstance));
 	RegisterCustomPropertyTypeLayout("EMaterialShadingModel", FOnGetPropertyTypeCustomizationInstance::CreateStatic(&FMaterialShadingModelCustomization::MakeInstance));
 	RegisterCustomPropertyTypeLayout("DebugCameraControllerSettingsViewModeIndex", FOnGetPropertyTypeCustomizationInstance::CreateStatic(&FDebugCameraControllerSettingsViewModeIndexCustomization::MakeInstance));
+	RegisterCustomPropertyTypeLayout("CustomPrimitiveData", FOnGetPropertyTypeCustomizationInstance::CreateStatic(&FCustomPrimitiveDataCustomization::MakeInstance));
 }
 
 #undef REGISTER_UIMINMAX_CUSTOMIZATION
diff --git a/Engine/Source/Runtime/Engine/Classes/Components/PrimitiveComponent.h b/Engine/Source/Runtime/Engine/Classes/Components/PrimitiveComponent.h
--- a/Engine/Source/Runtime/Engine/Classes/Components/PrimitiveComponent.h
+++ b/Engine/Source/Runtime/Engine/Classes/Components/PrimitiveComponent.h
@@ -675,7 +675,7 @@ public:
 
 private:
 	/** Optional user defined default values for the custom primitive data of this primitive */
-	UPROPERTY(EditAnywhere, AdvancedDisplay, Category=Rendering, meta = (DisplayName = "Custom Primitive Data Defaults"))
+	UPROPERTY(EditAnywhere, Category=Rendering, meta = (DisplayName = "Custom Primitive Data Defaults"))
 	FCustomPrimitiveData CustomPrimitiveData;
 
 	/** Custom data that can be read by a material through a material parameter expression. Set data using SetCustomPrimitiveData* functions */
@@ -921,6 +921,38 @@ public:
 	/** Get the mask filter checked when others move into us. */
 	FMaskFilter GetMaskFilterOnBodyInstance(FMaskFilter InMaskFilter) const { return BodyInstance.GetMaskFilter(); }
 
+	/**
+	 * Gets the index of the scalar parameter for the custom primitive data array
+	 * @param	ParameterName	The parameter name of the custom primitive
+	 * @return	The index of the custom primitive, INDEX_NONE (-1) if not found
+	 */
+	UFUNCTION(BlueprintCallable, Category = "Rendering|Material")
+	int32 GetCustomPrimitiveDataIndexForScalarParameter(FName ParameterName) const;
+
+	/**
+	 * Gets the index of the vector parameter for the custom primitive data array
+	 * @param	ParameterName	The parameter name of the custom primitive
+	 * @return	The index of the custom primitive, INDEX_NONE (-1) if not found
+	 */
+	UFUNCTION(BlueprintCallable, Category = "Rendering|Material")
+	int32 GetCustomPrimitiveDataIndexForVectorParameter(FName ParameterName) const;
+
+	/**
+	 * Set a scalar parameter for custom primitive data. This sets the run-time data only, so it doesn't serialize.
+	 * @param	ParameterName	The parameter name of the custom primitive
+	 * @param	Value			The new value of the custom primitive
+	 */
+	UFUNCTION(BlueprintCallable, Category = "Rendering|Material")
+	void SetScalarParameterForCustomPrimitiveData(FName ParameterName, float Value);
+
+	/**
+	 * Set a vector parameter for custom primitive data. This sets the run-time data only, so it doesn't serialize.
+	 * @param	ParameterName	The parameter name of the custom primitive
+	 * @param	Value			The new value of the custom primitive
+	 */
+	UFUNCTION(BlueprintCallable, Category = "Rendering|Material")
+	void SetVectorParameterForCustomPrimitiveData(FName ParameterName, FVector4 Value);
+
 	/** Set custom primitive data at index DataIndex. This sets the run-time data only, so it doesn't serialize. */
 	UFUNCTION(BlueprintCallable, Category="Rendering|Material")
 	void SetCustomPrimitiveDataFloat(int32 DataIndex, float Value);
@@ -943,6 +975,22 @@ public:
 	 */
 	const FCustomPrimitiveData& GetCustomPrimitiveData() const { return CustomPrimitiveDataInternal; }
 
+	/**
+	 * Set a scalar parameter for default custom primitive data. This will be serialized and is useful in construction scripts.
+	 * @param	ParameterName	The parameter name of the custom primitive
+	 * @param	Value			The new value of the custom primitive
+	 */
+	UFUNCTION(BlueprintCallable, Category = "Rendering|Material")
+	void SetScalarParameterForDefaultCustomPrimitiveData(FName ParameterName, float Value);
+
+	/**
+	 * Set a vector parameter for default custom primitive data. This will be serialized and is useful in construction scripts.
+	 * @param	ParameterName	The parameter name of the custom primitive
+	 * @param	Value			The new value of the custom primitive
+	 */
+	UFUNCTION(BlueprintCallable, Category = "Rendering|Material")
+	void SetVectorParameterForDefaultCustomPrimitiveData(FName ParameterName, FVector4 Value);
+
 	/** Set default custom primitive data at index DataIndex, and marks the render state dirty */
 	UFUNCTION(BlueprintCallable, Category = "Rendering|Material")
 	void SetDefaultCustomPrimitiveDataFloat(int32 DataIndex, float Value);
diff --git a/Engine/Source/Runtime/Engine/Classes/Materials/MaterialExpressionScalarParameter.h b/Engine/Source/Runtime/Engine/Classes/Materials/MaterialExpressionScalarParameter.h
--- a/Engine/Source/Runtime/Engine/Classes/Materials/MaterialExpressionScalarParameter.h
+++ b/Engine/Source/Runtime/Engine/Classes/Materials/MaterialExpressionScalarParameter.h
@@ -47,6 +47,12 @@ class UMaterialExpressionScalarParameter : public UMaterialExpressionParameter
 	virtual bool GetParameterValue(FMaterialParameterMetadata& OutMeta) const override
 	{
 		OutMeta.Value = DefaultValue;
+		
+		if (bUseCustomPrimitiveData)
+		{
+			OutMeta.PrimitiveDataIndex = PrimitiveDataIndex;
+		}
+
 		OutMeta.ScalarMin = SliderMin;
 		OutMeta.ScalarMax = SliderMax;
 		return Super::GetParameterValue(OutMeta);
diff --git a/Engine/Source/Runtime/Engine/Classes/Materials/MaterialExpressionVectorParameter.h b/Engine/Source/Runtime/Engine/Classes/Materials/MaterialExpressionVectorParameter.h
--- a/Engine/Source/Runtime/Engine/Classes/Materials/MaterialExpressionVectorParameter.h
+++ b/Engine/Source/Runtime/Engine/Classes/Materials/MaterialExpressionVectorParameter.h
@@ -36,6 +36,12 @@ class UMaterialExpressionVectorParameter : public UMaterialExpressionParameter
 	virtual bool GetParameterValue(FMaterialParameterMetadata& OutMeta) const override
 	{
 		OutMeta.Value = DefaultValue;
+
+		if (bUseCustomPrimitiveData)
+		{
+			OutMeta.PrimitiveDataIndex = PrimitiveDataIndex;
+		}
+
 		OutMeta.ChannelNames = ChannelNames;
 		return Super::GetParameterValue(OutMeta);
 	}
diff --git a/Engine/Source/Runtime/Engine/Private/Components/PrimitiveComponent.cpp b/Engine/Source/Runtime/Engine/Private/Components/PrimitiveComponent.cpp
--- a/Engine/Source/Runtime/Engine/Private/Components/PrimitiveComponent.cpp
+++ b/Engine/Source/Runtime/Engine/Private/Components/PrimitiveComponent.cpp
@@ -1945,6 +1945,65 @@ void UPrimitiveComponent::SetDefaultCustomPrimitiveData(int32 DataIndex, const T
 	}
 }
 
+int32 UPrimitiveComponent::GetCustomPrimitiveDataIndexForScalarParameter(FName ParameterName) const
+{
+	const int32 NumMaterials = GetNumMaterials();
+
+	for (int32 i = 0; i < NumMaterials; ++i)
+	{
+		if (UMaterialInterface* Material = GetMaterial(i))
+		{
+			FMaterialParameterMetadata ParameterMetadata;
+			if (Material->GetParameterValue(EMaterialParameterType::Scalar, FMemoryImageMaterialParameterInfo(ParameterName), ParameterMetadata, EMaterialGetParameterValueFlags::CheckAll))
+			{
+				return ParameterMetadata.PrimitiveDataIndex;
+			}
+		}
+	}
+
+	return INDEX_NONE;
+}
+
+int32 UPrimitiveComponent::GetCustomPrimitiveDataIndexForVectorParameter(FName ParameterName) const
+{
+	const int32 NumMaterials = GetNumMaterials();
+
+	for (int32 i = 0; i < NumMaterials; ++i)
+	{
+		if (UMaterialInterface* Material = GetMaterial(i))
+		{
+			FMaterialParameterMetadata ParameterMetadata;
+			if (Material->GetParameterValue(EMaterialParameterType::Vector, FMemoryImageMaterialParameterInfo(ParameterName), ParameterMetadata, EMaterialGetParameterValueFlags::CheckAll))
+			{
+				return ParameterMetadata.PrimitiveDataIndex;
+			}
+		}
+	}
+
+	return INDEX_NONE;
+}
+
+void UPrimitiveComponent::SetScalarParameterForCustomPrimitiveData(FName ParameterName, float Value)
+{
+	int32 PrimitiveDataIndex = GetCustomPrimitiveDataIndexForScalarParameter(ParameterName);
+
+	if (PrimitiveDataIndex > INDEX_NONE)
+	{
+		SetCustomPrimitiveDataInternal(PrimitiveDataIndex, {Value});
+	}
+}
+
+void UPrimitiveComponent::SetVectorParameterForCustomPrimitiveData(FName ParameterName, FVector4 Value)
+{
+	const int32 PrimitiveDataIndex = GetCustomPrimitiveDataIndexForVectorParameter(ParameterName);
+
+	if (PrimitiveDataIndex > INDEX_NONE)
+	{	// LWC_TODO: precision loss
+		FVector4f ValueFlt(Value);
+		SetCustomPrimitiveDataInternal(PrimitiveDataIndex, { ValueFlt.X, ValueFlt.Y, ValueFlt.Z, ValueFlt.W });
+	}
+}
+
 void UPrimitiveComponent::SetCustomPrimitiveDataFloat(int32 DataIndex, float Value)
 {
 	SetCustomPrimitiveDataInternal(DataIndex, {Value});
@@ -1969,6 +2028,28 @@ void UPrimitiveComponent::SetCustomPrimitiveDataVector4(int32 DataIndex, FVector
 	SetCustomPrimitiveDataInternal(DataIndex, {ValueFlt.X, ValueFlt.Y, ValueFlt.Z, ValueFlt.W});
 }
 
+void UPrimitiveComponent::SetScalarParameterForDefaultCustomPrimitiveData(FName ParameterName, float Value)
+{
+	int32 PrimitiveDataIndex = GetCustomPrimitiveDataIndexForScalarParameter(ParameterName);
+
+	if (PrimitiveDataIndex > INDEX_NONE)
+	{
+		SetDefaultCustomPrimitiveData(PrimitiveDataIndex, { Value });
+	}
+}
+
+void UPrimitiveComponent::SetVectorParameterForDefaultCustomPrimitiveData(FName ParameterName, FVector4 Value)
+{
+	const int32 PrimitiveDataIndex = GetCustomPrimitiveDataIndexForVectorParameter(ParameterName);
+
+	if (PrimitiveDataIndex > INDEX_NONE)
+	{
+		// LWC_TODO: precision loss
+		FVector4f ValueFlt(Value);
+		SetDefaultCustomPrimitiveData(PrimitiveDataIndex, { ValueFlt.X, ValueFlt.Y, ValueFlt.Z, ValueFlt.W });
+	}
+}
+
 void UPrimitiveComponent::SetDefaultCustomPrimitiveDataFloat(int32 DataIndex, float Value)
 {
 	SetDefaultCustomPrimitiveData(DataIndex, { Value });
diff --git a/Engine/Source/Runtime/Engine/Private/Materials/MaterialCachedData.cpp b/Engine/Source/Runtime/Engine/Private/Materials/MaterialCachedData.cpp
--- a/Engine/Source/Runtime/Engine/Private/Materials/MaterialCachedData.cpp
+++ b/Engine/Source/Runtime/Engine/Private/Materials/MaterialCachedData.cpp
@@ -736,6 +736,8 @@ void FMaterialCachedParameters::GetParameterValueByIndex(EMaterialParameterType
 	{
 	case EMaterialParameterType::Scalar:
 		OutResult.Value = ScalarValues[ParameterIndex];
+		OutResult.PrimitiveDataIndex = ScalarPrimitiveDataIndexValues[ParameterIndex];
+
 #if WITH_EDITORONLY_DATA
 		OutResult.ScalarMin = ScalarMinMaxValues[ParameterIndex].X;
 		OutResult.ScalarMax = ScalarMinMaxValues[ParameterIndex].Y;
@@ -753,6 +755,8 @@ void FMaterialCachedParameters::GetParameterValueByIndex(EMaterialParameterType
 		break;
 	case EMaterialParameterType::Vector:
 		OutResult.Value = VectorValues[ParameterIndex];
+		OutResult.PrimitiveDataIndex = VectorPrimitiveDataIndexValues[ParameterIndex];
+
 #if  WITH_EDITORONLY_DATA
 		OutResult.ChannelNames = VectorChannelNameValues[ParameterIndex];
 		OutResult.bUsedAsChannelMask = VectorUsedAsChannelMaskValues[ParameterIndex];
diff --git a/Engine/Source/Runtime/Engine/Public/MaterialCachedData.h b/Engine/Source/Runtime/Engine/Public/MaterialCachedData.h
--- a/Engine/Source/Runtime/Engine/Public/MaterialCachedData.h
+++ b/Engine/Source/Runtime/Engine/Public/MaterialCachedData.h
@@ -133,6 +133,12 @@ struct FMaterialCachedParameters
 	UPROPERTY()
 	FMaterialCachedParameterEntry RuntimeEntries[NumMaterialRuntimeParameterTypes];
 
+	UPROPERTY()
+	TArray<int32> ScalarPrimitiveDataIndexValues;
+
+	UPROPERTY()
+	TArray<int32> VectorPrimitiveDataIndexValues;
+
 	UPROPERTY()
 	TArray<float> ScalarValues;
 
diff --git a/Engine/Source/Runtime/Engine/Public/MaterialTypes.h b/Engine/Source/Runtime/Engine/Public/MaterialTypes.h
--- a/Engine/Source/Runtime/Engine/Public/MaterialTypes.h
+++ b/Engine/Source/Runtime/Engine/Public/MaterialTypes.h
@@ -196,6 +196,7 @@ inline bool operator!=(const FMaterialParameterValue& Lhs, const FMaterialParame
 struct FMaterialParameterMetadata
 {
 	FMaterialParameterValue Value;
+	int32 PrimitiveDataIndex = INDEX_NONE;
 
 	FMaterialParameterMetadata() = default;
 	FMaterialParameterMetadata(const FMaterialParameterValue& InValue) : Value(InValue) {}
