FIXED: Explicitly set FVector4f for custom primitive data color pickers due to LWC

diff --git a/Engine/Source/Editor/DetailCustomizations/Private/CustomPrimitiveDataCustomization.cpp b/Engine/Source/Editor/DetailCustomizations/Private/CustomPrimitiveDataCustomization.cpp
--- a/Engine/Source/Editor/DetailCustomizations/Private/CustomPrimitiveDataCustomization.cpp
+++ b/Engine/Source/Editor/DetailCustomizations/Private/CustomPrimitiveDataCustomization.cpp
@@ -609,7 +609,7 @@ void FCustomPrimitiveDataCustomization::OnRemovedPrimitiveData(uint8 PrimIdx)
 
 FLinearColor FCustomPrimitiveDataCustomization::GetVectorColor(uint8 PrimIdx) const
 {
-	FVector4 Color(ForceInitToZero);
+	FVector4f Color(ForceInitToZero);
 
 	uint32 NumElems;
 	if (DataArrayHandle.IsValid() && DataArrayHandle->GetNumElements(NumElems) == FPropertyAccess::Success)
@@ -627,7 +627,7 @@ FLinearColor FCustomPrimitiveDataCustomization::GetVectorColor(uint8 PrimIdx) co
 
 void FCustomPrimitiveDataCustomization::SetVectorColor(FLinearColor NewColor, uint8 PrimIdx)
 {
-	FVector4 Color(NewColor);
+	FVector4f Color(NewColor);
 
 	uint32 NumElems;
 	if (DataArrayHandle.IsValid() && DataArrayHandle->GetNumElements(NumElems) == FPropertyAccess::Success)
