Ensure that vector group only created for first element of a vector #rb ben.ingram #jira none

diff --git a/Engine/Source/Editor/DetailCustomizations/Private/CustomPrimitiveDataCustomization.cpp b/Engine/Source/Editor/DetailCustomizations/Private/CustomPrimitiveDataCustomization.cpp
--- a/Engine/Source/Editor/DetailCustomizations/Private/CustomPrimitiveDataCustomization.cpp
+++ b/Engine/Source/Editor/DetailCustomizations/Private/CustomPrimitiveDataCustomization.cpp
@@ -102,12 +102,25 @@ void FCustomPrimitiveDataCustomization::CustomizeChildren(TSharedRef<IPropertyHa
 			VectorGroup = NULL;
 		}
 
-		// Always prioritize the first vector found
+		// Always prioritize the first vector found, and only if it's the first element of the vector
 		if (VectorGroup == NULL && VectorParameterData.Contains(PrimIdx))
 		{
-			// Create a collapsing group that contains our color picker, so we can quickly assign colors to our vector
-			VectorGroupPrimIdx = PrimIdx;
-			VectorGroup = CreateVectorGroup(ChildBuilder, PrimIdx, bDataEditable, NumElements);			
+			bool bContainsFirstElementOfVector = false;
+			for (const FParameterData& ParameterData : VectorParameterData[PrimIdx])
+			{
+				if (ParameterData.IndexOffset == 0)
+				{
+					bContainsFirstElementOfVector = true;
+					break;
+				}
+			}
+
+			if (bContainsFirstElementOfVector)
+			{
+				// Create a collapsing group that contains our color picker, so we can quickly assign colors to our vector
+				VectorGroupPrimIdx = PrimIdx;
+				VectorGroup = CreateVectorGroup(ChildBuilder, PrimIdx, bDataEditable, NumElements);
+			}
 		}
 
 		if (ScalarParameterData.Contains(PrimIdx) || VectorParameterData.Contains(PrimIdx))
