Fixed crash in CPD customization when navigating to expressions not defined in the material itself (i.e. functions/layers)

diff --git a/Engine/Source/Editor/DetailCustomizations/Private/CustomPrimitiveDataCustomization.cpp b/Engine/Source/Editor/DetailCustomizations/Private/CustomPrimitiveDataCustomization.cpp
--- a/Engine/Source/Editor/DetailCustomizations/Private/CustomPrimitiveDataCustomization.cpp
+++ b/Engine/Source/Editor/DetailCustomizations/Private/CustomPrimitiveDataCustomization.cpp
@@ -579,10 +579,22 @@ void FCustomPrimitiveDataCustomization::OnNavigate(TWeakObjectPtr<UMaterialInter
 	
 	if (UMaterialExpression* Expression = Material ? Material->FindExpressionByGUID<UMaterialExpression>(ExpressionID) : NULL)
 	{
+		// FindExpression is recursive, so we need to ensure we open the correct asset
+		UObject* Asset = Expression->GetOutermostObject();
 		UAssetEditorSubsystem* AssetEditorSubsystem = GEditor->GetEditorSubsystem<UAssetEditorSubsystem>();
-		AssetEditorSubsystem->OpenEditorForAsset(Material);
-		IMaterialEditor* MaterialEditor = (IMaterialEditor*)AssetEditorSubsystem->FindEditorForAsset(Material, true);
-		MaterialEditor->JumpToExpression(Expression);
+
+		IAssetEditorInstance* AssetEditorInstance = AssetEditorSubsystem->OpenEditorForAsset(Asset) ? AssetEditorSubsystem->FindEditorForAsset(Asset, true) : NULL;
+		if (AssetEditorInstance != NULL)
+		{
+			if (AssetEditorInstance->GetEditorName() == "MaterialEditor")
+			{
+				((IMaterialEditor*)AssetEditorInstance)->JumpToExpression(Expression);
+			}
+			else
+			{
+				ensureMsgf(false, TEXT("Missing navigate to expression for editor '%s'"), *AssetEditorInstance->GetEditorName().ToString());
+			}
+		}
 	}
 }
 
