Fix sliders not working in any component property

diff --git a/Engine/Source/Editor/DetailCustomizations/Private/CustomPrimitiveDataCustomization.cpp b/Engine/Source/Editor/DetailCustomizations/Private/CustomPrimitiveDataCustomization.cpp
--- a/Engine/Source/Editor/DetailCustomizations/Private/CustomPrimitiveDataCustomization.cpp
+++ b/Engine/Source/Editor/DetailCustomizations/Private/CustomPrimitiveDataCustomization.cpp
@@ -556,11 +556,14 @@ void FCustomPrimitiveDataCustomization::OnUpdated()
 	
 void FCustomPrimitiveDataCustomization::OnObjectPropertyChanged(UObject* Object, FPropertyChangedEvent& PropertyChangedEvent)
 {
-	const bool bIsCustomPrimitiveDataProperty = PropertyChangedEvent.GetPropertyName() == "CustomPrimitiveData"
-		|| (PropertyChangedEvent.MemberProperty != NULL && PropertyChangedEvent.MemberProperty->GetFName() == "CustomPrimitiveData");
-	if (ComponentsToWatch.Contains(Object) && !bIsCustomPrimitiveDataProperty)
+	if(PropertyChangedEvent.ChangeType != EPropertyChangeType::Interactive )
 	{
-		OnUpdated();
+		const bool bIsCustomPrimitiveDataProperty = PropertyChangedEvent.GetPropertyName() == "CustomPrimitiveData"
+			|| (PropertyChangedEvent.MemberProperty != NULL && PropertyChangedEvent.MemberProperty->GetFName() == "CustomPrimitiveData");
+		if (ComponentsToWatch.Contains(Object) && !bIsCustomPrimitiveDataProperty)
+		{
+			OnUpdated();
+		}
 	}
 }
 
diff --git a/Engine/Source/Editor/DetailCustomizations/Private/CustomPrimitiveDataCustomization.h b/Engine/Source/Editor/DetailCustomizations/Private/CustomPrimitiveDataCustomization.h
--- a/Engine/Source/Editor/DetailCustomizations/Private/CustomPrimitiveDataCustomization.h
+++ b/Engine/Source/Editor/DetailCustomizations/Private/CustomPrimitiveDataCustomization.h
@@ -5,6 +5,7 @@
 #include "CoreMinimal.h"
 #include "IPropertyTypeCustomization.h"
 #include "IDetailChildrenBuilder.h"
+#include "Materials/MaterialLayersFunctions.h"
 
 class UMaterialExpressionVectorParameter;
 class UMaterialInterface;
