Update view only when material has been changed for CPD purposes, prevents list refresh on changes to other variables

diff --git a/Engine/Source/Editor/DetailCustomizations/Private/CustomPrimitiveDataCustomization.cpp b/Engine/Source/Editor/DetailCustomizations/Private/CustomPrimitiveDataCustomization.cpp
--- a/Engine/Source/Editor/DetailCustomizations/Private/CustomPrimitiveDataCustomization.cpp
+++ b/Engine/Source/Editor/DetailCustomizations/Private/CustomPrimitiveDataCustomization.cpp
@@ -22,6 +22,10 @@
 #include "Materials/MaterialExpressionVectorParameter.h"
 #include "IMaterialEditor.h"
 
+#include "Components/PrimitiveComponent.h"
+#include "Components/MeshComponent.h"
+#include "Components/TextRenderComponent.h"
+
 #define LOCTEXT_NAMESPACE "CustomPrimitiveDataCustomization"
 
 TSharedRef<IPropertyTypeCustomization> FCustomPrimitiveDataCustomization::MakeInstance()
@@ -65,7 +69,6 @@ void FCustomPrimitiveDataCustomization::CustomizeChildren(TSharedRef<IPropertyHa
 	{
 		PopulateParameterData(Component, MaxPrimitiveDataIndex);
 		NumSelectedComponents++;
-		ComponentsToWatch.Add(Component);
 	});
 
 	uint32 NumElements;
@@ -76,6 +79,7 @@ void FCustomPrimitiveDataCustomization::CustomizeChildren(TSharedRef<IPropertyHa
 	DataProperty->SetOnPropertyValueChanged(OnElemsChanged);
 	DataArrayHandle->SetOnNumElementsChanged(OnElemsChanged);
 
+	// NOTE: Optimally would be bound to a "OnMaterialChanged" for each component
 	FCoreUObjectDelegates::OnObjectPropertyChanged.AddRaw(this, &FCustomPrimitiveDataCustomization::OnObjectPropertyChanged);
 	UMaterial::OnMaterialCompilationFinished().AddRaw(this, &FCustomPrimitiveDataCustomization::OnMaterialCompiled);
 
@@ -461,6 +465,8 @@ void FCustomPrimitiveDataCustomization::PopulateParameterData(UPrimitiveComponen
 {
 	const int32 NumMaterials = PrimitiveComponent->GetNumMaterials();
 
+	TSet<TSoftObjectPtr<UMaterial>>& CachedComponentMaterials = ComponentsToWatch.FindOrAdd(PrimitiveComponent);
+
 	for (int32 i = 0; i < NumMaterials; ++i)
 	{
 		UMaterialInterface* MaterialInterface = PrimitiveComponent->GetMaterial(i);
@@ -472,6 +478,7 @@ void FCustomPrimitiveDataCustomization::PopulateParameterData(UPrimitiveComponen
 		}
 
 		MaterialsToWatch.Add(Material);
+		CachedComponentMaterials.Add(Material);
 
 		TMap<FMaterialParameterInfo, FMaterialParameterMetadata> Parameters;
 
@@ -556,11 +563,55 @@ void FCustomPrimitiveDataCustomization::OnUpdated()
 	
 void FCustomPrimitiveDataCustomization::OnObjectPropertyChanged(UObject* Object, FPropertyChangedEvent& PropertyChangedEvent)
 {
-	if(PropertyChangedEvent.ChangeType != EPropertyChangeType::Interactive )
+	const EPropertyChangeType::Type IgnoreFlags = EPropertyChangeType::Interactive | EPropertyChangeType::Redirected;
+
+	if (!(PropertyChangedEvent.ChangeType & IgnoreFlags) && ComponentsToWatch.Contains(Object))
 	{
-		const bool bIsCustomPrimitiveDataProperty = PropertyChangedEvent.GetPropertyName() == "CustomPrimitiveData"
-			|| (PropertyChangedEvent.MemberProperty != NULL && PropertyChangedEvent.MemberProperty->GetFName() == "CustomPrimitiveData");
-		if (ComponentsToWatch.Contains(Object) && !bIsCustomPrimitiveDataProperty)
+		bool bMaterialChange = false;
+
+		if (Object->IsA<UMeshComponent>())
+		{
+			bMaterialChange = PropertyChangedEvent.GetPropertyName() == GET_MEMBER_NAME_CHECKED(UMeshComponent, OverrideMaterials);
+		}
+		else if (Object->IsA<UTextRenderComponent>())
+		{
+			bMaterialChange = PropertyChangedEvent.GetPropertyName() == GET_MEMBER_NAME_CHECKED(UTextRenderComponent, TextMaterial);
+		}
+		else
+		{
+			// Fall back if not handled
+			// NOTE: Optimally would be done via an "OnMaterialChanged" for each component, however,
+			// the property name checks above should handle most cases
+			TSet<TSoftObjectPtr<UMaterial>>& CachedComponentMaterials = ComponentsToWatch[Object];
+
+			const UPrimitiveComponent* PrimComponent = CastChecked<UPrimitiveComponent>(Object);
+			const int32 NumMaterials = PrimComponent->GetNumMaterials();
+
+			if (NumMaterials != CachedComponentMaterials.Num())
+			{
+				bMaterialChange = true;
+			}
+			else
+			{
+				TSet<TSoftObjectPtr<UMaterial>> CurrentMaterials;
+				CurrentMaterials.Reserve(NumMaterials);
+
+				for (int32 i = 0; i < NumMaterials; ++i)
+				{
+					UMaterialInterface* MaterialInterface = PrimComponent->GetMaterial(i);
+					UMaterial* Material = MaterialInterface ? MaterialInterface->GetMaterial() : NULL;
+
+					if (Material)
+					{
+						CurrentMaterials.Add(Material);
+					}
+				}
+
+				bMaterialChange = CurrentMaterials.Difference(CachedComponentMaterials).Num() > 0;
+			}
+		}
+
+		if (bMaterialChange)
 		{
 			OnUpdated();
 		}
diff --git a/Engine/Source/Editor/DetailCustomizations/Private/CustomPrimitiveDataCustomization.h b/Engine/Source/Editor/DetailCustomizations/Private/CustomPrimitiveDataCustomization.h
--- a/Engine/Source/Editor/DetailCustomizations/Private/CustomPrimitiveDataCustomization.h
+++ b/Engine/Source/Editor/DetailCustomizations/Private/CustomPrimitiveDataCustomization.h
@@ -40,7 +40,7 @@ private:
 	TSharedPtr<IPropertyUtilities> PropertyUtils;
 	TSharedPtr<IPropertyHandleArray> DataArrayHandle;
 
-	TSet<TWeakObjectPtr<UObject>> ComponentsToWatch;
+	TMap<TWeakObjectPtr<UObject>, TSet<TSoftObjectPtr<UMaterial>>> ComponentsToWatch;
 	TSet<TSoftObjectPtr<UMaterialInterface>> MaterialsToWatch;
 
 	TMap<uint8, TArray<FParameterData>> VectorParameterData;
