Added primitive data indices to rows of CPD customisation
Added tooltips to hyperlinks
Cleaned up layout to not use border
Fixed crash due to accessed vector when supposed to be scalar

diff --git a/Engine/Source/Editor/DetailCustomizations/Private/CustomPrimitiveDataCustomization.cpp b/Engine/Source/Editor/DetailCustomizations/Private/CustomPrimitiveDataCustomization.cpp
--- a/Engine/Source/Editor/DetailCustomizations/Private/CustomPrimitiveDataCustomization.cpp
+++ b/Engine/Source/Editor/DetailCustomizations/Private/CustomPrimitiveDataCustomization.cpp
@@ -133,8 +133,9 @@ void FCustomPrimitiveDataCustomization::CustomizeChildren(TSharedRef<IPropertyHa
 		}
 		else
 		{
-			// We've encountered a gap in defined custom primitive data, mark it undefined
-			TSharedRef<SWidget> NameWidget = GetUndefinedParameterWidget(PrimIdx, CustomizationUtils);
+			// We've encountered a gap in declared custom primitive data, mark it undeclared
+			TSharedRef<SWidget> UndeclaredWidget = GetUndeclaredParameterWidget(PrimIdx, CustomizationUtils);
+			TSharedRef<SWidget> NameWidget = CreateNameWidget(PrimIdx, UndeclaredWidget, CustomizationUtils);
 
 			if (ElementHandle.IsValid())
 			{
@@ -183,34 +184,20 @@ IDetailGroup* FCustomPrimitiveDataCustomization::CreateVectorGroup(IDetailChildr
 
 		AddedParametersForThisGroup.Add(ParameterData.ExpressionID);
 
-		VectorGroupNameBox->AddSlot()
-		.Padding(2.f)
-		[
-			CreateHyperlink(FText::FromName(ParameterData.Info.Name), ParameterData.Material, ParameterData.ExpressionID)
-		];
-	}
+		TSharedRef<SHyperlink> Hyperlink = CreateHyperlink(FText::FromName(ParameterData.Info.Name), ParameterData.Material, ParameterData.ExpressionID);
+		Hyperlink->SetToolTipText(LOCTEXT("VectorHyperlinkTooltip", "Jump to Vector Parameter"));
 
-	TSharedPtr<SWidget> VectorGroupNameWidget;
-	if (VectorGroupNameBox->NumSlots() > 1)
-	{
-		// We have multiple overlapping parameter names, make sure to put a border around it to contain it
-		// TODO: Maybe a UI designer needs to take a look at this?
-		VectorGroupNameWidget = SNew(SBorder)
-		.BorderImage(FEditorStyle::Get().GetBrush("FilledBorder"))
+		VectorGroupNameBox->AddSlot()
 		.Padding(2.f)
 		[
-			VectorGroupNameBox.ToSharedRef()
+			Hyperlink
 		];
 	}
-	else
-	{
-		VectorGroupNameWidget = VectorGroupNameBox;
-	}
 
 	VectorGroup->HeaderRow()
 	.NameContent()
 	[
-		VectorGroupNameWidget.ToSharedRef()
+		VectorGroupNameBox.ToSharedRef()
 	]
 	.ValueContent()
 	[
@@ -313,10 +300,13 @@ void FCustomPrimitiveDataCustomization::CreateParameterRow(IDetailChildrenBuilde
 				ChannelName
 			);
 
+			TSharedRef<SHyperlink> Hyperlink = CreateHyperlink(ParameterName, ParameterData.Material, ParameterData.ExpressionID);
+			Hyperlink->SetToolTipText(LOCTEXT("VectorChannelHyperlinkTooltip", "Jump to Vector Parameter Channel"));
+
 			VerticalBox->AddSlot()
 			.Padding(2.f)
 			[
-				CreateHyperlink(ParameterName, ParameterData.Material, ParameterData.ExpressionID)
+				Hyperlink
 			];
 
 			SearchText.Add(ParameterName);
@@ -338,10 +328,13 @@ void FCustomPrimitiveDataCustomization::CreateParameterRow(IDetailChildrenBuilde
 
 			const FText ParameterName = FText::FromName(ParameterData.Info.Name);
 
+			TSharedRef<SHyperlink> Hyperlink = CreateHyperlink(ParameterName, ParameterData.Material, ParameterData.ExpressionID);
+			Hyperlink->SetToolTipText(LOCTEXT("ScalarHyperlinkTooltip", "Jump to Scalar Parameter"));
+
 			VerticalBox->AddSlot()
 			.Padding(2.f)
 			[
-				CreateHyperlink(ParameterName, ParameterData.Material, ParameterData.ExpressionID)
+				Hyperlink
 			];
 
 			SearchText.Add(ParameterName);
@@ -350,30 +343,13 @@ void FCustomPrimitiveDataCustomization::CreateParameterRow(IDetailChildrenBuilde
 
 	if (Components.Num() != NumSelectedComponents)
 	{
-		// Some components aren't defining parameters at this index, add the undefined parameter widget in case this was user error
+		// Some components aren't defining parameters at this index, add the undeclared parameter widget in case this was user error
 		VerticalBox->AddSlot()
 		.Padding(2.f)
 		[
-			GetUndefinedParameterWidget(PrimIdx, CustomizationUtils)
-		];
-	}
-
-	TSharedPtr<SWidget> NameWidget;
-	if (VerticalBox->NumSlots() > 1)
-	{
-		// We have multiple overlapping parameter names, make sure to put a border around it to contain it
-		// TODO: Maybe a UI designer needs to take a look at this?
-		NameWidget = SNew(SBorder)
-		.BorderImage(FEditorStyle::Get().GetBrush("FilledBorder"))
-		.Padding(2.f)
-		[
-			VerticalBox.ToSharedRef()
+			GetUndeclaredParameterWidget(PrimIdx, CustomizationUtils)
 		];
 	}
-	else
-	{
-		NameWidget = VerticalBox;
-	}
 
 	if (ElementHandle.IsValid())
 	{
@@ -388,7 +364,7 @@ void FCustomPrimitiveDataCustomization::CreateParameterRow(IDetailChildrenBuilde
 		Row.CustomWidget()
 		.NameContent()
 		[
-			NameWidget.ToSharedRef()
+			CreateNameWidget(PrimIdx, VerticalBox.ToSharedRef(), CustomizationUtils)
 		]
 		.ValueContent()
 		[
@@ -402,7 +378,7 @@ void FCustomPrimitiveDataCustomization::CreateParameterRow(IDetailChildrenBuilde
 
 		Row.NameContent()
 		[
-			NameWidget.ToSharedRef()
+			CreateNameWidget(PrimIdx, VerticalBox.ToSharedRef(), CustomizationUtils)
 		]
 		.ValueContent()
 		[
@@ -751,7 +727,7 @@ void FCustomPrimitiveDataCustomization::SetDefaultValue(TSharedPtr<IPropertyHand
 
 		if (ScalarParameterData.Contains(PrimIdx))
 		{
-			for (const FParameterData& ParameterData : VectorParameterData[PrimIdx])
+			for (const FParameterData& ParameterData : ScalarParameterData[PrimIdx])
 			{
 				if (ParameterData.Component.IsValid() && !ChangedComponents.Contains(ParameterData.Component))
 				{
@@ -823,6 +799,28 @@ void FCustomPrimitiveDataCustomization::OnColorPickerWindowClosed(const TSharedR
 	GEditor->EndTransaction();
 }
 
+BEGIN_SLATE_FUNCTION_BUILD_OPTIMIZATION
+TSharedRef<SWidget> FCustomPrimitiveDataCustomization::CreateNameWidget(int32 PrimIdx, TSharedRef<SWidget> ParameterName, IPropertyTypeCustomizationUtils& CustomizationUtils) const
+{
+	return SNew(SHorizontalBox)
+		+ SHorizontalBox::Slot()
+		.VAlign(VAlign_Center)
+		.AutoWidth()
+		.Padding(2.f, 2.f, 16.0f, 2.f)
+		[
+			SNew(STextBlock)
+			.Text(FText::AsNumber(PrimIdx))
+			.Font(CustomizationUtils.GetRegularFont())
+		]
+		+ SHorizontalBox::Slot()
+		.AutoWidth()
+		.Padding(0.f, 2.0f)
+		[
+			ParameterName
+		];
+}
+END_SLATE_FUNCTION_BUILD_OPTIMIZATION
+
 BEGIN_SLATE_FUNCTION_BUILD_OPTIMIZATION
 TSharedRef<SHyperlink> FCustomPrimitiveDataCustomization::CreateHyperlink(FText Text, TWeakObjectPtr<UMaterialInterface> Material, const FGuid& ExpressionID)
 {
@@ -835,15 +833,15 @@ TSharedRef<SHyperlink> FCustomPrimitiveDataCustomization::CreateHyperlink(FText
 END_SLATE_FUNCTION_BUILD_OPTIMIZATION
 
 BEGIN_SLATE_FUNCTION_BUILD_OPTIMIZATION
-TSharedRef<SWidget> FCustomPrimitiveDataCustomization::GetUndefinedParameterWidget(int32 PrimIdx, IPropertyTypeCustomizationUtils& CustomizationUtils) const
+TSharedRef<SWidget> FCustomPrimitiveDataCustomization::GetUndeclaredParameterWidget(int32 PrimIdx, IPropertyTypeCustomizationUtils& CustomizationUtils) const
 {
 	FText PrimIdxText = FText::AsNumber(PrimIdx);
-	TSharedRef<SWidget> UndefinedParamWidget = SNew(STextBlock)
-		.Text(FText::Format(LOCTEXT("UndefinedParameter", "{0} (Undefined)"), PrimIdxText))
+	TSharedRef<SWidget> UndeclaredParamWidget = SNew(STextBlock)
+		.Text(LOCTEXT("UndeclaredParameter", "(Undeclared)"))
 		.Font(CustomizationUtils.GetRegularFont());
-	UndefinedParamWidget->SetToolTipText(FText::Format(LOCTEXT("UndefinedParameterTooltip", "A component is selected that doesn't define a parameter for primitive index {0}"), PrimIdxText));
+	UndeclaredParamWidget->SetToolTipText(FText::Format(LOCTEXT("UndeclaredParameterTooltip", "A component is selected that doesn't declare a parameter for primitive index {0}"), PrimIdxText));
 		
-	return UndefinedParamWidget;
+	return UndeclaredParamWidget;
 }
 END_SLATE_FUNCTION_BUILD_OPTIMIZATION
 
diff --git a/Engine/Source/Editor/DetailCustomizations/Private/CustomPrimitiveDataCustomization.h b/Engine/Source/Editor/DetailCustomizations/Private/CustomPrimitiveDataCustomization.h
--- a/Engine/Source/Editor/DetailCustomizations/Private/CustomPrimitiveDataCustomization.h
+++ b/Engine/Source/Editor/DetailCustomizations/Private/CustomPrimitiveDataCustomization.h
@@ -11,6 +11,7 @@ class UMaterialExpressionVectorParameter;
 class UMaterialInterface;
 
 class FDetailWidgetRow;
+class SWidget;
 class SColorBlock;
 
 class IPropertyHandleArray;
@@ -97,10 +98,12 @@ private:
 	void OnColorPickerCancelled(FLinearColor OriginalColor, uint8 PrimIdx);
 	void OnColorPickerWindowClosed(const TSharedRef<SWindow>& Window);
 
+	TSharedRef<SWidget> CreateNameWidget(int32 PrimIdx, TSharedRef<SWidget> ParameterName, IPropertyTypeCustomizationUtils& CustomizationUtils) const;
+
 	/** Hyperlink that's used for navigating to parameters */
 	TSharedRef<class SHyperlink> CreateHyperlink(FText Text, TWeakObjectPtr<UMaterialInterface> Material, const FGuid& ExpressionId);
 
-	/** Default widget that's used when a component is missing defined parameters for the specified primitive index */
-	TSharedRef<class SWidget> GetUndefinedParameterWidget(int32 PrimIdx, IPropertyTypeCustomizationUtils& CustomizationUtils) const;
+	/** Default widget that's used when a component is missing declared parameters for the specified primitive index */
+	TSharedRef<SWidget> GetUndeclaredParameterWidget(int32 PrimIdx, IPropertyTypeCustomizationUtils& CustomizationUtils) const;
 };
 
