SceneCapture: Expose path tracing show flag.

diff --git a/Engine/Source/Editor/DetailCustomizations/Private/SceneCaptureDetails.cpp b/Engine/Source/Editor/DetailCustomizations/Private/SceneCaptureDetails.cpp
--- a/Engine/Source/Editor/DetailCustomizations/Private/SceneCaptureDetails.cpp
+++ b/Engine/Source/Editor/DetailCustomizations/Private/SceneCaptureDetails.cpp
@@ -104,6 +104,7 @@ void FSceneCaptureDetails::CustomizeDetails( IDetailLayoutBuilder& DetailLayout
 	ShowFlagsToAllowForCaptures.Add(FEngineShowFlags::EShowFlag::SF_Game);
 	ShowFlagsToAllowForCaptures.Add(FEngineShowFlags::EShowFlag::SF_ToneCurve); 
 	ShowFlagsToAllowForCaptures.Add(FEngineShowFlags::EShowFlag::SF_SingleLayerWaterRefractionFullPrecision);
+	ShowFlagsToAllowForCaptures.Add(FEngineShowFlags::EShowFlag::SF_PathTracing);
 
 	// Create array of flag name strings for each group
 	TArray< TArray<FString> > ShowFlagsByGroup;
