Removed redundant force refresh calls; this removes small hitch when adding/removing many elements
CPD Set to Defaults now can pull from overriden parameter values (not just material defaults)
Ensure get num elems also takes into consideration multi-select situation, as array property handle always returns success

diff --git a/Engine/Source/Editor/DetailCustomizations/Private/CustomPrimitiveDataCustomization.cpp b/Engine/Source/Editor/DetailCustomizations/Private/CustomPrimitiveDataCustomization.cpp
--- a/Engine/Source/Editor/DetailCustomizations/Private/CustomPrimitiveDataCustomization.cpp
+++ b/Engine/Source/Editor/DetailCustomizations/Private/CustomPrimitiveDataCustomization.cpp
@@ -59,8 +59,8 @@ void FCustomPrimitiveDataCustomization::CustomizeChildren(TSharedRef<IPropertyHa
 
 	PropertyUtils = CustomizationUtils.GetPropertyUtilities();
 
-	TSharedPtr<IPropertyHandle> DataProperty = PropertyHandle->GetChildHandle("Data");
-	DataArrayHandle = DataProperty->AsArray();
+	DataHandle = PropertyHandle->GetChildHandle("Data");
+	DataArrayHandle = DataHandle->AsArray();
 
 	int32 NumSelectedComponents = 0;
 	int32 MaxPrimitiveDataIndex = INDEX_NONE;
@@ -72,12 +72,11 @@ void FCustomPrimitiveDataCustomization::CustomizeChildren(TSharedRef<IPropertyHa
 	});
 
 	uint32 NumElements;
-	DataArrayHandle->GetNumElements(NumElements);
+	FPropertyAccess::Result AccessResult = GetNumElements(NumElements);
 
 	FSimpleDelegate OnElemsChanged = FSimpleDelegate::CreateRaw(this, &FCustomPrimitiveDataCustomization::OnUpdated);
-	PropertyHandle->SetOnPropertyValueChanged(OnElemsChanged);
-	DataProperty->SetOnPropertyValueChanged(OnElemsChanged);
 	DataArrayHandle->SetOnNumElementsChanged(OnElemsChanged);
+	DataHandle->SetOnPropertyValueChanged(FSimpleDelegate::CreateRaw(this, &FCustomPrimitiveDataCustomization::OnElementsModified, AccessResult, NumElements));
 
 	// NOTE: Optimally would be bound to a "OnMaterialChanged" for each component
 	FCoreUObjectDelegates::OnObjectPropertyChanged.AddRaw(this, &FCustomPrimitiveDataCustomization::OnObjectPropertyChanged);
@@ -90,8 +89,8 @@ void FCustomPrimitiveDataCustomization::CustomizeChildren(TSharedRef<IPropertyHa
 		return;
 	}
 
-	FString ArrayString;
-	const bool bDataEditable = DataProperty.IsValid() && DataProperty->IsEditable() && DataProperty->GetValueAsDisplayString(ArrayString) == FPropertyAccess::Success;
+	// We're only editable if the property is editable and if we're not a multi-selection situation
+	const bool bDataEditable = DataHandle.IsValid() && DataHandle->IsEditable() && AccessResult == FPropertyAccess::Success;
 
 	uint8 VectorGroupPrimIdx = 0;
 	IDetailGroup* VectorGroup = NULL;
@@ -428,6 +427,7 @@ void FCustomPrimitiveDataCustomization::Cleanup()
 	UMaterial::OnMaterialCompilationFinished().RemoveAll(this);
 
 	PropertyUtils = NULL;
+	DataHandle = NULL;
 	DataArrayHandle = NULL;
 
 	ComponentsToWatch.Empty();
@@ -531,12 +531,24 @@ void FCustomPrimitiveDataCustomization::PopulateParameterData(UPrimitiveComponen
 
 void FCustomPrimitiveDataCustomization::OnUpdated()
 {
-	if (PropertyUtils.IsValid())
+	if (!bIgnoreUpdate && PropertyUtils.IsValid())
 	{
 		PropertyUtils->ForceRefresh();
 	}
 }
 	
+void FCustomPrimitiveDataCustomization::OnElementsModified(const FPropertyAccess::Result OldAccessResult, const uint32 OldNumElements)
+{
+	uint32 NumElements;
+	FPropertyAccess::Result AccessResult = GetNumElements(NumElements);
+
+	// There's been a change in our array structure, whether that be from change in access or size
+	if (AccessResult != OldAccessResult || NumElements != OldNumElements)
+	{
+		OnUpdated();
+	}
+}
+
 void FCustomPrimitiveDataCustomization::OnObjectPropertyChanged(UObject* Object, FPropertyChangedEvent& PropertyChangedEvent)
 {
 	const EPropertyChangeType::Type IgnoreFlags = EPropertyChangeType::Interactive | EPropertyChangeType::Redirected;
@@ -631,10 +643,12 @@ void FCustomPrimitiveDataCustomization::OnNavigate(TWeakObjectPtr<UMaterialInter
 void FCustomPrimitiveDataCustomization::OnAddedDesiredPrimitiveData(uint8 PrimIdx)
 {
 	uint32 NumElements;
-	if (DataArrayHandle.IsValid() && DataArrayHandle->GetNumElements(NumElements) == FPropertyAccess::Success)
+	if (GetNumElements(NumElements) == FPropertyAccess::Success && PrimIdx >= NumElements)
 	{
 		GEditor->BeginTransaction(LOCTEXT("OnAddedDesiredPrimitiveData", "Added Items"));
 
+		bIgnoreUpdate = true;
+
 		for (int32 i = NumElements; i <= PrimIdx; ++i)
 		{
 			DataArrayHandle->AddItem();
@@ -642,6 +656,10 @@ void FCustomPrimitiveDataCustomization::OnAddedDesiredPrimitiveData(uint8 PrimId
 			SetDefaultValue(DataArrayHandle->GetElement(i), i);
 		}
 
+		bIgnoreUpdate = false;
+
+		OnUpdated();
+
 		GEditor->EndTransaction();
 	}
 }
@@ -649,15 +667,21 @@ void FCustomPrimitiveDataCustomization::OnAddedDesiredPrimitiveData(uint8 PrimId
 void FCustomPrimitiveDataCustomization::OnRemovedPrimitiveData(uint8 PrimIdx)
 {
 	uint32 NumElements;
-	if (DataArrayHandle.IsValid() && DataArrayHandle->GetNumElements(NumElements) == FPropertyAccess::Success)
+	if (GetNumElements(NumElements) == FPropertyAccess::Success && PrimIdx < NumElements)
 	{
 		GEditor->BeginTransaction(LOCTEXT("OnRemovedPrimitiveData", "Removed Items"));
 
+		bIgnoreUpdate = true;
+
 		for (int32 i = NumElements - 1; i >= PrimIdx; --i)
 		{
 			DataArrayHandle->DeleteItem(i);
 		}
 
+		bIgnoreUpdate = false;
+
+		OnUpdated();
+
 		GEditor->EndTransaction();
 	}
 }
@@ -667,7 +691,7 @@ FLinearColor FCustomPrimitiveDataCustomization::GetVectorColor(uint8 PrimIdx) co
 	FVector4f Color(ForceInitToZero);
 
 	uint32 NumElems;
-	if (DataArrayHandle.IsValid() && DataArrayHandle->GetNumElements(NumElems) == FPropertyAccess::Success)
+	if (GetNumElements(NumElems) == FPropertyAccess::Success)
 	{
 		const int32 MaxElems = FMath::Min((int32)NumElems, PrimIdx + 4);
 
@@ -685,7 +709,7 @@ void FCustomPrimitiveDataCustomization::SetVectorColor(FLinearColor NewColor, ui
 	FVector4f Color(NewColor);
 
 	uint32 NumElems;
-	if (DataArrayHandle.IsValid() && DataArrayHandle->GetNumElements(NumElems) == FPropertyAccess::Success)
+	if (GetNumElements(NumElems) == FPropertyAccess::Success)
 	{
 		const int32 MaxElems = FMath::Min((int32)NumElems, PrimIdx + 4);
 
@@ -710,7 +734,7 @@ void FCustomPrimitiveDataCustomization::SetDefaultValue(TSharedPtr<IPropertyHand
 				if (ParameterData.Component.IsValid() && !ChangedComponents.Contains(ParameterData.Component))
 				{
 					FLinearColor Color(ForceInitToZero);
-					if (!ParameterData.Material.IsValid() || ParameterData.Material->GetVectorParameterDefaultValue(ParameterData.Info, Color))
+					if (!ParameterData.Material.IsValid() || ParameterData.Material->GetVectorParameterValue(ParameterData.Info, Color))
 					{
 						float* ColorPtr = reinterpret_cast<float*>(&Color);
 						ParameterData.Component->SetDefaultCustomPrimitiveDataFloat(PrimIdx, ColorPtr[ParameterData.IndexOffset]);
@@ -732,7 +756,7 @@ void FCustomPrimitiveDataCustomization::SetDefaultValue(TSharedPtr<IPropertyHand
 				if (ParameterData.Component.IsValid() && !ChangedComponents.Contains(ParameterData.Component))
 				{
 					float Value = 0.f;
-					if (!ParameterData.Material.IsValid() || ParameterData.Material->GetScalarParameterDefaultValue(ParameterData.Info, Value))
+					if (!ParameterData.Material.IsValid() || ParameterData.Material->GetScalarParameterValue(ParameterData.Info, Value))
 					{
 						ParameterData.Component->SetDefaultCustomPrimitiveDataFloat(PrimIdx, Value);
 
@@ -751,7 +775,7 @@ void FCustomPrimitiveDataCustomization::SetDefaultValue(TSharedPtr<IPropertyHand
 void FCustomPrimitiveDataCustomization::SetDefaultVectorValue(uint8 PrimIdx)
 {
 	uint32 NumElems;
-	if (DataArrayHandle.IsValid() && DataArrayHandle->GetNumElements(NumElems) == FPropertyAccess::Success)
+	if (GetNumElements(NumElems) == FPropertyAccess::Success)
 	{
 		GEditor->BeginTransaction(LOCTEXT("SetDefaultVectorValue", "Reset Vector To Default"));
 
@@ -845,4 +869,20 @@ TSharedRef<SWidget> FCustomPrimitiveDataCustomization::GetUndeclaredParameterWid
 }
 END_SLATE_FUNCTION_BUILD_OPTIMIZATION
 
+FPropertyAccess::Result FCustomPrimitiveDataCustomization::GetNumElements(uint32& NumElements) const
+{
+	if (DataHandle.IsValid() && DataArrayHandle.IsValid())
+	{
+		DataArrayHandle->GetNumElements(NumElements);
+
+		// This is a low touch way to work out whether we have multiple selections or not,
+		// since FPropertyHandleArray::GetNumElements above always returns FPropertyAccess::Success
+		void* Address;
+		return DataHandle->GetValueData(Address);
+	}
+
+	NumElements = 0;
+	return FPropertyAccess::Fail;
+}
+
 #undef LOCTEXT_NAMESPACE
diff --git a/Engine/Source/Editor/DetailCustomizations/Private/CustomPrimitiveDataCustomization.h b/Engine/Source/Editor/DetailCustomizations/Private/CustomPrimitiveDataCustomization.h
--- a/Engine/Source/Editor/DetailCustomizations/Private/CustomPrimitiveDataCustomization.h
+++ b/Engine/Source/Editor/DetailCustomizations/Private/CustomPrimitiveDataCustomization.h
@@ -38,7 +38,10 @@ private:
 		uint8 IndexOffset;
 	};
 
+	bool bIgnoreUpdate = false;
+
 	TSharedPtr<IPropertyUtilities> PropertyUtils;
+	TSharedPtr<IPropertyHandle> DataHandle;
 	TSharedPtr<IPropertyHandleArray> DataArrayHandle;
 
 	TMap<TWeakObjectPtr<UObject>, TSet<TSoftObjectPtr<UMaterial>>> ComponentsToWatch;
@@ -73,6 +76,7 @@ private:
 	void PopulateParameterData(UPrimitiveComponent* PrimitiveComponent, int32& MaxPrimitiveDataIndex);
 
 	void OnUpdated();
+	void OnElementsModified(const enum FPropertyAccess::Result OldAccessResult, const uint32 OldNumElements);
 	void OnObjectPropertyChanged(UObject* Object, FPropertyChangedEvent& PropertyChangedEvent);
 	void OnMaterialCompiled(UMaterialInterface* Material);
 
@@ -105,5 +109,8 @@ private:
 
 	/** Default widget that's used when a component is missing declared parameters for the specified primitive index */
 	TSharedRef<SWidget> GetUndeclaredParameterWidget(int32 PrimIdx, IPropertyTypeCustomizationUtils& CustomizationUtils) const;
+
+	/** Get the number of elements and the access result of our primitive data array */
+	enum FPropertyAccess::Result GetNumElements(uint32& NumElements) const;
 };
 
