Defined section mappings for static meshes, lights, triggers and most of characters. These will also be visible in other places because the sections are inherited, but those are the ones I've focused on in this pass.

diff --git a/Engine/Source/Editor/DetailCustomizations/Private/DetailCustomizations.cpp b/Engine/Source/Editor/DetailCustomizations/Private/DetailCustomizations.cpp
--- a/Engine/Source/Editor/DetailCustomizations/Private/DetailCustomizations.cpp
+++ b/Engine/Source/Editor/DetailCustomizations/Private/DetailCustomizations.cpp
@@ -156,6 +156,7 @@ void FDetailCustomizationsModule::StartupModule()
 	
 	RegisterPropertyTypeCustomizations();
 	RegisterObjectCustomizations();
+	RegisterSectionMappings();
 
 	TSharedPtr<FMotionControllerPinFactory> MotionControllerPinFactory = MakeShareable(new FMotionControllerPinFactory());
 	FEdGraphUtilities::RegisterVisualPinFactory(MotionControllerPinFactory);
@@ -410,6 +411,304 @@ void FDetailCustomizationsModule::RegisterObjectCustomizations()
 	RegisterCustomClassLayout("BoundsCopyComponent", FOnGetDetailCustomizationInstance::CreateStatic(&FBoundsCopyComponentDetailsCustomization::MakeInstance));
 }
 
+#define LOCTEXT_NAMESPACE "DetailsSections"
+
+void FDetailCustomizationsModule::RegisterSectionMappings()
+{
+	static const FName PropertyEditor("PropertyEditor");
+	FPropertyEditorModule& PropertyModule = FModuleManager::GetModuleChecked<FPropertyEditorModule>(PropertyEditor);
+
+	// Actor
+	{
+		{
+			TSharedRef<FPropertySection> Section = PropertyModule.FindOrCreateSection("Actor", "Actor", LOCTEXT("Actor", "Actor"));
+			Section->AddCategory("Actor");
+		}
+
+		{
+			TSharedRef<FPropertySection> Section = PropertyModule.FindOrCreateSection("Actor", "Misc", LOCTEXT("Misc", "Misc"));
+			Section->AddCategory("Cooking");
+			Section->AddCategory("Input");
+			Section->AddCategory("Replication");
+			Section->AddCategory("World Partition");
+		}
+	}
+
+	// Pawn
+	{
+		{
+			TSharedRef<FPropertySection> Section = PropertyModule.FindOrCreateSection("Pawn", "General", LOCTEXT("General", "General"));
+			Section->AddCategory("Pawn");
+			Section->AddCategory("Camera");
+		}
+	}
+
+	// ActorComponent
+	{
+		{
+			TSharedRef<FPropertySection> Section = PropertyModule.FindOrCreateSection("ActorComponent", "Misc", LOCTEXT("Misc", "Misc"));
+			Section->AddCategory("Asset User Data");
+			Section->AddCategory("Cooking");
+			Section->AddCategory("Data Layers");
+			Section->AddCategory("Tags");
+			Section->AddCategory("World Partition");
+		}
+	}
+
+	// SceneComponent
+	{
+		{
+			TSharedRef<FPropertySection> Section = PropertyModule.FindOrCreateSection("SceneComponent", "General", LOCTEXT("General", "General"));
+			Section->AddCategory("Transform");
+			Section->AddCategory("TransformCommon");
+			Section->AddCategory("Mobility");
+		}
+	}
+
+	// PrimitiveComponent
+	{
+		{
+			TSharedRef<FPropertySection> Section = PropertyModule.FindOrCreateSection("PrimitiveComponent", "General", LOCTEXT("General", "General"));
+			Section->AddCategory("Materials");
+		}
+
+		{
+			TSharedRef<FPropertySection> Section = PropertyModule.FindOrCreateSection("PrimitiveComponent", "LOD", LOCTEXT("LOD", "LOD"));
+			Section->AddCategory("HLOD");
+			Section->AddCategory("LOD");
+		}
+
+		{
+			TSharedRef<FPropertySection> Section = PropertyModule.FindOrCreateSection("PrimitiveComponent", "Misc", LOCTEXT("Misc", "Misc"));
+			Section->AddCategory("Navigation");
+		}
+
+		{
+			TSharedRef<FPropertySection> Section = PropertyModule.FindOrCreateSection("PrimitiveComponent", "Physics", LOCTEXT("Physics", "Physics"));
+			Section->AddCategory("Collision");
+			Section->AddCategory("Physics");
+		}
+
+		{
+			TSharedRef<FPropertySection> Section = PropertyModule.FindOrCreateSection("PrimitiveComponent", "Rendering", LOCTEXT("Rendering", "Rendering"));
+			Section->AddCategory("Lighting");
+			Section->AddCategory("Lightmass");
+			Section->AddCategory("Materials");
+			Section->AddCategory("Mobile");
+			Section->AddCategory("Ray Tracing");
+			Section->AddCategory("Rendering");
+			Section->AddCategory("Texture Streaming");
+			Section->AddCategory("Virtual Texture");
+		}
+	}
+
+	// MeshComponent
+	{
+		{
+			TSharedRef<FPropertySection> Section = PropertyModule.FindOrCreateSection("MeshComponent", "General", LOCTEXT("General", "General"));
+			Section->AddCategory("Mesh");
+		}
+
+		{
+			TSharedRef<FPropertySection> Section = PropertyModule.FindOrCreateSection("MeshComponent", "Rendering", LOCTEXT("Rendering", "Rendering"));
+			Section->AddCategory("Material Parameters");
+		}
+	}
+
+	// StaticMeshComponent
+	{
+		{
+			TSharedRef<FPropertySection> Section = PropertyModule.FindOrCreateSection("StaticMeshComponent", "General", LOCTEXT("General", "General"));
+			Section->AddCategory("Static Mesh");
+		}
+
+		{
+			TSharedRef<FPropertySection> Section = PropertyModule.FindOrCreateSection("StaticMeshComponent", "Misc", LOCTEXT("Misc", "Misc"));
+			Section->AddCategory("Navigation");
+		}
+	}
+
+	// LightComponentBase
+	{
+		{
+			TSharedRef<FPropertySection> Section = PropertyModule.FindOrCreateSection("LightComponentBase", "General", LOCTEXT("General", "General"));
+			Section->AddCategory("Light");
+		}
+
+		{
+			TSharedRef<FPropertySection> Section = PropertyModule.FindOrCreateSection("LightComponentBase", "Rendering", LOCTEXT("Rendering", "Rendering"));
+			Section->AddCategory("Light");
+			Section->AddCategory("Light Function");
+			Section->AddCategory("Light Profiles");
+		}
+	}
+
+	// LightComponent
+	{
+		{
+			TSharedRef<FPropertySection> Section = PropertyModule.FindOrCreateSection("LightComponent", "Misc", LOCTEXT("Misc", "Misc"));
+			Section->AddCategory("Performance");
+		}
+
+		{
+			TSharedRef<FPropertySection> Section = PropertyModule.FindOrCreateSection("LightComponent", "Rendering", LOCTEXT("Rendering", "Rendering"));
+			Section->AddCategory("Distance Field Shadows");
+			Section->AddCategory("Light Shafts");
+		}
+	}
+
+	// DirectionalLightComponent
+	{
+		{
+			TSharedRef<FPropertySection> Section = PropertyModule.FindOrCreateSection("DirectionalLightComponent", "Rendering", LOCTEXT("Rendering", "Rendering"));
+			Section->AddCategory("Atmosphere and Cloud");
+			Section->AddCategory("Cascaded Shadow Maps");
+		}
+	}
+
+	// SkyLightComponent
+	{
+		{
+			TSharedRef<FPropertySection> Section = PropertyModule.FindOrCreateSection("SkyLightComponent", "General", LOCTEXT("General", "General"));
+			Section->AddCategory("Sky Light");
+		}
+
+		{
+			TSharedRef<FPropertySection> Section = PropertyModule.FindOrCreateSection("SkyLightComponent", "Rendering", LOCTEXT("Rendering", "Rendering"));
+			Section->AddCategory("Atmosphere and Cloud");
+			Section->AddCategory("Distance Field Ambient Occlusion");
+		}
+	}
+
+	// PlayerStart
+	{
+		{
+			TSharedRef<FPropertySection> Section = PropertyModule.FindOrCreateSection("PlayerStart", "General", LOCTEXT("General", "General"));
+			Section->AddCategory("Object");
+		}
+	}
+
+	// ShapeComponent
+	{
+		{
+			TSharedRef<FPropertySection> Section = PropertyModule.FindOrCreateSection("ShapeComponent", "General", LOCTEXT("General", "General"));
+			Section->AddCategory("Shape");
+		}
+	}
+
+	// BillboardComponent
+	{
+		{
+			TSharedRef<FPropertySection> Section = PropertyModule.FindOrCreateSection("BillboardComponent", "Rendering", LOCTEXT("Rendering", "Rendering"));
+			Section->AddCategory("Sprite");
+		}
+	}
+
+	// SkinnedMeshComponent
+	{
+		{
+			TSharedRef<FPropertySection> Section = PropertyModule.FindOrCreateSection("SkinnedMeshComponent", "General", LOCTEXT("General", "General"));
+			Section->AddCategory("Mesh");
+		}
+
+		{
+			TSharedRef<FPropertySection> Section = PropertyModule.FindOrCreateSection("SkinnedMeshComponent", "Animation", LOCTEXT("Animation", "Animation"));
+			Section->AddCategory("Skeletal Mesh");
+		}
+
+		{
+			TSharedRef<FPropertySection> Section = PropertyModule.FindOrCreateSection("SkinnedMeshComponent", "Misc", LOCTEXT("Misc", "Misc"));
+			Section->AddCategory("Optimization");
+		}
+	}
+
+	// SkeletalMeshComponent
+	{
+		{
+			TSharedRef<FPropertySection> Section = PropertyModule.FindOrCreateSection("SkeletalMeshComponent", "Animation", LOCTEXT("Animation", "Animation"));
+			Section->AddCategory("Animation");
+			Section->AddCategory("Animation Rig");
+			Section->AddCategory("Master Pose Component");
+		}
+
+		{
+			TSharedRef<FPropertySection> Section = PropertyModule.FindOrCreateSection("SkeletalMeshComponent", "Physics", LOCTEXT("Physics", "Physics"));
+			Section->AddCategory("Clothing");
+		}
+	}
+
+	// SkeletalMesh
+	{
+		TSharedRef<FPropertySection> Section = PropertyModule.FindOrCreateSection("SkeletalMesh", "Animation", LOCTEXT("Animation", "Animation"));
+		Section->AddCategory("Skin Weights");
+	}
+
+	// AnimInstance
+	{
+		{
+			TSharedRef<FPropertySection> Section = PropertyModule.FindOrCreateSection("AnimInstance", "Movement", LOCTEXT("Movement", "Movement"));
+			Section->AddCategory("Root Motion");
+		}
+
+		{
+			TSharedRef<FPropertySection> Section = PropertyModule.FindOrCreateSection("AnimInstance", "Physics", LOCTEXT("Physics", "Physics"));
+			Section->AddCategory("Root Motion");
+		}
+	}
+
+	// Character
+	{
+		{
+			TSharedRef<FPropertySection> Section = PropertyModule.FindOrCreateSection("Character", "Movement", LOCTEXT("Movement", "Movement"));
+			Section->AddCategory("Camera");
+			Section->AddCategory("Character");
+			Section->AddCategory("Pawn");
+		}
+	}
+
+	// CharacterMovementComponent
+	{
+		{
+			TSharedRef<FPropertySection> Section = PropertyModule.FindOrCreateSection("CharacterMovementComponent", "Movement", LOCTEXT("Movement", "Movement"));
+			Section->AddCategory("Character Movement");
+			Section->AddCategory("Character Movement (General Settings)");
+			Section->AddCategory("Character Movement (Networking)");
+			Section->AddCategory("Character Movement (Rotation Settings)");
+			Section->AddCategory("Character Movement: Avoidance");
+			Section->AddCategory("Character Movement: Custom Movement");
+			Section->AddCategory("Character Movement: Flying");
+			Section->AddCategory("Character Movement: Jumping / Falling");
+			Section->AddCategory("Character Movement: MovementMode");
+			Section->AddCategory("Character Movement: Physics Interaction");
+			Section->AddCategory("Character Movement: Swimming");
+			Section->AddCategory("Character Movement: Walking");
+			Section->AddCategory("Root Motion");
+		}
+
+		{
+			TSharedRef<FPropertySection> Section = PropertyModule.FindOrCreateSection("CharacterMovementComponent", "Physics", LOCTEXT("Physics", "Physics"));
+			Section->AddCategory("Root Motion");
+		}
+	}
+
+	// CapsuleComponent
+	{
+		{
+			TSharedRef<FPropertySection> Section = PropertyModule.FindOrCreateSection("CapsuleComponent", "Actor", LOCTEXT("Actor", "Actor"));
+			Section->AddCategory("Shape");
+		}
+	}
+
+	// Brush
+	{
+		{
+			TSharedRef<FPropertySection> Section = PropertyModule.FindOrCreateSection("Brush", "General", LOCTEXT("General", "General"));
+			Section->AddCategory("Brush Settings");
+		}
+	}
+}
+
+#undef LOCTEXT_NAMESPACE
 
 void FDetailCustomizationsModule::RegisterCustomClassLayout(FName ClassName, FOnGetDetailCustomizationInstance DetailLayoutDelegate )
 {
