The advanced sub-categories of Rendering defined in FPrimitiveComponentDetails have now been removed - Sections subsume this functionality, so this special case is unnecessary.

diff --git a/Engine/Source/Editor/DetailCustomizations/Private/PrimitiveComponentDetails.cpp b/Engine/Source/Editor/DetailCustomizations/Private/PrimitiveComponentDetails.cpp
--- a/Engine/Source/Editor/DetailCustomizations/Private/PrimitiveComponentDetails.cpp
+++ b/Engine/Source/Editor/DetailCustomizations/Private/PrimitiveComponentDetails.cpp
@@ -85,10 +85,6 @@ void FPrimitiveComponentDetails::CustomizeDetails( IDetailLayoutBuilder& DetailB
 	{
 		AddLightingCategory(DetailBuilder);
 	}
-
-	AddAdvancedSubCategory( DetailBuilder, "Rendering", "TextureStreaming" );
-	AddAdvancedSubCategory( DetailBuilder, "Rendering", "RayTracing");
-	AddAdvancedSubCategory(DetailBuilder, "Rendering", "LOD");
 }
 
 void FPrimitiveComponentDetails::AddMaterialCategory( IDetailLayoutBuilder& DetailBuilder )
