PR #7897: Hide all unused pins from Make structs node (Contributed by [REDACTED_PERSON])

diff --git a/Engine/Source/Editor/DetailCustomizations/Private/SkeletalControlNodeDetails.cpp b/Engine/Source/Editor/DetailCustomizations/Private/SkeletalControlNodeDetails.cpp
--- a/Engine/Source/Editor/DetailCustomizations/Private/SkeletalControlNodeDetails.cpp
+++ b/Engine/Source/Editor/DetailCustomizations/Private/SkeletalControlNodeDetails.cpp
@@ -15,6 +15,7 @@
 #include "K2Node.h"
 #include "Widgets/Input/SButton.h"
 #include "K2Node_BreakStruct.h"
+#include "K2Node_MakeStruct.h"
 #include "K2Node_GetClassDefaults.h"
 #include "ScopedTransaction.h"
 
@@ -35,12 +36,14 @@ void FSkeletalControlNodeDetails::CustomizeDetails(class IDetailLayoutBuilder& D
 	ArrayProperty = AvailablePins->AsArray();
 
 	bool bShowAvailablePins = true;
+	bool bHideInputPins = false;
 
 	const TArray< TWeakObjectPtr<UObject> >& SelectedObjects = DetailBuilder.GetSelectedObjects();
 	if (SelectedObjects.Num() == 1)
 	{
 		UObject* CurObj = SelectedObjects[0].Get();
-		HideUnconnectedPinsNode = ((CurObj && (CurObj->IsA<UK2Node_BreakStruct>() || CurObj->IsA<UK2Node_GetClassDefaults>())) ? static_cast<UK2Node*>(CurObj) : nullptr);
+		HideUnconnectedPinsNode = ((CurObj && (CurObj->IsA<UK2Node_BreakStruct>() || CurObj->IsA<UK2Node_MakeStruct>() || CurObj->IsA<UK2Node_GetClassDefaults>())) ? static_cast<UK2Node*>(CurObj) : nullptr);
+		bHideInputPins = (CurObj && CurObj->IsA<UK2Node_MakeStruct>());
 	}
 	else
 	{
@@ -120,7 +123,7 @@ void FSkeletalControlNodeDetails::CustomizeDetails(class IDetailLayoutBuilder& D
 			.MaxDesiredWidth(250.f)
 			[
 				SNew(SButton)
-				.OnClicked(this, &FSkeletalControlNodeDetails::HideAllUnconnectedPins)
+				.OnClicked(this, &FSkeletalControlNodeDetails::HideAllUnconnectedPins, bHideInputPins)
 				.ToolTipText(LOCTEXT("HideAllUnconnectedPinsTooltip", "All unconnected pins get hidden (removed from the graph node)"))
 				[
 					SNew(STextBlock)
@@ -271,7 +274,7 @@ void FSkeletalControlNodeDetails::OnGenerateElementForPropertyPin(TSharedRef<IPr
 	];
 }
 
-FReply FSkeletalControlNodeDetails::HideAllUnconnectedPins()
+FReply FSkeletalControlNodeDetails::HideAllUnconnectedPins(const bool bHideInputPins)
 {
 	if (ArrayProperty.IsValid() && HideUnconnectedPinsNode.IsValid())
 	{
@@ -288,7 +291,7 @@ FReply FSkeletalControlNodeDetails::HideAllUnconnectedPins()
 			FName ActualPropertyName;
 			if (PropertyNameHandle.IsValid() && PropertyNameHandle->GetValue(ActualPropertyName) == FPropertyAccess::Success)
 			{
-				const UEdGraphPin* Pin = HideUnconnectedPinsNode->FindPin(ActualPropertyName.ToString(), EGPD_Output);
+				const UEdGraphPin* Pin = HideUnconnectedPinsNode->FindPin(ActualPropertyName.ToString(), bHideInputPins ? EGPD_Input : EGPD_Output);
 				if (Pin && Pin->LinkedTo.Num() <= 0)
 				{
 					OnShowPinChanged(ECheckBoxState::Unchecked, ElementHandle);
diff --git a/Engine/Source/Editor/DetailCustomizations/Private/SkeletalControlNodeDetails.h b/Engine/Source/Editor/DetailCustomizations/Private/SkeletalControlNodeDetails.h
--- a/Engine/Source/Editor/DetailCustomizations/Private/SkeletalControlNodeDetails.h
+++ b/Engine/Source/Editor/DetailCustomizations/Private/SkeletalControlNodeDetails.h
@@ -33,7 +33,7 @@ protected:
 	void OnShowPinChanged(ECheckBoxState InNewState, TSharedRef<IPropertyHandle> InElementProperty);
 
 	/** Handler to hide all unconnected pins on a BreakStruct node */
-	FReply HideAllUnconnectedPins();
+	FReply HideAllUnconnectedPins(const bool bHideInputPins);
 
 private:
 	IDetailCategoryBuilder* DetailCategory;
