-Add a new way to override UProperty per Quality Level instead of Per Plaform. The goal is to decouple per asset overrides from the platform concept and introduce the PerQuality override that would be driven by scalability or device profile ini files. With the PerQualityLevelProperty, artists can now enter override value per quality level (ex: low,medium,high, ect..) and specify which quality level the platform supports in the ini files. This would allow us to support multiple override values at runtime and change between them depending on the device profile used by the platform. Unsupported quality level can still be stripped at cook time.

-Add Commandlet to be able to convert current MinLod PerPlatform data to PerQuality.Options:
     -output collections so that art can see the current override groups (ex: Console, Mobile)
     -convert all packages and checkout in perforce

diff --git a/Engine/Source/Editor/DetailCustomizations/Private/PerQualityLevelPropertyCustomization.cpp b/Engine/Source/Editor/DetailCustomizations/Private/PerQualityLevelPropertyCustomization.cpp
new file mode 100644
--- /dev/null
+++ b/Engine/Source/Editor/DetailCustomizations/Private/PerQualityLevelPropertyCustomization.cpp
@@ -0,0 +1,256 @@
+// Copyright Epic Games, Inc. All Rights Reserved.
+
+#include "PerQualityLevelPropertyCustomization.h"
+#include "Widgets/DeclarativeSyntaxSupport.h"
+#include "Engine/GameViewportClient.h"
+#include "Widgets/SBoxPanel.h"
+#include "Widgets/Layout/SSpacer.h"
+#include "Widgets/Text/STextBlock.h"
+#include "Widgets/Layout/SBox.h"
+#include "Widgets/Input/SComboBox.h"
+#include "Widgets/Images/SImage.h"
+#include "DetailWidgetRow.h"
+#include "Editor.h"
+#include "PropertyHandle.h"
+#include "DetailLayoutBuilder.h"
+#include "SPerQualityLevelPropertiesWidget.h"
+#include "ScopedTransaction.h"
+#include "IPropertyUtilities.h"
+#include "UObject/MetaData.h"
+#include "Scalability.h"
+
+#define LOCTEXT_NAMESPACE "PerOverridePropertyCustomization"
+
+template<typename OverrideType>
+TSharedRef<SWidget> FPerQualityLevelPropertyCustomization<OverrideType>::GetWidget(FName InQualityLevelName, TSharedRef<IPropertyHandle> StructPropertyHandle) const
+{
+	TSharedPtr<IPropertyHandle>	EditProperty;
+
+	if (InQualityLevelName == NAME_None)
+	{
+		EditProperty = StructPropertyHandle->GetChildHandle(FName("Default"));
+	}
+	else
+	{
+		TSharedPtr<IPropertyHandle>	MapProperty = StructPropertyHandle->GetChildHandle(FName("PerQuality"));
+		if (MapProperty.IsValid())
+		{
+			uint32 NumChildren = 0;
+			MapProperty->GetNumChildren(NumChildren);
+			for (uint32 ChildIdx = 0; ChildIdx < NumChildren; ChildIdx++)
+			{
+				TSharedPtr<IPropertyHandle> ChildProperty = MapProperty->GetChildHandle(ChildIdx);
+				if (ChildProperty.IsValid())
+				{
+					TSharedPtr<IPropertyHandle> KeyProperty = ChildProperty->GetKeyHandle();
+					if (KeyProperty.IsValid())
+					{
+						FName KeyName;
+						if (KeyProperty->GetValue(KeyName) == FPropertyAccess::Success && KeyName == InQualityLevelName)
+						{
+							EditProperty = ChildProperty;
+							break;
+						}
+					}
+				}
+			}
+		}
+
+	}
+
+	// Push down struct metadata to per-quality level properties
+	{
+		// First get the source map
+		const TMap<FName, FString>* SourceMap = StructPropertyHandle->GetMetaDataProperty()->GetMetaDataMap();
+		// Iterate through source map, setting each key/value pair in the destination
+		for (const auto& It : *SourceMap)
+		{
+			EditProperty->SetInstanceMetaData(*It.Key.ToString(), *It.Value);
+		}
+
+		// Copy instance metadata as well
+		const TMap<FName, FString>* InstanceSourceMap = StructPropertyHandle->GetInstanceMetaDataMap();
+		for (const auto& It : *InstanceSourceMap)
+		{
+			EditProperty->SetInstanceMetaData(*It.Key.ToString(), *It.Value);
+		}
+	}
+
+	if (EditProperty.IsValid())
+	{
+		return EditProperty->CreatePropertyValueWidget(false);
+	}
+	else
+	{
+		return
+			SNew(STextBlock)
+			.Text(NSLOCTEXT("FPerQualityLevelPropertyCustomization", "GetWidget", "Could not find valid property"))
+			.ColorAndOpacity(FLinearColor::Red);
+	}
+}
+
+template<typename OverrideType>
+float FPerQualityLevelPropertyCustomization<OverrideType>::CalcDesiredWidth(TSharedRef<IPropertyHandle> StructPropertyHandle)
+{
+	int32 NumOverrides = 0;
+	TSharedPtr<IPropertyHandle>	MapProperty = StructPropertyHandle->GetChildHandle(FName("PerQuality"));
+	if (MapProperty.IsValid())
+	{
+		TArray<const void*> RawData;
+		MapProperty->AccessRawData(RawData);
+		for (const void* Data : RawData)
+		{
+			const TMap<FName, typename OverrideType::ValueType>* PerQualityLevelMap = (const TMap<FName, typename OverrideType::ValueType>*)(Data);
+			NumOverrides = FMath::Max<int32>(PerQualityLevelMap->Num(), NumOverrides);
+		}
+	}
+	return (float)(1 + NumOverrides) * 125.f;
+}
+
+template<typename OverrideType>
+bool FPerQualityLevelPropertyCustomization<OverrideType>::AddOverride(FName InQualityLevelName, TSharedRef<IPropertyHandle> StructPropertyHandle)
+{
+	FScopedTransaction Transaction(LOCTEXT("AddOverride", "Add Quality Level Override"));
+
+	TSharedPtr<IPropertyHandle>	PerQualityLevelProperty = StructPropertyHandle->GetChildHandle(FName("PerQuality"));
+	TSharedPtr<IPropertyHandle>	DefaultProperty = StructPropertyHandle->GetChildHandle(FName("Default"));
+	if (PerQualityLevelProperty.IsValid() && DefaultProperty.IsValid())
+	{
+		TSharedPtr<IPropertyHandleMap> MapProperty = PerQualityLevelProperty->AsMap();
+		if (MapProperty.IsValid())
+		{
+			MapProperty->AddItem();
+			uint32 NumChildren = 0;
+			PerQualityLevelProperty->GetNumChildren(NumChildren);
+			for (uint32 ChildIdx = 0; ChildIdx < NumChildren; ChildIdx++)
+			{
+				TSharedPtr<IPropertyHandle> ChildProperty = PerQualityLevelProperty->GetChildHandle(ChildIdx);
+				if (ChildProperty.IsValid())
+				{
+					TSharedPtr<IPropertyHandle> KeyProperty = ChildProperty->GetKeyHandle();
+					if (KeyProperty.IsValid())
+					{
+						FName KeyName;
+						if (KeyProperty->GetValue(KeyName) == FPropertyAccess::Success && KeyName == NAME_None)
+						{
+							// Set Key
+							KeyProperty->SetValue(InQualityLevelName);
+
+							// Set Value
+							typename OverrideType::ValueType DefaultValue;
+							DefaultProperty->GetValue(DefaultValue);
+							ChildProperty->SetValue(DefaultValue);
+
+							if (PropertyUtilities.IsValid())
+							{
+								PropertyUtilities.Pin()->ForceRefresh();
+							}
+
+							return true;
+						}
+					}
+				}
+			}
+		}
+	}
+	return false;
+}
+
+template<typename OverrideType>
+bool FPerQualityLevelPropertyCustomization<OverrideType>::RemoveOverride(FName InQualityLevelName, TSharedRef<IPropertyHandle> StructPropertyHandle)
+{
+
+	FScopedTransaction Transaction(LOCTEXT("RemoveQualityLevelOverride", "Remove Quality Level Override"));
+
+	TSharedPtr<IPropertyHandle>	MapProperty = StructPropertyHandle->GetChildHandle(FName("PerQuality"));
+	if (MapProperty.IsValid())
+	{
+		TArray<const void*> RawData;
+		MapProperty->AccessRawData(RawData);
+		for (const void* Data : RawData)
+		{
+			TMap<FName, typename OverrideType::ValueType>* PerQualityLevelMap = (TMap<FName, typename OverrideType::ValueType>*)(Data);
+			check(PerQualityLevelMap);
+			TArray<FName> KeyArray;
+			PerQualityLevelMap->GenerateKeyArray(KeyArray);
+			for (FName QualityLevelName : KeyArray)
+			{
+				if (QualityLevelName == InQualityLevelName)
+				{
+
+					PerQualityLevelMap->Remove(InQualityLevelName);
+
+					if (PropertyUtilities.IsValid())
+					{
+						PropertyUtilities.Pin()->ForceRefresh();
+					}
+					return true;
+				}
+			}
+		}
+	}
+	return false;
+
+}
+
+template<typename OverrideType>
+TArray<FName> FPerQualityLevelPropertyCustomization<OverrideType>::GetOverrideNames(TSharedRef<IPropertyHandle> StructPropertyHandle) const
+{
+	TArray<FName> QualityLevelOverrideNames;
+
+	TSharedPtr<IPropertyHandle>	MapProperty = StructPropertyHandle->GetChildHandle(FName("PerQuality"));
+	if (MapProperty.IsValid())
+	{
+		TArray<const void*> RawData;
+		MapProperty->AccessRawData(RawData);
+		for (const void* Data : RawData)
+		{
+			const TMap<FName, typename OverrideType::ValueType>* PerQualityLevelMap = (const TMap<FName, typename OverrideType::ValueType>*)(Data);
+			check(PerQualityLevelMap);
+			TArray<FName> KeyArray;
+			PerQualityLevelMap->GenerateKeyArray(KeyArray);
+			for (FName QualityLevelName : KeyArray)
+			{
+				QualityLevelOverrideNames.AddUnique(QualityLevelName);
+			}
+		}
+
+	}
+	return QualityLevelOverrideNames;
+}
+
+template<typename OverrideType>
+void FPerQualityLevelPropertyCustomization<OverrideType>::CustomizeHeader(TSharedRef<IPropertyHandle> StructPropertyHandle, FDetailWidgetRow& HeaderRow, IPropertyTypeCustomizationUtils& StructCustomizationUtils)
+{
+	FPerQualityLevelPropertyCustomization<OverrideType>::PropertyUtilities = StructCustomizationUtils.GetPropertyUtilities();
+
+	HeaderRow.NameContent()
+		[
+			StructPropertyHandle->CreatePropertyNameWidget()
+		]
+	.ValueContent()
+		.MinDesiredWidth(FPerQualityLevelPropertyCustomization<OverrideType>::CalcDesiredWidth(StructPropertyHandle))
+		.MaxDesiredWidth((float)(static_cast<int32>(QualityLevelProperty::EQualityLevels::Num) + 1) * 125.0f)
+		[
+			SNew(SPerQualityLevelPropertiesWidget)
+			.OnGenerateWidget(this, &FPerQualityLevelPropertyCustomization<OverrideType>::GetWidget, StructPropertyHandle)
+			.OnAddEntry(this, &FPerQualityLevelPropertyCustomization<OverrideType>::AddOverride, StructPropertyHandle)
+			.OnRemoveEntry(this, &FPerQualityLevelPropertyCustomization<OverrideType>::RemoveOverride, StructPropertyHandle)
+			.EntryNames(this, &FPerQualityLevelPropertyCustomization<OverrideType>::GetOverrideNames, StructPropertyHandle)
+		];
+}
+
+template<typename OverrideType>
+TSharedRef<IPropertyTypeCustomization> FPerQualityLevelPropertyCustomization<OverrideType>::MakeInstance()
+{
+	return MakeShareable(new FPerQualityLevelPropertyCustomization<OverrideType>);
+}
+
+
+
+/* Only explicitly instantiate the types which are supported
+*****************************************************************************/
+template class FPerQualityLevelPropertyCustomization<FPerQualityLevelInt>;
+
+
+#undef LOCTEXT_NAMESPACE
\ No newline at end of file
diff --git a/Engine/Source/Editor/DetailCustomizations/Private/PerQualityLevelPropertyCustomization.h b/Engine/Source/Editor/DetailCustomizations/Private/PerQualityLevelPropertyCustomization.h
new file mode 100644
--- /dev/null
+++ b/Engine/Source/Editor/DetailCustomizations/Private/PerQualityLevelPropertyCustomization.h
@@ -0,0 +1,44 @@
+// Copyright Epic Games, Inc. All Rights Reserved.
+#pragma once
+
+#include "CoreMinimal.h"
+#include "Widgets/SWidget.h"
+#include "UnrealClient.h"
+#include "IPropertyTypeCustomization.h"
+#include "PerQualityLevelProperties.h"
+
+class FDetailWidgetDecl;
+
+/**
+* Implements a details panel customization for the FPerQualityLevel structures.
+*/
+template<typename OverrideType>
+class FPerQualityLevelPropertyCustomization : public IPropertyTypeCustomization
+{
+public:
+	FPerQualityLevelPropertyCustomization()
+	{}
+
+	// IPropertyTypeCustomization interface
+	virtual void CustomizeHeader(TSharedRef<IPropertyHandle> StructPropertyHandle, FDetailWidgetRow& HeaderRow, IPropertyTypeCustomizationUtils& StructCustomizationUtils) override;
+	virtual void CustomizeChildren(TSharedRef<IPropertyHandle> StructPropertyHandle, IDetailChildrenBuilder& StructBuilder, IPropertyTypeCustomizationUtils& StructCustomizationUtils) override {}
+
+	/**
+	* Creates a new instance.
+	*
+	* @return A new customization for FPerQualityLevel structs.
+	*/
+	static TSharedRef<IPropertyTypeCustomization> MakeInstance();
+
+protected:
+	TSharedRef<SWidget> GetWidget(FName InQualityLevelName, TSharedRef<IPropertyHandle> StructPropertyHandle) const;
+	TArray<FName> GetOverrideNames(TSharedRef<IPropertyHandle> StructPropertyHandle) const;
+	bool AddOverride(FName InQualityLevelName, TSharedRef<IPropertyHandle> StructPropertyHandle);
+	bool RemoveOverride(FName InQualityLevelName, TSharedRef<IPropertyHandle> StructPropertyHandle);
+	float CalcDesiredWidth(TSharedRef<IPropertyHandle> StructPropertyHandle);
+
+private:
+	/** Cached utils used for resetting customization when layout changes */
+	TWeakPtr<IPropertyUtilities> PropertyUtilities;
+};
+
diff --git a/Engine/Source/Editor/PropertyEditor/Private/SPerQualityLevelPropertiesWidget.cpp b/Engine/Source/Editor/PropertyEditor/Private/SPerQualityLevelPropertiesWidget.cpp
new file mode 100644
--- /dev/null
+++ b/Engine/Source/Editor/PropertyEditor/Private/SPerQualityLevelPropertiesWidget.cpp
@@ -0,0 +1,212 @@
+// Copyright Epic Games, Inc. All Rights Reserved.
+
+#include "SPerQualityLevelPropertiesWidget.h"
+#include "Layout/Margin.h"
+#include "Fonts/SlateFontInfo.h"
+#include "Styling/CoreStyle.h"
+#include "Framework/MultiBox/MultiBoxBuilder.h"
+#include "Widgets/Input/SButton.h"
+#include "Widgets/Input/SComboButton.h"
+#include "PlatformInfo.h"
+#include "Widgets/Layout/SWrapBox.h"
+#include "EditorStyleSet.h"
+#include "Widgets/Layout/SBox.h"
+#include "DetailLayoutBuilder.h"
+#include "Widgets/Images/SImage.h"
+#include "Scalability.h"
+
+void SOverridePropertiesWidget::Construct(const typename SOverridePropertiesWidget::FArguments& InArgs)
+{
+	this->OnGenerateWidget = InArgs._OnGenerateWidget;
+	this->OnAddEntry = InArgs._OnAddEntry;
+	this->OnRemoveEntry = InArgs._OnRemoveEntry;
+	this->EntryNames = InArgs._EntryNames;
+}
+
+
+FReply SOverridePropertiesWidget::RemoveEntry(FName OverrideName)
+{
+	if (OnRemoveEntry.IsBound() && OnRemoveEntry.Execute(OverrideName))
+	{
+		ConstructChildren();
+		Invalidate(EInvalidateWidget::LayoutAndVolatility);
+	}
+	return FReply::Handled();
+}
+
+TSharedRef<SWidget> SOverridePropertiesWidget::MakeOverrideWidget(FName InName, FText InDisplayText, const TArray<FName>& InEntries, FMenuBuilder& InAddMenuBuilder)
+{
+	TSharedPtr<SHorizontalBox> HorizontalBox;
+
+	TSharedRef<SWidget> Widget =
+		SNew(SBox)
+		.ToolTipText(InName == NAME_None ?
+			NSLOCTEXT("SPerQualityLevelPropertiesWidget", "DefaultQualityLevelDesc", "This property can have per quality level overrides.\nThis is the default value used when no override has been set for a specific quality level.") :
+			FText::Format(NSLOCTEXT("SPerQualityLevelPropertiesWidget", "QualityLevelDesc", "Override for {0}"), InDisplayText))
+		.Padding(FMargin(0.0f, 2.0f, 4.0f, 2.0f))
+		.MinDesiredWidth(50.0f)
+		[
+			SNew(SVerticalBox)
+			+ SVerticalBox::Slot()
+		.AutoHeight()
+		[
+			SAssignNew(HorizontalBox, SHorizontalBox)
+			+ SHorizontalBox::Slot()
+		.FillWidth(1.0f)
+		.Padding(0.0f, 0.0f, 2.0f, 2.0f)
+		.VAlign(VAlign_Center)
+		[
+			SNew(STextBlock)
+			.Font(IDetailLayoutBuilder::GetDetailFont())
+		.Text(InDisplayText)
+		]
+		]
+	+ SVerticalBox::Slot()
+		.AutoHeight()
+		[
+			OnGenerateWidget.Execute(InName)
+		]
+		];
+
+	if (InName != NAME_None)
+	{
+		HorizontalBox->AddSlot()
+			.AutoWidth()
+			.Padding(2.0f, 0.0f, 0.0f, 2.0f)
+			.VAlign(VAlign_Center)
+			[
+				SNew(SButton)
+				.ContentPadding(2.0f)
+			.ButtonStyle(FEditorStyle::Get(), "HoverHintOnly")
+			.OnClicked(this, &SOverridePropertiesWidget::RemoveEntry, InName)
+			.ToolTipText(FText::Format(NSLOCTEXT("SOverridePropertiesWidget", "RemoveOverrideFor", "Remove Override for {0}"), InDisplayText))
+			.ForegroundColor(FSlateColor::UseForeground())
+			.HAlign(HAlign_Center)
+			.VAlign(VAlign_Center)
+			.Content()
+			[
+				SNew(SImage)
+				.Image(FEditorStyle::GetBrush("Icons.Cross"))
+			]
+			];
+	}
+	else
+	{
+		HorizontalBox->AddSlot()
+			.AutoWidth()
+			.Padding(2.0f, 0.0f, 0.0f, 2.0f)
+			.VAlign(VAlign_Center)
+			[
+				SNew(SComboButton)
+				.Visibility_Lambda([this]() { return bAddedMenuItem ? EVisibility::Visible : EVisibility::Hidden; })
+			.ButtonStyle(FEditorStyle::Get(), "HoverHintOnly")
+			.ContentPadding(2.0f)
+			.ForegroundColor(FSlateColor::UseForeground())
+			.HasDownArrow(false)
+			.ButtonContent()
+			[
+				SNew(SImage)
+				.Image(FEditorStyle::GetBrush("PropertyWindow.Button_AddToArray"))
+			]
+		.MenuContent()
+			[
+				InAddMenuBuilder.MakeWidget()
+			]
+		.ToolTipText(NSLOCTEXT("SOverridePropertiesWidget", "AddOverrideToolTip", "Add an override for a specific quality level"))
+			];
+	}
+
+	return Widget;
+}
+
+void SOverridePropertiesWidget::AddEntryToMenu(const FName& EntryName, const FTextFormat Format, FMenuBuilder& AddEntryMenuBuilder)
+{
+	const FText MenuText = FText::Format(FText::FromString(TEXT("{0}")), FText::AsCultureInvariant(EntryName.ToString()));
+	const FText MenuTooltipText = FText::Format(Format, FText::AsCultureInvariant(EntryName.ToString()));
+	AddEntryMenuBuilder.AddMenuEntry(
+		MenuText,
+		MenuTooltipText,
+		FSlateIcon(FEditorStyle::GetStyleSetName(), "OverridePropertiesWidget.AddEntry"),
+		FUIAction(FExecuteAction::CreateSP(this, &SOverridePropertiesWidget::AddEntry, EntryName))
+	);
+}
+
+void SOverridePropertiesWidget::AddEntry(FName EntryName)
+{
+	if (OnAddEntry.IsBound() && OnAddEntry.Execute(EntryName))
+	{
+		ConstructChildren();
+		Invalidate(EInvalidateWidget::LayoutAndVolatility);
+	}
+}
+
+void SPerQualityLevelPropertiesWidget::Construct(const typename SOverridePropertiesWidget::FArguments& InArgs)
+{
+	ToolTip = FString(TEXT("Add an override for a specific quality level"));
+
+	SOverridePropertiesWidget::Construct(InArgs);
+
+	ConstructChildren();
+}
+
+
+void SPerQualityLevelPropertiesWidget::ConstructChildren()
+{
+	TSharedPtr<SWrapBox> WrapBox;
+
+	TArray<FName> Overrides = EntryNames.Get();
+	LastEntryNames = Overrides.Num();
+
+	ChildSlot
+		.VAlign(VAlign_Center)
+		.HAlign(HAlign_Fill)
+		[
+			SAssignNew(WrapBox, SWrapBox)
+			.UseAllottedSize(true)
+		];
+
+	if (OnGenerateWidget.IsBound())
+	{
+		// Build quality level menu
+		bAddedMenuItem = false;
+		FMenuBuilder AddMenuBuilder(true, nullptr, nullptr, true);
+
+		const FName Section(TEXT("PlatformGroupSection"));
+		AddMenuBuilder.BeginSection(Section, FText::FromString(TEXT("Quality Levels")));
+
+		int32 QualityLevelCount = 0;
+		while (QualityLevelCount < static_cast<int32>(QualityLevelProperty::EQualityLevels::Num))
+		{
+			FName QualityName = QualityLevelProperty::QualityLevelToFName(QualityLevelCount++);
+			const FTextFormat Format = NSLOCTEXT("SPerQualityLevelPropertiesWidget", "AddOverrideGroupFor", "Add Override specifically for {0}");
+			AddEntryToMenu(QualityName, Format, AddMenuBuilder);
+			bAddedMenuItem = true;
+			
+		}
+
+		AddMenuBuilder.EndSection();
+
+		// Default control
+		WrapBox->AddSlot()
+			[
+				MakeOverrideWidget(NAME_None, NSLOCTEXT("SPerQualityLevelPropertiesWidget", "DefaultQuality", "Default"), Overrides, AddMenuBuilder)
+			];
+
+		for (FName Override : Overrides)
+		{
+			WrapBox->AddSlot()
+				[
+					MakeOverrideWidget(Override, FText::AsCultureInvariant(Override.ToString()), Overrides, AddMenuBuilder)
+				];
+		}
+	}
+	else
+	{
+		WrapBox->AddSlot()
+			[
+				SNew(STextBlock)
+				.Text(NSLOCTEXT("SPerQualityLevelPropertiesWidget", "OnGenerateWidgetWarning", "No OnGenerateWidget() Provided"))
+			.ColorAndOpacity(FLinearColor::Red)
+			];
+	}
+}
\ No newline at end of file
diff --git a/Engine/Source/Editor/PropertyEditor/Public/SPerQualityLevelPropertiesWidget.h b/Engine/Source/Editor/PropertyEditor/Public/SPerQualityLevelPropertiesWidget.h
new file mode 100644
--- /dev/null
+++ b/Engine/Source/Editor/PropertyEditor/Public/SPerQualityLevelPropertiesWidget.h
@@ -0,0 +1,76 @@
+// Copyright Epic Games, Inc. All Rights Reserved.
+
+#pragma once
+
+#include "CoreMinimal.h"
+#include "Widgets/DeclarativeSyntaxSupport.h"
+#include "Widgets/SWidget.h"
+#include "Widgets/SCompoundWidget.h"
+#include "Framework/SlateDelegates.h"
+
+class FMenuBuilder;
+
+class PROPERTYEDITOR_API SOverridePropertiesWidget : public SCompoundWidget
+{
+public:
+	typedef typename TSlateDelegates<FName>::FOnGenerateWidget FOnGenerateWidget;
+
+	DECLARE_DELEGATE_RetVal_OneParam(bool, FOnOverrideAction, FName);
+
+	SLATE_BEGIN_ARGS(SOverridePropertiesWidget)
+		: _OnGenerateWidget()
+	{}
+
+	SLATE_EVENT(FOnGenerateWidget, OnGenerateWidget)
+		SLATE_EVENT(FOnOverrideAction, OnAddEntry)
+		SLATE_EVENT(FOnOverrideAction, OnRemoveEntry)
+
+		SLATE_ATTRIBUTE(TArray<FName>, EntryNames)
+
+		SLATE_END_ARGS()
+
+	/**
+	* Construct this widget
+	*
+	* @param	InArgs	The declaration data for this widget
+	*/
+	virtual void Construct(const typename SOverridePropertiesWidget::FArguments& InArgs);
+
+protected:
+	virtual void ConstructChildren() = 0;
+
+	TSharedRef<SWidget> MakeOverrideWidget(FName InName, FText InDisplayText, const TArray<FName>& InEntries, FMenuBuilder& InAddMenuBuilder);
+
+	void AddEntry(FName EntryName);
+
+	FReply RemoveEntry(FName EntryName);
+
+	void AddEntryToMenu(const FName& EntryName, const FTextFormat Format, FMenuBuilder& AddEntryMenuBuilder);
+
+	FOnGenerateWidget OnGenerateWidget;
+	FOnOverrideAction OnAddEntry;
+	FOnOverrideAction OnRemoveEntry;
+	TAttribute<TArray<FName>> EntryNames;
+	int32 LastEntryNames;
+	bool bAddedMenuItem;
+
+	FString ToolTip;
+};
+
+/**
+* SPerQualityLevelPropertiesWidget
+*/
+class PROPERTYEDITOR_API SPerQualityLevelPropertiesWidget : public SOverridePropertiesWidget
+{
+public:
+
+	virtual void Construct(const typename SPerQualityLevelPropertiesWidget::FArguments& InArgs);
+
+protected:
+	virtual void ConstructChildren() override;
+
+};
+
+
+
+
diff --git a/Engine/Source/Editor/UnrealEd/Classes/Commandlets/StaticMeshMinLodCommandlet.h b/Engine/Source/Editor/UnrealEd/Classes/Commandlets/StaticMeshMinLodCommandlet.h
new file mode 100644
--- /dev/null
+++ b/Engine/Source/Editor/UnrealEd/Classes/Commandlets/StaticMeshMinLodCommandlet.h
@@ -0,0 +1,20 @@
+// Copyright Epic Games, Inc. All Rights Reserved.
+
+#pragma once
+
+#include "CoreMinimal.h"
+#include "UObject/ObjectMacros.h"
+#include "Commandlets/Commandlet.h"
+#include "StaticMeshMinLodCommandlet.generated.h"
+
+/** Commandlet  */
+UCLASS()
+class UStaticMeshMinLodCommandlet : public UCommandlet
+{
+	GENERATED_UCLASS_BODY()
+
+public:
+	//~ Begin UCommandlet Interface
+	virtual int32 Main(const FString& Params) override;
+	//~ End UCommandlet Interface
+};
\ No newline at end of file
diff --git a/Engine/Source/Runtime/Engine/Private/PerQualityLevelProperties.cpp b/Engine/Source/Runtime/Engine/Private/PerQualityLevelProperties.cpp
new file mode 100644
--- /dev/null
+++ b/Engine/Source/Runtime/Engine/Private/PerQualityLevelProperties.cpp
@@ -0,0 +1,183 @@
+// Copyright Epic Games, Inc. All Rights Reserved.
+#include "PerQualityLevelProperties.h"
+#include "Serialization/Archive.h"
+
+namespace QualityLevelProperty
+{
+	static TArray<FName> Names = { FName("Low"), FName("Medium"), FName("High"), FName("Epic"), FName("Cinematic") };
+
+	FName QualityLevelToFName(int32 QL)
+	{
+		if (QL >= 0 && QL < static_cast<int32>(EQualityLevels::Num))
+		{
+			return Names[QL];
+		}
+		else
+		{
+			return NAME_None;
+		}
+		return Names[QL];
+	}
+
+	int32 FNameToQualityLevel(FName QL)
+	{
+		return Names.IndexOfByKey(QL);
+	}
+}
+
+#if WITH_EDITOR
+#include "Interfaces/ITargetPlatform.h"
+#include "PlatformInfo.h"
+
+static TMap<FString, FSupportedQualityLevelArray> SupportedQualitLevels;
+static FCriticalSection CookCriticalSection;
+
+template<typename _StructType, typename _ValueType, EName _BasePropertyName>
+FSupportedQualityLevelArray FPerQualityLevelProperty<_StructType, _ValueType, _BasePropertyName>::GetSupportedQualityLevels(const TCHAR* InPlatformName) const
+{
+	FScopeLock ScopeLock(&CookCriticalSection);
+
+	FSupportedQualityLevelArray* CachedCookingQualitLevelInfo  = SupportedQualitLevels.Find(InPlatformName);
+	if (CachedCookingQualitLevelInfo)
+	{
+		return *CachedCookingQualitLevelInfo;
+	}
+
+	FSupportedQualityLevelArray CookingQualitLevelInfo;
+
+	// check the Engine file
+	FConfigFile EngineSettings;
+	FConfigCacheIni::LoadLocalIniFile(EngineSettings, TEXT("Engine"), true, InPlatformName);
+
+	int32 PropertyQualityLevel = -1;
+	if (EngineSettings.GetInt(TEXT("SystemSettings"), *CVarName, PropertyQualityLevel))
+	{
+		CookingQualitLevelInfo.Add(PropertyQualityLevel);
+	}
+
+	// Load the scalability platform file
+	FConfigFile ScalabilitySettings;
+	FConfigCacheIni::LoadLocalIniFile(ScalabilitySettings, TEXT("Scalability"), true, InPlatformName);
+
+	//check all possible quality levels specify in the scalability ini 
+	for (int32 QualityLevel = 0; QualityLevel < (int32)QualityLevelProperty::EQualityLevels::Num; ++QualityLevel)
+	{
+		FString QualitLevelSectionName = Scalability::GetScalabilitySectionString(*ScalabilitySection, QualityLevel, (int32)QualityLevelProperty::EQualityLevels::Num);
+		PropertyQualityLevel = -1;
+		ScalabilitySettings.GetInt(*QualitLevelSectionName, *CVarName, PropertyQualityLevel);
+
+		// add supported quality level to the property map
+		if (PropertyQualityLevel != -1)
+		{
+			CookingQualitLevelInfo.Add(PropertyQualityLevel);
+		}
+	}
+
+	// Cache the Scalability setting for this platform
+	CachedCookingQualitLevelInfo = &SupportedQualitLevels.Add(FString(InPlatformName), CookingQualitLevelInfo);
+
+	return *CachedCookingQualitLevelInfo;
+}
+
+template<typename _StructType, typename _ValueType, EName _BasePropertyName>
+void FPerQualityLevelProperty<_StructType, _ValueType, _BasePropertyName>::StripQualtiyLevelForCooking(const TCHAR* InPlatformName)
+{
+	_StructType* This = StaticCast<_StructType*>(this);
+	if (This->PerQuality.Num() > 0)
+	{
+		FSupportedQualityLevelArray CookQualityLevelInfo = This->GetSupportedQualityLevels(InPlatformName);
+
+		int32 LowestQualityLevel = (int32)QualityLevelProperty::EQualityLevels::Num;
+
+		// remove unsupported quality levels 
+		for (TMap<int32, int32>::TIterator It(This->PerQuality); It; ++It)
+		{
+			if(!CookQualityLevelInfo.Contains(It.Key()))
+			{
+				It.RemoveCurrent();
+			}
+			else
+			{
+				LowestQualityLevel = (It.Key() < LowestQualityLevel) ? It.Key() : LowestQualityLevel;
+			}
+		}
+
+		//if found supported platforms, put the lowest quality level in Default
+		if (LowestQualityLevel != (int32)QualityLevelProperty::EQualityLevels::Num)
+		{
+			This->Default = This->GetValue(LowestQualityLevel);
+		}
+	}
+}
+
+template<typename _StructType, typename _ValueType, EName _BasePropertyName>
+bool FPerQualityLevelProperty<_StructType, _ValueType, _BasePropertyName>::IsQualityLevelValid(int32 QualityLevel) const
+{
+	const _StructType* This = StaticCast<const _StructType*>(this);
+	int32* Value = (int32*)This->PerQuality.Find(QualityLevel);
+
+	if (Value != nullptr)
+	{
+		return true;
+	}
+	else
+	{
+		return false;
+	}
+}
+#endif
+
+/** Serializer to cook out the most appropriate platform override */
+template<typename _StructType, typename _ValueType, EName _BasePropertyName>
+ENGINE_API FArchive& operator<<(FArchive& Ar, FPerQualityLevelProperty<_StructType, _ValueType, _BasePropertyName>& Property)
+{
+	bool bCooked = false;
+	_StructType* This = StaticCast<_StructType*>(&Property);
+
+#if WITH_EDITOR
+	if (Ar.IsCooking())
+	{
+		bCooked = true;
+		This->StripQualtiyLevelForCooking(*(Ar.CookingTarget()->PlatformName()));
+	}
+#endif
+	{
+		Ar << bCooked;
+		Ar << This->Default;
+		Ar << This->PerQuality;
+	}
+	return Ar;
+}
+
+/** Serializer to cook out the most appropriate platform override */
+template<typename _StructType, typename _ValueType, EName _BasePropertyName>
+ENGINE_API void operator<<(FStructuredArchive::FSlot Slot, FPerQualityLevelProperty<_StructType, _ValueType, _BasePropertyName>& Property)
+{
+	FArchive& UnderlyingArchive = Slot.GetUnderlyingArchive();
+	FStructuredArchive::FRecord Record = Slot.EnterRecord();
+
+	bool bCooked = false;
+	_StructType* This = StaticCast<_StructType*>(&Property);
+
+#if WITH_EDITOR
+	if (UnderlyingArchive.IsCooking())
+	{
+		bCooked = true;
+		This->StripQualtiyLevelForCooking(*(UnderlyingArchive.CookingTarget()->GetPlatformInfo().PlatformGroupName.ToString()));
+	}
+#endif
+	{
+		Record << SA_VALUE(TEXT("bCooked"), bCooked);
+		Record << SA_VALUE(TEXT("Value"), This->Default);
+		Record << SA_VALUE(TEXT("PerQuality"), This->PerQuality);
+	}
+}
+// 
+template ENGINE_API FArchive& operator<<(FArchive&, FPerQualityLevelProperty<FPerQualityLevelInt, int32, NAME_IntProperty>&);
+template ENGINE_API void operator<<(FStructuredArchive::FSlot Slot, FPerQualityLevelProperty<FPerQualityLevelInt, int32, NAME_IntProperty>&);
+
+#if WITH_EDITOR
+template FSupportedQualityLevelArray FPerQualityLevelProperty<FPerQualityLevelInt, int32, NAME_IntProperty>::GetSupportedQualityLevels(const TCHAR* InPlatformName) const;
+template void FPerQualityLevelProperty<FPerQualityLevelInt, int32, NAME_IntProperty>::StripQualtiyLevelForCooking(const TCHAR* InPlatformName);
+template bool FPerQualityLevelProperty<FPerQualityLevelInt, int32, NAME_IntProperty>::IsQualityLevelValid(int32 QualityLevel) const;
+#endif
\ No newline at end of file
diff --git a/Engine/Source/Runtime/Engine/Public/PerQualityLevelProperties.h b/Engine/Source/Runtime/Engine/Public/PerQualityLevelProperties.h
new file mode 100644
--- /dev/null
+++ b/Engine/Source/Runtime/Engine/Public/PerQualityLevelProperties.h
@@ -0,0 +1,182 @@
+// Copyright Epic Games, Inc. All Rights Reserved.
+
+/*=============================================================================
+FPerQualityLevelProperties.h: Property types that can be overridden on a quality level basis at cook time
+=============================================================================*/
+
+#pragma once
+
+#include "Serialization/Archive.h"
+#include "Containers/Map.h"
+#include "Algo/Find.h"
+#include "Serialization/MemoryLayout.h"
+#include "Scalability.h"
+#include "CoreMinimal.h"
+//#if WITH_EDITORONLY_DATA && WITH_EDITOR
+#include "Engine/Engine.h"
+#include "HAL/ConsoleManager.h"
+//#endif
+
+#include "PerQualityLevelProperties.generated.h"
+
+#if WITH_EDITOR
+	typedef TSet<int32> FSupportedQualityLevelArray;
+#endif
+
+namespace QualityLevelProperty
+{
+	UENUM()
+	enum class EQualityLevels : uint8
+	{
+		Low,
+		Medium,
+		High,
+		Epic,
+		Cinematic,
+		Num
+	};
+
+	ENGINE_API FName QualityLevelToFName(int32 QL);
+	ENGINE_API int32 FNameToQualityLevel(FName QL);
+};
+
+template<typename _StructType, typename _ValueType, EName _BasePropertyName>
+struct ENGINE_API FPerQualityLevelProperty
+{
+	typedef _ValueType ValueType;
+
+	FPerQualityLevelProperty() 
+	{
+	}
+	~FPerQualityLevelProperty() {}
+
+	bool operator==(const FPerQualityLevelProperty& Other) const;
+
+	_ValueType GetValueForQualityLevel(int32 QualityLevel) const
+	{
+		const _StructType* This = StaticCast<const _StructType*>(this);
+		if (This->PerQuality.Num() == 0 || QualityLevel < 0)
+		{
+			return This->Default;
+		}
+
+		int32* Value = (int32*)This->PerQuality.Find(QualityLevel);
+		
+		if (Value)
+		{
+			return *Value;
+		}
+		else
+		{
+			return This->Default;
+		}
+	}
+
+#if WITH_EDITOR
+	FSupportedQualityLevelArray GetSupportedQualityLevels(const TCHAR* InPlatformName = nullptr) const;
+	void StripQualtiyLevelForCooking(const TCHAR* InPlatformName = nullptr);
+	bool IsQualityLevelValid(int32 QualityLevel) const;
+#endif
+
+	void Init(const TCHAR* InCVarName, const TCHAR* InSection)
+	{
+#if WITH_EDITOR
+		ScalabilitySection = FString(InSection);
+#endif
+		CVarName = FString(InCVarName);
+	}
+
+	_ValueType GetDefault() const
+	{
+		const _StructType* This = StaticCast<const _StructType*>(this);
+		return This->Default;
+	}
+
+	_ValueType GetValue(int32 QualityLevel) const
+	{
+		return GetValueForQualityLevel(QualityLevel);
+	}
+
+	/* Load old properties that have been converted to FPerQualityLevel */
+	bool SerializeFromMismatchedTag(const FPropertyTag& Tag, FArchive& Ar)
+	{
+		if (Tag.Type == _BasePropertyName)
+		{
+			_StructType* This = StaticCast<_StructType*>(this);
+			_ValueType OldValue;
+			Ar << OldValue;
+			*This = _StructType(OldValue);
+			return true;
+		}
+		return false;
+	}
+
+	/* Serialization */
+	bool Serialize(FArchive& Ar)
+	{
+		Ar << *this;
+		return true;
+	}
+
+	/* Serialization */
+	bool Serialize(FStructuredArchive::FSlot Slot)
+	{
+		Slot << *this;
+		return true;
+	}
+
+#if WITH_EDITOR
+	FString ScalabilitySection;
+#endif
+	FString CVarName;
+};
+
+template< typename _StructType, typename _ValueType, EName _BasePropertyName>
+ENGINE_API FArchive& operator<<(FArchive& Ar, FPerQualityLevelProperty<_StructType, _ValueType, _BasePropertyName>& Property);
+
+template<typename _StructType, typename _ValueType, EName _BasePropertyName>
+ENGINE_API void operator<<(FStructuredArchive::FSlot Slot, FPerQualityLevelProperty<_StructType, _ValueType, _BasePropertyName>& Property);
+
+USTRUCT()
+struct ENGINE_API FPerQualityLevelInt 
+#if CPP
+	:	public FPerQualityLevelProperty<FPerQualityLevelInt, int32, NAME_IntProperty>
+#endif
+{
+	GENERATED_USTRUCT_BODY()
+
+	UPROPERTY(EditAnywhere, Category = PerQualityLevel)
+	int32 Default;
+
+	UPROPERTY(EditAnywhere, Category = PerQualityLevel)
+	TMap<int32, int32> PerQuality;
+
+	UPROPERTY(EditAnywhere, Category = PerQualityLevel)
+	bool bIsEnabled;
+
+	FPerQualityLevelInt()
+	{
+		Default = 0;
+		bIsEnabled = false;
+	}
+
+	FPerQualityLevelInt(int32 InDefaultValue)
+	{
+		Default = InDefaultValue;
+		bIsEnabled = false;
+	}
+};
+
+extern template ENGINE_API FArchive& operator<<(FArchive&, FPerQualityLevelProperty<FPerQualityLevelInt, int32, NAME_IntProperty>&);
+extern template ENGINE_API void operator<<(FStructuredArchive::FSlot Slot, FPerQualityLevelProperty<FPerQualityLevelInt, int32, NAME_IntProperty>&);
+
+template<>
+struct TStructOpsTypeTraits<FPerQualityLevelInt>
+	: public TStructOpsTypeTraitsBase2<FPerQualityLevelInt>
+{
+	enum
+	{
+		WithSerializeFromMismatchedTag = true,
+		WithSerializer = true
+	};
+};
\ No newline at end of file
