Compile fix for non unity build (missing includes)

diff --git a/Engine/Source/Editor/DetailCustomizations/Private/SPerPlatformPropertiesWidget.cpp b/Engine/Source/Editor/DetailCustomizations/Private/SPerPlatformPropertiesWidget.cpp
--- a/Engine/Source/Editor/DetailCustomizations/Private/SPerPlatformPropertiesWidget.cpp
+++ b/Engine/Source/Editor/DetailCustomizations/Private/SPerPlatformPropertiesWidget.cpp
@@ -13,7 +13,7 @@
 #include "Widgets/Layout/SBox.h"
 #include "DetailLayoutBuilder.h"
 #include "Widgets/Images/SImage.h"
-
+#include "PerQualityLevelProperties.h"
 
 void SPerPlatformPropertiesRow::Construct(const FArguments& InArgs, FName PlatformName)
 {
diff --git a/Engine/Source/Runtime/Engine/Private/PerQualityLevelProperties.cpp b/Engine/Source/Runtime/Engine/Private/PerQualityLevelProperties.cpp
--- a/Engine/Source/Runtime/Engine/Private/PerQualityLevelProperties.cpp
+++ b/Engine/Source/Runtime/Engine/Private/PerQualityLevelProperties.cpp
@@ -1,6 +1,7 @@
 // Copyright Epic Games, Inc. All Rights Reserved.
 #include "PerQualityLevelProperties.h"
 #include "Serialization/Archive.h"
+#include "Misc/ConfigCacheIni.h"
 
 namespace QualityLevelProperty
 {
