SceneCapture: Fixed crash when SceneCapture is deleted while its Details are visible and a show flag setting has been changed from the archetype defaults.

diff --git a/Engine/Source/Editor/DetailCustomizations/Private/SceneCaptureDetails.cpp b/Engine/Source/Editor/DetailCustomizations/Private/SceneCaptureDetails.cpp
--- a/Engine/Source/Editor/DetailCustomizations/Private/SceneCaptureDetails.cpp
+++ b/Engine/Source/Editor/DetailCustomizations/Private/SceneCaptureDetails.cpp
@@ -296,7 +296,12 @@ ECheckBoxState FSceneCaptureDetails::OnGetDisplayCheckState(FString ShowFlagName
 	for (int32 ObjectIdx = 0; ObjectIdx < RawData.Num(); ++ObjectIdx)
 	{
 		const void* Data = RawData[ObjectIdx];
-		check(Data);
+
+		// Data can be nullptr when the SceneCapture is removed while its details are visible
+		if (!Data)
+		{
+			return ECheckBoxState::Unchecked;
+		}
 
 		const TArray<FEngineShowFlagsSetting>& ShowFlagSettings = *reinterpret_cast<const TArray<FEngineShowFlagsSetting>*>(Data);
 		const FEngineShowFlagsSetting* Setting = ShowFlagSettings.FindByPredicate([&ShowFlagName](const FEngineShowFlagsSetting& S) { return S.ShowFlagName == ShowFlagName; });
