Fixed infinite loop when assigning a material to many components via the details view.

diff --git a/Engine/Source/Editor/DetailCustomizations/Private/CustomPrimitiveDataCustomization.cpp b/Engine/Source/Editor/DetailCustomizations/Private/CustomPrimitiveDataCustomization.cpp
--- a/Engine/Source/Editor/DetailCustomizations/Private/CustomPrimitiveDataCustomization.cpp
+++ b/Engine/Source/Editor/DetailCustomizations/Private/CustomPrimitiveDataCustomization.cpp
@@ -33,9 +33,19 @@ TSharedRef<IPropertyTypeCustomization> FCustomPrimitiveDataCustomization::MakeIn
 	return MakeShareable(new FCustomPrimitiveDataCustomization);
 }
 
+FCustomPrimitiveDataCustomization::FCustomPrimitiveDataCustomization()
+{
+	// NOTE: Optimally would be bound to a "OnMaterialChanged" for each component
+	FCoreUObjectDelegates::OnObjectPropertyChanged.AddRaw(this, &FCustomPrimitiveDataCustomization::OnObjectPropertyChanged);
+	UMaterial::OnMaterialCompilationFinished().AddRaw(this, &FCustomPrimitiveDataCustomization::OnMaterialCompiled);
+}
+
 FCustomPrimitiveDataCustomization::~FCustomPrimitiveDataCustomization()
 {
 	Cleanup();
+
+	FCoreUObjectDelegates::OnObjectPropertyChanged.RemoveAll(this);
+	UMaterial::OnMaterialCompilationFinished().RemoveAll(this);
 }
 
 void FCustomPrimitiveDataCustomization::CustomizeHeader(TSharedRef<IPropertyHandle> PropertyHandle, FDetailWidgetRow& HeaderRow, IPropertyTypeCustomizationUtils& CustomizationUtils)
@@ -74,13 +84,9 @@ void FCustomPrimitiveDataCustomization::CustomizeChildren(TSharedRef<IPropertyHa
 	uint32 NumElements;
 	FPropertyAccess::Result AccessResult = GetNumElements(NumElements);
 
-	FSimpleDelegate OnElemsChanged = FSimpleDelegate::CreateRaw(this, &FCustomPrimitiveDataCustomization::OnUpdated);
+	FSimpleDelegate OnElemsChanged = FSimpleDelegate::CreateSP(this, &FCustomPrimitiveDataCustomization::RequestRefresh);
 	DataArrayHandle->SetOnNumElementsChanged(OnElemsChanged);
-	DataHandle->SetOnPropertyValueChanged(FSimpleDelegate::CreateRaw(this, &FCustomPrimitiveDataCustomization::OnElementsModified, AccessResult, NumElements));
-
-	// NOTE: Optimally would be bound to a "OnMaterialChanged" for each component
-	FCoreUObjectDelegates::OnObjectPropertyChanged.AddRaw(this, &FCustomPrimitiveDataCustomization::OnObjectPropertyChanged);
-	UMaterial::OnMaterialCompilationFinished().AddRaw(this, &FCustomPrimitiveDataCustomization::OnMaterialCompiled);
+	DataHandle->SetOnPropertyValueChanged(FSimpleDelegate::CreateSP(this, &FCustomPrimitiveDataCustomization::OnElementsModified, AccessResult, NumElements));
 
 	const int32 NumPrimitiveIndices = FMath::Max(MaxPrimitiveDataIndex + 1, (int32)NumElements);
 
@@ -423,9 +429,6 @@ void FCustomPrimitiveDataCustomization::ForEachSelectedComponent(Predicate Pred)
 
 void FCustomPrimitiveDataCustomization::Cleanup()
 {
-	FCoreUObjectDelegates::OnObjectPropertyChanged.RemoveAll(this);
-	UMaterial::OnMaterialCompilationFinished().RemoveAll(this);
-
 	PropertyUtils = NULL;
 	DataHandle = NULL;
 	DataArrayHandle = NULL;
@@ -529,14 +532,26 @@ void FCustomPrimitiveDataCustomization::PopulateParameterData(UPrimitiveComponen
 	}
 }
 
-void FCustomPrimitiveDataCustomization::OnUpdated()
+void FCustomPrimitiveDataCustomization::RequestRefresh()
 {
-	if (!bIgnoreUpdate && PropertyUtils.IsValid())
+	if (!bDeferringRefresh && PropertyUtils.IsValid())
 	{
-		PropertyUtils->ForceRefresh();
+		bDeferringRefresh = true;
+
+		PropertyUtils->EnqueueDeferredAction(FSimpleDelegate::CreateSP(this, &FCustomPrimitiveDataCustomization::OnDeferredRefresh));
 	}
 }
 	
+void FCustomPrimitiveDataCustomization::OnDeferredRefresh()
+{
+	if (PropertyUtils.IsValid())
+	{
+		PropertyUtils->ForceRefresh();
+	}
+
+	bDeferringRefresh = false;
+}
+
 void FCustomPrimitiveDataCustomization::OnElementsModified(const FPropertyAccess::Result OldAccessResult, const uint32 OldNumElements)
 {
 	uint32 NumElements;
@@ -545,7 +560,7 @@ void FCustomPrimitiveDataCustomization::OnElementsModified(const FPropertyAccess
 	// There's been a change in our array structure, whether that be from change in access or size
 	if (AccessResult != OldAccessResult || NumElements != OldNumElements)
 	{
-		OnUpdated();
+		RequestRefresh();
 	}
 }
 
@@ -601,7 +616,7 @@ void FCustomPrimitiveDataCustomization::OnObjectPropertyChanged(UObject* Object,
 
 		if (bMaterialChange)
 		{
-			OnUpdated();
+			RequestRefresh();
 		}
 	}
 }
@@ -611,7 +626,7 @@ void FCustomPrimitiveDataCustomization::OnMaterialCompiled(UMaterialInterface* M
 	// NOTE: We use a soft object ptr here as the old material object will be stale on compile
 	if (MaterialsToWatch.Contains(Material))
 	{
-		OnUpdated();
+		RequestRefresh();
 	}
 }
 
@@ -647,8 +662,6 @@ void FCustomPrimitiveDataCustomization::OnAddedDesiredPrimitiveData(uint8 PrimId
 	{
 		GEditor->BeginTransaction(LOCTEXT("OnAddedDesiredPrimitiveData", "Added Items"));
 
-		bIgnoreUpdate = true;
-
 		for (int32 i = NumElements; i <= PrimIdx; ++i)
 		{
 			DataArrayHandle->AddItem();
@@ -656,10 +669,6 @@ void FCustomPrimitiveDataCustomization::OnAddedDesiredPrimitiveData(uint8 PrimId
 			SetDefaultValue(DataArrayHandle->GetElement(i), i);
 		}
 
-		bIgnoreUpdate = false;
-
-		OnUpdated();
-
 		GEditor->EndTransaction();
 	}
 }
@@ -671,17 +680,11 @@ void FCustomPrimitiveDataCustomization::OnRemovedPrimitiveData(uint8 PrimIdx)
 	{
 		GEditor->BeginTransaction(LOCTEXT("OnRemovedPrimitiveData", "Removed Items"));
 
-		bIgnoreUpdate = true;
-
 		for (int32 i = NumElements - 1; i >= PrimIdx; --i)
 		{
 			DataArrayHandle->DeleteItem(i);
 		}
 
-		bIgnoreUpdate = false;
-
-		OnUpdated();
-
 		GEditor->EndTransaction();
 	}
 }
diff --git a/Engine/Source/Editor/DetailCustomizations/Private/CustomPrimitiveDataCustomization.h b/Engine/Source/Editor/DetailCustomizations/Private/CustomPrimitiveDataCustomization.h
--- a/Engine/Source/Editor/DetailCustomizations/Private/CustomPrimitiveDataCustomization.h
+++ b/Engine/Source/Editor/DetailCustomizations/Private/CustomPrimitiveDataCustomization.h
@@ -22,6 +22,7 @@ class FCustomPrimitiveDataCustomization : public IPropertyTypeCustomization
 public:
 	static TSharedRef<IPropertyTypeCustomization> MakeInstance();
 
+	FCustomPrimitiveDataCustomization();
 	~FCustomPrimitiveDataCustomization();
 
 	/** IPropertyTypeCustomization interface */
@@ -39,7 +40,7 @@ private:
 		uint8 IndexOffset;
 	};
 
-	bool bIgnoreUpdate = false;
+	bool bDeferringRefresh = false;
 
 	TSharedPtr<IPropertyUtilities> PropertyUtils;
 	TSharedPtr<IPropertyHandle> DataHandle;
@@ -76,7 +77,8 @@ private:
 	/** Populates our parameter data for scalar and vector parameters. Updates the max primitive index if necessary. */
 	void PopulateParameterData(UPrimitiveComponent* PrimitiveComponent, int32& MaxPrimitiveDataIndex);
 
-	void OnUpdated();
+	void RequestRefresh();
+	void OnDeferredRefresh();
 	void OnElementsModified(const enum FPropertyAccess::Result OldAccessResult, const uint32 OldNumElements);
 	void OnObjectPropertyChanged(UObject* Object, FPropertyChangedEvent& PropertyChangedEvent);
 	void OnMaterialCompiled(UMaterialInterface* Material);
