Menu Searching for Material Editor Dropdown

diff --git a/Engine/Source/Editor/DetailCustomizations/Private/MaterialAttributePropertyDetails.cpp b/Engine/Source/Editor/DetailCustomizations/Private/MaterialAttributePropertyDetails.cpp
--- a/Engine/Source/Editor/DetailCustomizations/Private/MaterialAttributePropertyDetails.cpp
+++ b/Engine/Source/Editor/DetailCustomizations/Private/MaterialAttributePropertyDetails.cpp
@@ -57,60 +57,53 @@ void FMaterialAttributePropertyDetails::OnBuildChild(TSharedRef<IPropertyHandle>
 	// Add an overridden combo box
 	IDetailPropertyRow& PropertyArrayRow = ChildrenBuilder.AddProperty(ChildHandle);
 
-	PropertyArrayRow.CustomWidget()
-	.NameContent()
-	[
-		ChildHandle->CreatePropertyNameWidget()
-	]
-	.ValueContent()
-	[
-		SNew(SHorizontalBox)
-		+SHorizontalBox::Slot()
-		[
-			SNew(SComboBox<TSharedPtr<FString>>)
-			.OptionsSource(&AttributeDisplayNameList)
- 			.OnGenerateWidget_Lambda( [] (TSharedPtr<FString> InItem)
+	FPropertyComboBoxArgs ComboArgs(
+		ChildHandle,
+		FOnGetPropertyComboBoxStrings::CreateLambda([=](TArray<TSharedPtr<FString>>& OutComboBoxStrings, TArray<TSharedPtr<SToolTip>>& OutToolTips, TArray<bool>& OutRestrictedItems) -> void  
+		{									
+			OutComboBoxStrings = AttributeDisplayNameList;
+		}), 
+		FOnGetPropertyComboBoxValue::CreateLambda( [=]() -> FString
+		{
+			FString AttributeName;
+			if (ChildHandle->IsValidHandle())
 			{
-				return SNew(STextBlock)
-				.Font(IDetailLayoutBuilder::GetDetailFont())
-				.Text(FText::FromString(*InItem));
-			})
-			.OnSelectionChanged_Lambda( [=] (TSharedPtr<FString> Selection, ESelectInfo::Type)
+				// Convert attribute ID string to display name
+				FString IDString; FGuid IDValue;
+				ChildHandle->GetValueAsFormattedString(IDString);
+				FGuid::ParseExact(IDString, EGuidFormats::Digits, IDValue);
+
+				AttributeName = FMaterialAttributeDefinitionMap::GetAttributeName(IDValue);
+			} 
+			return AttributeName;
+		}),
+		FOnPropertyComboBoxValueSelected::CreateLambda( [=] (const FString& Selection)
+		{
+			if (ChildHandle->IsValidHandle())
 			{
-				if (ChildHandle->IsValidHandle())
+				// Convert display name to attribute ID
+				for (const auto& NameIDPair : AttributeNameToIDList)
 				{
-					// Convert display name to attribute ID
-					for (const auto& NameIDPair : AttributeNameToIDList)
+					if (NameIDPair.Key == Selection)
 					{
-						if (NameIDPair.Key == *Selection)
-						{
-							ChildHandle->SetValueFromFormattedString(NameIDPair.Value.ToString(EGuidFormats::Digits));
-							break;
-						}
+						ChildHandle->SetValueFromFormattedString(NameIDPair.Value.ToString(EGuidFormats::Digits));
+						break;
 					}
 				}
-			})
-			.ContentPadding(FMargin(2, 0))
-			[
-				SNew(STextBlock)
-				.Font(IDetailLayoutBuilder::GetDetailFont())
-				.Text_Lambda( [=]() -> FText
-				{ 
-					if (ChildHandle->IsValidHandle())
-					{
-						// Convert attribute ID string to display name
-						FString IDString; FGuid IDValue;
-						ChildHandle->GetValueAsFormattedString(IDString);
-						FGuid::ParseExact(IDString, EGuidFormats::Digits, IDValue);
+			}
+		})
+	);
+	ComboArgs.ShowSearchForItemCount = 1;
 
-						FString AttributeName = FMaterialAttributeDefinitionMap::GetAttributeName(IDValue);
-						return FText::FromString(AttributeName);
-					}
 
-					return FText::GetEmpty();
-				} )
-			]
-		]
+	PropertyArrayRow.CustomWidget()
+	.NameContent()
+	[
+		ChildHandle->CreatePropertyNameWidget()
+	]
+	.ValueContent()
+	[
+		PropertyCustomizationHelpers::MakePropertyComboBox(ComboArgs)
 	];
 }
 
