DetailCustomization - Add Streaming section for World Partition & co (Data layers, HLOD, RuntimeHash, FoliageGrid).
HLODLayer - can edit only in World Partition world.

diff --git a/Engine/Source/Editor/DetailCustomizations/Private/DetailCustomizations.cpp b/Engine/Source/Editor/DetailCustomizations/Private/DetailCustomizations.cpp
--- a/Engine/Source/Editor/DetailCustomizations/Private/DetailCustomizations.cpp
+++ b/Engine/Source/Editor/DetailCustomizations/Private/DetailCustomizations.cpp
@@ -455,7 +455,12 @@ void FDetailCustomizationsModule::RegisterSectionMappings()
 			Section->AddCategory("Cooking");
 			Section->AddCategory("Input");
 			Section->AddCategory("Replication");
+		}
+
+		{
+			TSharedRef<FPropertySection> Section = PropertyModule.FindOrCreateSection("Actor", "Streaming", LOCTEXT("Streaming", "Streaming"));
 			Section->AddCategory("World Partition");
+			Section->AddCategory("HLOD");
 		}
 	}
 
@@ -474,9 +479,13 @@ void FDetailCustomizationsModule::RegisterSectionMappings()
 			TSharedRef<FPropertySection> Section = PropertyModule.FindOrCreateSection("ActorComponent", "Misc", LOCTEXT("Misc", "Misc"));
 			Section->AddCategory("Asset User Data");
 			Section->AddCategory("Cooking");
-			Section->AddCategory("Data Layers");
 			Section->AddCategory("Tags");
+		}
+
+		{
+			TSharedRef<FPropertySection> Section = PropertyModule.FindOrCreateSection("ActorComponent", "Streaming", LOCTEXT("Streaming", "Streaming"));
 			Section->AddCategory("World Partition");
+			Section->AddCategory("Data Layers");
 		}
 	}
 
@@ -731,6 +740,22 @@ void FDetailCustomizationsModule::RegisterSectionMappings()
 			Section->AddCategory("Brush Settings");
 		}
 	}
+
+	// WorldPartition
+	{
+		{
+			TSharedRef<FPropertySection> Section = PropertyModule.FindOrCreateSection("WorldPartition", "Streaming", LOCTEXT("Streaming", "Streaming"));
+			Section->AddCategory("RuntimeHash");
+		}
+	}
+
+	// WorldSettings
+	{
+		{
+			TSharedRef<FPropertySection> Section = PropertyModule.FindOrCreateSection("WorldSettings", "Streaming", LOCTEXT("Streaming", "Streaming"));
+			Section->AddCategory("Foliage");
+		}
+	}
 }
 
 #undef LOCTEXT_NAMESPACE
diff --git a/Engine/Source/Runtime/Engine/Private/ActorEditor.cpp b/Engine/Source/Runtime/Engine/Private/ActorEditor.cpp
--- a/Engine/Source/Runtime/Engine/Private/ActorEditor.cpp
+++ b/Engine/Source/Runtime/Engine/Private/ActorEditor.cpp
@@ -71,8 +71,9 @@ bool AActor::CanEditChange(const FProperty* PropertyThatWillChange) const
 	const bool bIsGridPlacement = PropertyThatWillChange->GetFName() == GET_MEMBER_NAME_CHECKED(AActor, GridPlacement);
 	const bool bIsRuntimeGrid = PropertyThatWillChange->GetFName() == GET_MEMBER_NAME_CHECKED(AActor, RuntimeGrid);
 	const bool bIsDataLayers = PropertyThatWillChange->GetFName() == GET_MEMBER_NAME_CHECKED(AActor, DataLayers);
+	const bool bIsHLODLayer = PropertyThatWillChange->GetFName() == GET_MEMBER_NAME_CHECKED(AActor, HLODLayer);
 
-	if (bIsGridPlacement || bIsRuntimeGrid || bIsDataLayers)
+	if (bIsGridPlacement || bIsRuntimeGrid || bIsDataLayers || bIsHLODLayer)
 	{
 		if (!IsTemplate())
 		{
