ADDED: Custom Primitive Data Customization
ADDED: Set/Get Custom Primitive Data by [REDACTED_PERSON]
ADDED: Primitive Index to Material Cached Data for runtime use

diff --git a/Engine/Source/Editor/DetailCustomizations/Private/DetailCustomizations.cpp b/Engine/Source/Editor/DetailCustomizations/Private/DetailCustomizations.cpp
--- a/Engine/Source/Editor/DetailCustomizations/Private/DetailCustomizations.cpp
+++ b/Engine/Source/Editor/DetailCustomizations/Private/DetailCustomizations.cpp
@@ -149,7 +149,7 @@
 #include "WindowsTargetSettingsDetails.h"
 #include "WorldSettingsDetails.h"
 #include "MaterialInstanceDynamicDetails.h"
-
+#include "CustomPrimitiveDataCustomization.h"
 
 IMPLEMENT_MODULE( FDetailCustomizationsModule, DetailCustomizations );
 
