Reorder include statement to fix bad merge

#rnx

diff --git a/Engine/Source/Editor/DetailCustomizations/Private/DetailCustomizations.cpp b/Engine/Source/Editor/DetailCustomizations/Private/DetailCustomizations.cpp
--- a/Engine/Source/Editor/DetailCustomizations/Private/DetailCustomizations.cpp
+++ b/Engine/Source/Editor/DetailCustomizations/Private/DetailCustomizations.cpp
@@ -148,8 +148,6 @@
 #include "VectorStructCustomization.h"
 #include "WindowsTargetSettingsDetails.h"
 #include "WorldSettingsDetails.h"
-#include "MaterialInstanceDynamicDetails.h"
-#include "CustomPrimitiveDataCustomization.h"
 
 IMPLEMENT_MODULE( FDetailCustomizationsModule, DetailCustomizations );
 
