FIXED: Crash due to hitting stale object in ComponentsToWatch during automated test in EngineTest

diff --git a/Engine/Source/Editor/DetailCustomizations/Private/CustomPrimitiveDataCustomization.cpp b/Engine/Source/Editor/DetailCustomizations/Private/CustomPrimitiveDataCustomization.cpp
--- a/Engine/Source/Editor/DetailCustomizations/Private/CustomPrimitiveDataCustomization.cpp
+++ b/Engine/Source/Editor/DetailCustomizations/Private/CustomPrimitiveDataCustomization.cpp
@@ -427,6 +427,16 @@ void FCustomPrimitiveDataCustomization::ForEachSelectedComponent(Predicate Pred)
 	}
 }
 
+bool FCustomPrimitiveDataCustomization::IsSelected(UPrimitiveComponent* Component) const
+{
+	return Component && PropertyUtils.IsValid() && PropertyUtils->GetSelectedObjects().ContainsByPredicate(
+		[WeakComp = TWeakObjectPtr<UObject>(Component), WeakActor = TWeakObjectPtr<UObject>(Component->GetOwner())](const TWeakObjectPtr<UObject>& SelectedObject)
+		{
+			// Selected objects could be components or actors
+			return SelectedObject.IsValid() && (SelectedObject == WeakComp || SelectedObject == WeakActor);
+		});
+}
+
 void FCustomPrimitiveDataCustomization::Cleanup()
 {
 	PropertyUtils = NULL;
@@ -566,17 +576,19 @@ void FCustomPrimitiveDataCustomization::OnElementsModified(const FPropertyAccess
 
 void FCustomPrimitiveDataCustomization::OnObjectPropertyChanged(UObject* Object, FPropertyChangedEvent& PropertyChangedEvent)
 {
+	UPrimitiveComponent* PrimComponent = Cast<UPrimitiveComponent>(Object);
 	const EPropertyChangeType::Type IgnoreFlags = EPropertyChangeType::Interactive | EPropertyChangeType::Redirected;
 
-	if (!(PropertyChangedEvent.ChangeType & IgnoreFlags) && ComponentsToWatch.Contains(Object))
+	if (!(PropertyChangedEvent.ChangeType & IgnoreFlags) && ComponentsToWatch.Contains(PrimComponent)
+		&& IsSelected(PrimComponent)) // Need to test this in case we're hitting a stale hash in ComponentsToWatch (#jira UE-136687)
 	{
 		bool bMaterialChange = false;
 
-		if (Object->IsA<UMeshComponent>())
+		if (PrimComponent->IsA<UMeshComponent>())
 		{
 			bMaterialChange = PropertyChangedEvent.GetPropertyName() == GET_MEMBER_NAME_CHECKED(UMeshComponent, OverrideMaterials);
 		}
-		else if (Object->IsA<UTextRenderComponent>())
+		else if (PrimComponent->IsA<UTextRenderComponent>())
 		{
 			bMaterialChange = PropertyChangedEvent.GetPropertyName() == GET_MEMBER_NAME_CHECKED(UTextRenderComponent, TextMaterial);
 		}
@@ -585,9 +597,8 @@ void FCustomPrimitiveDataCustomization::OnObjectPropertyChanged(UObject* Object,
 			// Fall back if not handled
 			// NOTE: Optimally would be done via an "OnMaterialChanged" for each component, however,
 			// the property name checks above should handle most cases
-			TSet<TSoftObjectPtr<UMaterial>>& CachedComponentMaterials = ComponentsToWatch[Object];
+			TSet<TSoftObjectPtr<UMaterial>>& CachedComponentMaterials = ComponentsToWatch[PrimComponent];
 
-			const UPrimitiveComponent* PrimComponent = CastChecked<UPrimitiveComponent>(Object);
 			const int32 NumMaterials = PrimComponent->GetNumMaterials();
 
 			if (NumMaterials != CachedComponentMaterials.Num())
diff --git a/Engine/Source/Editor/DetailCustomizations/Private/CustomPrimitiveDataCustomization.h b/Engine/Source/Editor/DetailCustomizations/Private/CustomPrimitiveDataCustomization.h
--- a/Engine/Source/Editor/DetailCustomizations/Private/CustomPrimitiveDataCustomization.h
+++ b/Engine/Source/Editor/DetailCustomizations/Private/CustomPrimitiveDataCustomization.h
@@ -46,7 +46,7 @@ private:
 	TSharedPtr<IPropertyHandle> DataHandle;
 	TSharedPtr<IPropertyHandleArray> DataArrayHandle;
 
-	TMap<TWeakObjectPtr<UObject>, TSet<TSoftObjectPtr<UMaterial>>> ComponentsToWatch;
+	TMap<TWeakObjectPtr<UPrimitiveComponent>, TSet<TSoftObjectPtr<UMaterial>>> ComponentsToWatch;
 	TSet<TSoftObjectPtr<UMaterialInterface>> MaterialsToWatch;
 
 	TMap<uint8, TArray<FParameterData>> VectorParameterData;
@@ -71,6 +71,7 @@ private:
 
 	template<typename Predicate>
 	void ForEachSelectedComponent(Predicate Pred);
+	bool IsSelected(UPrimitiveComponent* Component) const;
 
 	void Cleanup();
 
