Add customizable names to FPerPlatformPropertyCustomNodeBuilder

diff --git a/Engine/Source/Editor/DetailCustomizations/Private/PerPlatformPropertyCustomization.cpp b/Engine/Source/Editor/DetailCustomizations/Private/PerPlatformPropertyCustomization.cpp
--- a/Engine/Source/Editor/DetailCustomizations/Private/PerPlatformPropertyCustomization.cpp
+++ b/Engine/Source/Editor/DetailCustomizations/Private/PerPlatformPropertyCustomization.cpp
@@ -360,8 +360,7 @@ void FPerPlatformPropertyCustomNodeBuilder::GenerateChildContent(IDetailChildren
 
 FName FPerPlatformPropertyCustomNodeBuilder::GetName() const
 {
-	static const FName Name("FPerPlatformPropertyCustomNodeBuilder");
-	return Name;
+	return Args.Name;
 }
 
 void FPerPlatformPropertyCustomNodeBuilder::OnAddPlatformOverride(const FName PlatformName)
diff --git a/Engine/Source/Editor/DetailCustomizations/Public/PerPlatformPropertyCustomization.h b/Engine/Source/Editor/DetailCustomizations/Public/PerPlatformPropertyCustomization.h
--- a/Engine/Source/Editor/DetailCustomizations/Public/PerPlatformPropertyCustomization.h
+++ b/Engine/Source/Editor/DetailCustomizations/Public/PerPlatformPropertyCustomization.h
@@ -32,7 +32,11 @@ struct FPerPlatformPropertyCustomNodeBuilderArgs
 
 	FOnPlatformOverrideAction OnRemovePlatformOverride;
 
+	/** Used for search box queries */
 	FText FilterText;
+
+	/** Used to identify the builder through GetName() */
+	FName Name;
 };
 
 class DETAILCUSTOMIZATIONS_API FPerPlatformPropertyCustomNodeBuilder : public IDetailCustomNodeBuilder, public TSharedFromThis<FPerPlatformPropertyCustomNodeBuilder>
diff --git a/Engine/Source/Editor/StaticMeshEditor/Private/StaticMeshEditorTools.cpp b/Engine/Source/Editor/StaticMeshEditor/Private/StaticMeshEditorTools.cpp
--- a/Engine/Source/Editor/StaticMeshEditor/Private/StaticMeshEditorTools.cpp
+++ b/Engine/Source/Editor/StaticMeshEditor/Private/StaticMeshEditorTools.cpp
@@ -3256,6 +3256,7 @@ void FLevelOfDetailSettingsLayout::AddToDetailsPanel( IDetailLayoutBuilder& Deta
 
 		FPerPlatformPropertyCustomNodeBuilderArgs Args;
 		Args.FilterText = LOCTEXT("MinLOD", "Minimum LOD");
+		Args.Name = "MinLod";
 		Args.OnGenerateNameWidget = FOnGetContent::CreateLambda([]()
 			{
 				return SNew(STextBlock)
