Fixes Crash after Undo for AudioVolume assets

diff --git a/Engine/Source/Editor/DetailCustomizations/Private/CurveStructCustomization.cpp b/Engine/Source/Editor/DetailCustomizations/Private/CurveStructCustomization.cpp
--- a/Engine/Source/Editor/DetailCustomizations/Private/CurveStructCustomization.cpp
+++ b/Engine/Source/Editor/DetailCustomizations/Private/CurveStructCustomization.cpp
@@ -38,6 +38,7 @@ FCurveStructCustomization::~FCurveStructCustomization()
 	}
 
 	DestroyPopOutWindow();
+	GEditor->UnregisterForUndo(this);
 }
 
 FCurveStructCustomization::FCurveStructCustomization()
@@ -46,6 +47,7 @@ FCurveStructCustomization::FCurveStructCustomization()
 	, ViewMinInput(0.0f)
 	, ViewMaxInput(5.0f)
 {
+	GEditor->RegisterForUndo(this);
 }
 
 void FCurveStructCustomization::CustomizeHeader( TSharedRef<IPropertyHandle> InStructPropertyHandle, FDetailWidgetRow& HeaderRow, IPropertyTypeCustomizationUtils& StructCustomizationUtils )
@@ -120,19 +122,19 @@ void FCurveStructCustomization::CustomizeHeader( TSharedRef<IPropertyHandle> InS
 	else
 	{
 		HeaderRow
-			.NameContent()
+		.NameContent()
+		[
+			InStructPropertyHandle->CreatePropertyNameWidget()
+		]
+		.ValueContent()
+		[
+			SNew(SBorder)
+			.VAlign(VAlign_Fill)
 			[
-				InStructPropertyHandle->CreatePropertyNameWidget()
+				SNew(STextBlock)
+				.Text(StructPtrs.Num() == 0 ? LOCTEXT("NoCurves", "No Curves - unable to modify") : LOCTEXT("MultipleCurves", "Multiple Curves - unable to modify"))
 			]
-			.ValueContent()
-			[
-				SNew(SBorder)
-				.VAlign(VAlign_Fill)
-				[
-					SNew(STextBlock)
-					.Text(StructPtrs.Num() == 0 ? LOCTEXT("NoCurves", "No Curves - unable to modify") : LOCTEXT("MultipleCurves", "Multiple Curves - unable to modify"))
-				]
-			];
+		];
 	}
 }
 
@@ -262,6 +264,31 @@ bool FCurveStructCustomization::IsValidCurve( FRichCurveEditInfo CurveInfo )
 	return CurveInfo.CurveToEdit == &RuntimeCurve->EditorCurveData;
 }
 
+void FCurveStructCustomization::PostUndo(bool bSuccess)
+{
+	// reset the cached curves  
+	TArray<UObject*> OuterObjects;
+	StructPropertyHandle->GetOuterObjects(OuterObjects);
+
+	TArray<void*> StructPtrs;
+	StructPropertyHandle->AccessRawData( StructPtrs );
+	if (StructPtrs.Num() == 1)
+	{
+		RuntimeCurve = reinterpret_cast<FRuntimeFloatCurve*>(StructPtrs[0]);
+		if (RuntimeCurve)
+		{
+			if (RuntimeCurve->ExternalCurve)
+			{
+				CurveWidget->SetCurveOwner(RuntimeCurve->ExternalCurve, false);
+			}
+			else
+			{
+				CurveWidget->SetCurveOwner(this, StructPropertyHandle->IsEditable());
+			}
+		}
+	}
+}
+
 float FCurveStructCustomization::GetTimelineLength() const
 {
 	return 0.f;
diff --git a/Engine/Source/Editor/DetailCustomizations/Private/CurveStructCustomization.h b/Engine/Source/Editor/DetailCustomizations/Private/CurveStructCustomization.h
--- a/Engine/Source/Editor/DetailCustomizations/Private/CurveStructCustomization.h
+++ b/Engine/Source/Editor/DetailCustomizations/Private/CurveStructCustomization.h
@@ -16,7 +16,7 @@ struct FRuntimeFloatCurve;
 /**
  * Customizes a RuntimeFloatCurve struct to display a Curve Editor
  */
-class FCurveStructCustomization : public IPropertyTypeCustomization, public FCurveOwnerInterface
+class FCurveStructCustomization : public IPropertyTypeCustomization, public FCurveOwnerInterface, public FEditorUndoClient
 {
 public:
 	static TSharedRef<IPropertyTypeCustomization> MakeInstance();
@@ -39,6 +39,10 @@ public:
 	virtual void OnCurveChanged(const TArray<FRichCurveEditInfo>& ChangedCurveEditInfos) override;
 	virtual bool IsValidCurve( FRichCurveEditInfo CurveInfo ) override;
 
+	//~ FEditorUndoClient interface
+	virtual void PostUndo(bool bSuccess) override;
+	virtual void PostRedo(bool bSuccess) override { PostUndo(bSuccess); }
+
 private:
 	/**
 	 * Constructor
