FIXED: Undeclared CPD warning wasn't showing for a single component with multiple materials
MODIFY: Undeclared warning now includes warning icon
ADDED: Warning box for overlapping CPD parameters
ADDED: Note for usage of vector channel params in setters

diff --git a/Engine/Source/Editor/DetailCustomizations/Private/CustomPrimitiveDataCustomization.cpp b/Engine/Source/Editor/DetailCustomizations/Private/CustomPrimitiveDataCustomization.cpp
--- a/Engine/Source/Editor/DetailCustomizations/Private/CustomPrimitiveDataCustomization.cpp
+++ b/Engine/Source/Editor/DetailCustomizations/Private/CustomPrimitiveDataCustomization.cpp
@@ -5,6 +5,7 @@
 #include "Widgets/Layout/SBox.h"
 #include "Widgets/Text/STextBlock.h"
 #include "Widgets/Input/SHyperlink.h"
+#include "Widgets/Images/SImage.h"
 #include "Widgets/Colors/SColorPicker.h"
 #include "Widgets/Colors/SColorBlock.h"
 #include "SlateOptMacros.h"
@@ -72,13 +73,11 @@ void FCustomPrimitiveDataCustomization::CustomizeChildren(TSharedRef<IPropertyHa
 	DataHandle = PropertyHandle->GetChildHandle("Data");
 	DataArrayHandle = DataHandle->AsArray();
 
-	int32 NumSelectedComponents = 0;
 	int32 MaxPrimitiveDataIndex = INDEX_NONE;
 
 	ForEachSelectedComponent([&](UPrimitiveComponent* Component)
 	{
 		PopulateParameterData(Component, MaxPrimitiveDataIndex);
-		NumSelectedComponents++;
 	});
 
 	uint32 NumElements;
@@ -128,13 +127,13 @@ void FCustomPrimitiveDataCustomization::CustomizeChildren(TSharedRef<IPropertyHa
 			{
 				// Create a collapsing group that contains our color picker, so we can quickly assign colors to our vector
 				VectorGroupPrimIdx = PrimIdx;
-				VectorGroup = CreateVectorGroup(ChildBuilder, PrimIdx, bDataEditable, NumElements);
+				VectorGroup = CreateVectorGroup(ChildBuilder, PrimIdx, bDataEditable, NumElements, CustomizationUtils);
 			}
 		}
 
 		if (ScalarParameterData.Contains(PrimIdx) || VectorParameterData.Contains(PrimIdx))
 		{
-			CreateParameterRow(ChildBuilder, PrimIdx, ElementHandle, NumSelectedComponents, bDataEditable, VectorGroup, CustomizationUtils);
+			CreateParameterRow(ChildBuilder, PrimIdx, ElementHandle, bDataEditable, VectorGroup, CustomizationUtils);
 		}
 		else
 		{
@@ -147,6 +146,7 @@ void FCustomPrimitiveDataCustomization::CustomizeChildren(TSharedRef<IPropertyHa
 				ChildBuilder.AddProperty(ElementHandle.ToSharedRef())
 				.CustomWidget()
 				.NameContent()
+				.HAlign(HAlign_Fill)
 				[
 					NameWidget
 				]
@@ -160,6 +160,7 @@ void FCustomPrimitiveDataCustomization::CustomizeChildren(TSharedRef<IPropertyHa
 			{
 				ChildBuilder.AddCustomRow(FText::AsNumber(PrimIdx))
 				.NameContent()
+				.HAlign(HAlign_Fill)
 				[
 					NameWidget
 				]
@@ -170,7 +171,7 @@ void FCustomPrimitiveDataCustomization::CustomizeChildren(TSharedRef<IPropertyHa
 }
 
 BEGIN_SLATE_FUNCTION_BUILD_OPTIMIZATION
-IDetailGroup* FCustomPrimitiveDataCustomization::CreateVectorGroup(IDetailChildrenBuilder& ChildBuilder, uint8 PrimIdx, bool bDataEditable, int32 NumElements)
+IDetailGroup* FCustomPrimitiveDataCustomization::CreateVectorGroup(IDetailChildrenBuilder& ChildBuilder, uint8 PrimIdx, bool bDataEditable, int32 NumElements, IPropertyTypeCustomizationUtils& CustomizationUtils)
 {
 	IDetailGroup* VectorGroup = &ChildBuilder.AddGroup(VectorParameterData[PrimIdx][0].Info.Name, FText::FromName(VectorParameterData[PrimIdx][0].Info.Name));
 
@@ -179,6 +180,8 @@ IDetailGroup* FCustomPrimitiveDataCustomization::CreateVectorGroup(IDetailChildr
 
 	// Use this to make sure we don't make duplicate parameters for the group header
 	TSet<FGuid> AddedParametersForThisGroup;
+	TSet<FName> ParameterNames;
+	TSet<TSoftObjectPtr<UMaterialInterface>> Materials;
 
 	for (const FParameterData& ParameterData : VectorParameterData[PrimIdx])
 	{
@@ -187,7 +190,9 @@ IDetailGroup* FCustomPrimitiveDataCustomization::CreateVectorGroup(IDetailChildr
 			continue;
 		}
 
+		Materials.Add(ParameterData.Material->GetMaterial());
 		AddedParametersForThisGroup.Add(ParameterData.ExpressionID);
+		ParameterNames.Add(ParameterData.Info.Name);
 
 		TSharedRef<SHyperlink> Hyperlink = CreateHyperlink(FText::FromName(ParameterData.Info.Name), ParameterData.Material, ParameterData.ExpressionID);
 		Hyperlink->SetToolTipText(LOCTEXT("VectorHyperlinkTooltip", "Jump to Vector Parameter"));
@@ -199,10 +204,31 @@ IDetailGroup* FCustomPrimitiveDataCustomization::CreateVectorGroup(IDetailChildr
 		];
 	}
 
+	if (MaterialsToWatch.Num() != Materials.Num())
+	{
+		// Some materials aren't defining parameters at this index, add the undeclared parameter widget in case this was user error
+		VectorGroupNameBox->AddSlot()
+		.Padding(2.f)
+		[
+			GetUndeclaredParameterWidget(PrimIdx, CustomizationUtils)
+		];
+	}
+
+	TSharedPtr<SWidget> NameContent;
+	if (ParameterNames.Num() > 1)
+	{
+		NameContent = CreateWarningWidget(VectorGroupNameBox.ToSharedRef(), LOCTEXT("OverlappingVectorParameters", "Primitive index has overlapping parameter names declared, make sure vector names match to remove warning"));
+	}
+	else
+	{
+		NameContent = VectorGroupNameBox;
+	}
+
 	VectorGroup->HeaderRow()
 	.NameContent()
+	.HAlign(HAlign_Fill)
 	[
-		VectorGroupNameBox.ToSharedRef()
+		NameContent.ToSharedRef()
 	]
 	.ValueContent()
 	[
@@ -253,7 +279,7 @@ IDetailGroup* FCustomPrimitiveDataCustomization::CreateVectorGroup(IDetailChildr
 END_SLATE_FUNCTION_BUILD_OPTIMIZATION
 
 BEGIN_SLATE_FUNCTION_BUILD_OPTIMIZATION
-void FCustomPrimitiveDataCustomization::CreateParameterRow(IDetailChildrenBuilder& ChildBuilder, uint8 PrimIdx, TSharedPtr<IPropertyHandle> ElementHandle, int32 NumSelectedComponents, bool bDataEditable, IDetailGroup* VectorGroup, IPropertyTypeCustomizationUtils& CustomizationUtils)
+void FCustomPrimitiveDataCustomization::CreateParameterRow(IDetailChildrenBuilder& ChildBuilder, uint8 PrimIdx, TSharedPtr<IPropertyHandle> ElementHandle, bool bDataEditable, IDetailGroup* VectorGroup, IPropertyTypeCustomizationUtils& CustomizationUtils)
 {
 	// Use this to make sure we don't make duplicate parameters in each row
 	TSet<FGuid> AddedParametersForThisRow;
@@ -261,19 +287,19 @@ void FCustomPrimitiveDataCustomization::CreateParameterRow(IDetailChildrenBuilde
 	TArray<FText> SearchText;
 	TSharedPtr<SVerticalBox> VerticalBox = SNew(SVerticalBox);
 
-	TSet<TWeakObjectPtr<UPrimitiveComponent>> Components;
+	TSet<UMaterial*> Materials;
+	TSet<FName> ParameterNames;
 
 	if (VectorParameterData.Contains(PrimIdx))
 	{
 		for (const FParameterData& ParameterData : VectorParameterData[PrimIdx])
 		{
-			Components.Add(ParameterData.Component);
-
 			if (AddedParametersForThisRow.Contains(ParameterData.ExpressionID))
 			{
 				continue;
 			}
 
+			Materials.Add(ParameterData.Material->GetMaterial());
 			AddedParametersForThisRow.Add(ParameterData.ExpressionID);
 
 			FMaterialParameterMetadata ParameterMetadata;
@@ -299,14 +325,22 @@ void FCustomPrimitiveDataCustomization::CreateParameterRow(IDetailChildrenBuilde
 				break;
 			}
 
+			ParameterNames.Add(*ChannelName.ToString());
+
+			const FText ParameterNameText = FText::FromName(ParameterData.Info.Name);
 			const FText ParameterName = FText::Format(
 				LOCTEXT("VectorParameterName", "{0}.{1}"),
-				FText::FromName(ParameterData.Info.Name),
+				ParameterNameText,
 				ChannelName
 			);
 
 			TSharedRef<SHyperlink> Hyperlink = CreateHyperlink(ParameterName, ParameterData.Material, ParameterData.ExpressionID);
-			Hyperlink->SetToolTipText(LOCTEXT("VectorChannelHyperlinkTooltip", "Jump to Vector Parameter Channel"));
+
+			const FText HyperlinkToolTipText = FText::Format(
+				LOCTEXT("VectorChannelHyperlinkTooltip", "Jump to Vector Parameter\nNOTE: Vector channel name not used as parameter name for setters,\nuse vector parameter name \"{0}\" or primitive index {1} instead."),
+				ParameterNameText,
+				FText::AsNumber(PrimIdx));
+			Hyperlink->SetToolTipText(HyperlinkToolTipText);
 
 			VerticalBox->AddSlot()
 			.Padding(2.f)
@@ -322,14 +356,14 @@ void FCustomPrimitiveDataCustomization::CreateParameterRow(IDetailChildrenBuilde
 	{
 		for (const FParameterData& ParameterData : ScalarParameterData[PrimIdx])
 		{
-			Components.Add(ParameterData.Component);
-
 			if (AddedParametersForThisRow.Contains(ParameterData.ExpressionID))
 			{
 				continue;
 			}
 
+			Materials.Add(ParameterData.Material->GetMaterial());
 			AddedParametersForThisRow.Add(ParameterData.ExpressionID);
+			ParameterNames.Add(ParameterData.Info.Name);
 
 			const FText ParameterName = FText::FromName(ParameterData.Info.Name);
 
@@ -346,7 +380,7 @@ void FCustomPrimitiveDataCustomization::CreateParameterRow(IDetailChildrenBuilde
 		}
 	}
 
-	if (Components.Num() != NumSelectedComponents)
+	if (MaterialsToWatch.Num() != Materials.Num())
 	{
 		// Some components aren't defining parameters at this index, add the undeclared parameter widget in case this was user error
 		VerticalBox->AddSlot()
@@ -356,6 +390,18 @@ void FCustomPrimitiveDataCustomization::CreateParameterRow(IDetailChildrenBuilde
 		];
 	}
 
+	TSharedPtr<SWidget> NameContent;
+	if (ParameterNames.Num() > 1)
+	{
+		NameContent = CreateWarningWidget(VerticalBox.ToSharedRef(), LOCTEXT("OverlappingParameters", "Primitive index has overlapping parameter names declared, make sure scalar and/or vector channel names match to remove warning"));
+	}
+	else
+	{
+		NameContent = VerticalBox;
+	}
+
+	NameContent = CreateNameWidget(PrimIdx, NameContent.ToSharedRef(), CustomizationUtils);
+
 	if (ElementHandle.IsValid())
 	{
 		// We already have data for this row, be sure to use it
@@ -368,8 +414,9 @@ void FCustomPrimitiveDataCustomization::CreateParameterRow(IDetailChildrenBuilde
 
 		Row.CustomWidget()
 		.NameContent()
+		.HAlign(HAlign_Fill)
 		[
-			CreateNameWidget(PrimIdx, VerticalBox.ToSharedRef(), CustomizationUtils)
+			NameContent.ToSharedRef()
 		]
 		.ValueContent()
 		[
@@ -382,8 +429,9 @@ void FCustomPrimitiveDataCustomization::CreateParameterRow(IDetailChildrenBuilde
 		FDetailWidgetRow& Row = VectorGroup ? VectorGroup->AddWidgetRow() : ChildBuilder.AddCustomRow(FText::Join(LOCTEXT("SearchTextDelimiter", " "), SearchText));
 
 		Row.NameContent()
+		.HAlign(HAlign_Fill)
 		[
-			CreateNameWidget(PrimIdx, VerticalBox.ToSharedRef(), CustomizationUtils)
+			NameContent.ToSharedRef()
 		]
 		.ValueContent()
 		[
@@ -459,7 +507,7 @@ void FCustomPrimitiveDataCustomization::PopulateParameterData(UPrimitiveComponen
 	for (int32 i = 0; i < NumMaterials; ++i)
 	{
 		UMaterialInterface* MaterialInterface = PrimitiveComponent->GetMaterial(i);
-		UMaterial* Material = MaterialInterface ? MaterialInterface->GetBaseMaterial() : NULL;
+		UMaterial* Material = MaterialInterface ? MaterialInterface->GetMaterial() : NULL;
 
 		if (Material == NULL)
 		{
@@ -842,6 +890,7 @@ TSharedRef<SWidget> FCustomPrimitiveDataCustomization::CreateNameWidget(int32 Pr
 {
 	return SNew(SHorizontalBox)
 		+ SHorizontalBox::Slot()
+		.HAlign(HAlign_Left)
 		.VAlign(VAlign_Center)
 		.AutoWidth()
 		.Padding(2.f, 2.f, 16.0f, 2.f)
@@ -851,7 +900,8 @@ TSharedRef<SWidget> FCustomPrimitiveDataCustomization::CreateNameWidget(int32 Pr
 			.Font(CustomizationUtils.GetRegularFont())
 		]
 		+ SHorizontalBox::Slot()
-		.AutoWidth()
+		.HAlign(HAlign_Fill)
+		.VAlign(VAlign_Center)
 		.Padding(0.f, 2.0f)
 		[
 			ParameterName
@@ -873,14 +923,64 @@ END_SLATE_FUNCTION_BUILD_OPTIMIZATION
 BEGIN_SLATE_FUNCTION_BUILD_OPTIMIZATION
 TSharedRef<SWidget> FCustomPrimitiveDataCustomization::GetUndeclaredParameterWidget(int32 PrimIdx, IPropertyTypeCustomizationUtils& CustomizationUtils) const
 {
-	FText PrimIdxText = FText::AsNumber(PrimIdx);
-	TSharedRef<SWidget> UndeclaredParamWidget = SNew(STextBlock)
-		.Text(LOCTEXT("UndeclaredParameter", "(Undeclared)"))
-		.Font(CustomizationUtils.GetRegularFont());
-	UndeclaredParamWidget->SetToolTipText(FText::Format(LOCTEXT("UndeclaredParameterTooltip", "A component is selected that doesn't declare a parameter for primitive index {0}"), PrimIdxText));
-		
-	return UndeclaredParamWidget;
+	const FText ToolTipText = FText::Format(
+		LOCTEXT("UndeclaredParameterTooltip", "An assigned material doesn't declare a parameter for primitive index {0}"),
+		FText::AsNumber(PrimIdx));
+	const uint16 FontSize = CustomizationUtils.GetRegularFont().GetClampSize();
+
+	return SNew(SHorizontalBox)
+		.ToolTipText(ToolTipText)
+		+ SHorizontalBox::Slot()
+		.Padding(4.f, 0.f)
+		.HAlign(HAlign_Left)
+		.VAlign(VAlign_Center)
+		.AutoWidth()
+		[
+			SNew(SImage)
+			.DesiredSizeOverride(FVector2D(FontSize))
+			.Image(FEditorStyle::GetBrush("Icons.Warning"))
+		]
+		+ SHorizontalBox::Slot()
+		.HAlign(HAlign_Left)
+		.VAlign(VAlign_Center)
+		.AutoWidth()
+		[
+			SNew(STextBlock)
+			.Text(LOCTEXT("UndeclaredParameter", "Undeclared"))
+			.Font(CustomizationUtils.GetRegularFont())
+		];
+}
+END_SLATE_FUNCTION_BUILD_OPTIMIZATION
+
+BEGIN_SLATE_FUNCTION_BUILD_OPTIMIZATION
+TSharedRef<SWidget> FCustomPrimitiveDataCustomization::CreateWarningWidget(TSharedRef<SWidget> Content, FText WarningText) const
+{
+	// Similar to SWarningOrErrorBox widget
+	return SNew(SBorder)
+		.BorderImage(FEditorStyle::GetBrush("RoundedWarning"))
+		[
+			SNew(SHorizontalBox)
+			.ToolTipText(WarningText)
+			+ SHorizontalBox::Slot()
+			.Padding(16.f, 2.f)
+			.HAlign(HAlign_Left)
+			.VAlign(VAlign_Center)
+			.AutoWidth()
+			[
+				SNew(SImage)
+				.Image(FEditorStyle::GetBrush("Icons.WarningWithColor"))
+			]
+			+ SHorizontalBox::Slot()
+			.Padding(0.f, 2.f, 16.f, 2.f)
+			.HAlign(HAlign_Fill)
+			.VAlign(VAlign_Center)
+			.AutoWidth()
+			[
+				Content
+			]
+		];
 }
+
 END_SLATE_FUNCTION_BUILD_OPTIMIZATION
 
 FPropertyAccess::Result FCustomPrimitiveDataCustomization::GetNumElements(uint32& NumElements) const
diff --git a/Engine/Source/Editor/DetailCustomizations/Private/CustomPrimitiveDataCustomization.h b/Engine/Source/Editor/DetailCustomizations/Private/CustomPrimitiveDataCustomization.h
--- a/Engine/Source/Editor/DetailCustomizations/Private/CustomPrimitiveDataCustomization.h
+++ b/Engine/Source/Editor/DetailCustomizations/Private/CustomPrimitiveDataCustomization.h
@@ -58,7 +58,7 @@ private:
 	 * a color block that opens a color picker when clicked. Also contains add, delete and reset to default buttons to
 	 * quickly populate a vector parameter
 	 */
-	IDetailGroup* CreateVectorGroup(IDetailChildrenBuilder& ChildBuilder, uint8 PrimIdx, bool bDataEditable, int32 NumElements);
+	IDetailGroup* CreateVectorGroup(IDetailChildrenBuilder& ChildBuilder, uint8 PrimIdx, bool bDataEditable, int32 NumElements, IPropertyTypeCustomizationUtils& CustomizationUtils);
 
 	/**
 	 * Creates a row for the specified primitive index. Will create hyperlinks for each parameter, to take the user to the parameter's node.
@@ -67,7 +67,7 @@ private:
 	 * 
 	 * If the parameter doesn't have an element yet, provides an add button for the user to quickly add elements up until that parameter.
 	 */
-	void CreateParameterRow(IDetailChildrenBuilder& ChildBuilder, uint8 PrimIdx, TSharedPtr<IPropertyHandle> ElementHandle, int32 NumSelectedComponents, bool bDataEditable, IDetailGroup* VectorGroup, IPropertyTypeCustomizationUtils& CustomizationUtils);
+	void CreateParameterRow(IDetailChildrenBuilder& ChildBuilder, uint8 PrimIdx, TSharedPtr<IPropertyHandle> ElementHandle, bool bDataEditable, IDetailGroup* VectorGroup, IPropertyTypeCustomizationUtils& CustomizationUtils);
 
 	template<typename Predicate>
 	void ForEachSelectedComponent(Predicate Pred);
@@ -114,6 +114,9 @@ private:
 	/** Default widget that's used when a component is missing declared parameters for the specified primitive index */
 	TSharedRef<SWidget> GetUndeclaredParameterWidget(int32 PrimIdx, IPropertyTypeCustomizationUtils& CustomizationUtils) const;
 
+	/** Creates a widget wraps content with a warning icon, border and tooltip text */
+	TSharedRef<SWidget> CreateWarningWidget(TSharedRef<SWidget> Content, FText WarningText) const;
+
 	/** Get the number of elements and the access result of our primitive data array */
 	enum FPropertyAccess::Result GetNumElements(uint32& NumElements) const;
 };
