Fix rounded box outlines when using default style with full transparency

diff --git a/Engine/Source/Editor/DetailCustomizations/Private/SlateBrushCustomization.cpp b/Engine/Source/Editor/DetailCustomizations/Private/SlateBrushCustomization.cpp
--- a/Engine/Source/Editor/DetailCustomizations/Private/SlateBrushCustomization.cpp
+++ b/Engine/Source/Editor/DetailCustomizations/Private/SlateBrushCustomization.cpp
@@ -1345,6 +1345,7 @@ void FSlateBrushStructCustomization::CustomizeChildren( TSharedRef<IPropertyHand
 	TSharedPtr<IPropertyHandle> TilingProperty = StructPropertyHandle->GetChildHandle( TEXT("Tiling") );
 	TSharedPtr<IPropertyHandle> MarginProperty = StructPropertyHandle->GetChildHandle( TEXT("Margin") );
 	TSharedPtr<IPropertyHandle> TintProperty = StructPropertyHandle->GetChildHandle( TEXT("TintColor") );
+	TSharedPtr<IPropertyHandle> OutlineSettingsProperty = StructPropertyHandle->GetChildHandle(TEXT("OutlineSettings"));
 	ResourceObjectProperty = StructPropertyHandle->GetChildHandle( TEXT("ResourceObject") );
 	
 	FDetailWidgetRow& ResourceObjectRow = StructBuilder.AddProperty(ResourceObjectProperty.ToSharedRef()).CustomWidget();
@@ -1369,8 +1370,10 @@ void FSlateBrushStructCustomization::CustomizeChildren( TSharedRef<IPropertyHand
 		FResetToDefaultHandler::CreateSP(this, &FSlateBrushStructCustomization::OnImageSizeResetToDefault),
 		bOverrideDefaultOnVectorChildren));
 
-	StructBuilder.AddProperty(TintProperty.ToSharedRef());
+	StructBuilder.AddProperty( TintProperty.ToSharedRef() );
 	StructBuilder.AddProperty( DrawAsProperty.ToSharedRef() );
+	StructBuilder.AddProperty( OutlineSettingsProperty.ToSharedRef() )
+	.Visibility(TAttribute<EVisibility>::Create(TAttribute<EVisibility>::FGetter::CreateSP( this, &FSlateBrushStructCustomization::GetOutlineSettingsPropertyVisibility ) ) );
 	StructBuilder.AddProperty( TilingProperty.ToSharedRef() )
 	.Visibility( TAttribute<EVisibility>::Create( TAttribute<EVisibility>::FGetter::CreateSP( this, &FSlateBrushStructCustomization::GetTilingPropertyVisibility ) ) );
 	StructBuilder.AddProperty( MarginProperty.ToSharedRef() )
@@ -1423,6 +1426,14 @@ void FSlateBrushStructCustomization::CustomizeChildren( TSharedRef<IPropertyHand
 	}
 }
 
+EVisibility FSlateBrushStructCustomization::GetOutlineSettingsPropertyVisibility() const
+{
+	uint8 DrawAsType;
+	FPropertyAccess::Result Result = DrawAsProperty->GetValue(DrawAsType);
+
+	return (Result == FPropertyAccess::MultipleValues || DrawAsType == ESlateBrushDrawType::RoundedBox) ? EVisibility::Visible : EVisibility::Collapsed;
+}
+
 EVisibility FSlateBrushStructCustomization::GetTilingPropertyVisibility() const
 {
 	uint8 DrawAsType;
diff --git a/Engine/Source/Editor/DetailCustomizations/Public/Customizations/SlateBrushCustomization.h b/Engine/Source/Editor/DetailCustomizations/Public/Customizations/SlateBrushCustomization.h
--- a/Engine/Source/Editor/DetailCustomizations/Public/Customizations/SlateBrushCustomization.h
+++ b/Engine/Source/Editor/DetailCustomizations/Public/Customizations/SlateBrushCustomization.h
@@ -22,6 +22,11 @@ public:
 	virtual void CustomizeChildren( TSharedRef<IPropertyHandle> StructPropertyHandle, class IDetailChildrenBuilder& StructBuilder, IPropertyTypeCustomizationUtils& StructCustomizationUtils ) override;              
 
 private:
+	/**
+	 * Get the Slate Brush outline settings property row visibility
+	 */
+	EVisibility GetOutlineSettingsPropertyVisibility() const;
+
 	/**
 	 * Get the Slate Brush tiling property row visibility
 	 */
diff --git a/Engine/Source/Runtime/SlateCore/Private/Rendering/DrawElements.cpp b/Engine/Source/Runtime/SlateCore/Private/Rendering/DrawElements.cpp
--- a/Engine/Source/Runtime/SlateCore/Private/Rendering/DrawElements.cpp
+++ b/Engine/Source/Runtime/SlateCore/Private/Rendering/DrawElements.cpp
@@ -102,7 +102,8 @@ static bool ShouldCull(const FSlateWindowElementList& ElementList, const FPaintG
 
 static bool ShouldCull(const FSlateWindowElementList& ElementList, const FPaintGeometry& PaintGeometry, const FSlateBrush* InBrush, const FLinearColor& InTint)
 {
-	if ((InTint.A == 0 && InBrush->OutlineSettings.Color.GetSpecifiedColor().A == 0) || ShouldCull(ElementList, PaintGeometry, InBrush))
+	bool bIsFullyTransparent = InTint.A == 0 && (InBrush->OutlineSettings.Color.GetSpecifiedColor().A == 0 || InBrush->OutlineSettings.bUseBrushTransparency);
+	if (bIsFullyTransparent || ShouldCull(ElementList, PaintGeometry, InBrush))
 	{
 		return true;
 	}
@@ -244,7 +245,16 @@ FSlateDrawElement& FSlateDrawElement::MakeBoxInternal(
 			CornerRadii = FVector4(UniformRadius, UniformRadius, UniformRadius, UniformRadius);
 		}
 		RBoxPayload->SetRadius(CornerRadii);
-		RBoxPayload->SetOutline(InBrush->OutlineSettings.Color.GetSpecifiedColor(), InBrush->OutlineSettings.Width);
+
+		if (InBrush->OutlineSettings.bUseBrushTransparency)
+		{
+			FLinearColor Color = InBrush->OutlineSettings.Color.GetSpecifiedColor().CopyWithNewOpacity(InTint.A);
+			RBoxPayload->SetOutline(Color, InBrush->OutlineSettings.Width);
+		}
+		else
+		{
+			RBoxPayload->SetOutline(InBrush->OutlineSettings.Color.GetSpecifiedColor(), InBrush->OutlineSettings.Width);
+		}
 		BoxPayload = RBoxPayload;
 	}
 	else 
diff --git a/Engine/Source/Runtime/SlateCore/Private/Styling/UMGCoreStyle.cpp b/Engine/Source/Runtime/SlateCore/Private/Styling/UMGCoreStyle.cpp
--- a/Engine/Source/Runtime/SlateCore/Private/Styling/UMGCoreStyle.cpp
+++ b/Engine/Source/Runtime/SlateCore/Private/Styling/UMGCoreStyle.cpp
@@ -386,10 +386,19 @@ TSharedRef<ISlateStyle> FUMGCoreStyle::Create()
 	}
 
 	// SButton defaults...
+	FSlateRoundedBoxBrush ButtonNormal = { FUMGColor::Secondary, 4.0f, FUMGColor::InputOutline, InputFocusThickness, Icon32x32 };
+	FSlateRoundedBoxBrush ButtonHovered = { FUMGColor::Hover, 4.0f, FUMGColor::Hover, InputFocusThickness, Icon32x32 };
+	FSlateRoundedBoxBrush ButtonPressed = { FUMGColor::Header, 4.0f, FUMGColor::Hover, InputFocusThickness, Icon32x32 };
+
+	// For backward compatability, the default style's outlines should use the transparency of the brush itself
+	ButtonNormal.OutlineSettings.bUseBrushTransparency = true;
+	ButtonHovered.OutlineSettings.bUseBrushTransparency = true;
+	ButtonPressed.OutlineSettings.bUseBrushTransparency = true;
+
 	static const FButtonStyle Button = FButtonStyle()
-		.SetNormal(FSlateRoundedBoxBrush(FUMGColor::Secondary, 4.0f, FUMGColor::InputOutline, InputFocusThickness, Icon32x32))
-		.SetHovered(FSlateRoundedBoxBrush(FUMGColor::Hover, 4.0f, FUMGColor::Hover, InputFocusThickness, Icon32x32))
-		.SetPressed(FSlateRoundedBoxBrush(FUMGColor::Header, 4.0f, FUMGColor::Hover, InputFocusThickness, Icon32x32))
+		.SetNormal(ButtonNormal)
+		.SetHovered(ButtonHovered)
+		.SetPressed(ButtonPressed)
 		.SetNormalForeground(FUMGColor::ForegroundHover)
 		.SetHoveredForeground(FUMGColor::ForegroundHover)
 		.SetPressedForeground(FUMGColor::ForegroundHover)
diff --git a/Engine/Source/Runtime/SlateCore/Public/Styling/SlateBrush.h b/Engine/Source/Runtime/SlateCore/Public/Styling/SlateBrush.h
--- a/Engine/Source/Runtime/SlateCore/Public/Styling/SlateBrush.h
+++ b/Engine/Source/Runtime/SlateCore/Public/Styling/SlateBrush.h
@@ -134,6 +134,7 @@ struct SLATECORE_API FSlateBrushOutlineSettings
 		, Color(FLinearColor::Transparent)
 		, Width(0.0)
 		, RoundingType(ESlateBrushRoundingType::HalfHeightRadius)
+		, bUseBrushTransparency(false)
 	{}
 
 	FSlateBrushOutlineSettings(float InUniformRadius)
@@ -141,6 +142,7 @@ struct SLATECORE_API FSlateBrushOutlineSettings
 		, Color(FLinearColor::Transparent)
 		, Width(0.0)
 		, RoundingType(ESlateBrushRoundingType::FixedRadius)
+		, bUseBrushTransparency(false)
 	{}
 
 	FSlateBrushOutlineSettings(FVector4 InRadius)
@@ -148,6 +150,7 @@ struct SLATECORE_API FSlateBrushOutlineSettings
 		, Color(FLinearColor::Transparent)
 		, Width(0.0)
 		, RoundingType(ESlateBrushRoundingType::FixedRadius)
+		, bUseBrushTransparency(false)
 	{}
 
 	FSlateBrushOutlineSettings(const FSlateColor& InColor, float InWidth)
@@ -155,6 +158,7 @@ struct SLATECORE_API FSlateBrushOutlineSettings
 		, Color(InColor)
 		, Width(InWidth)
 		, RoundingType(ESlateBrushRoundingType::HalfHeightRadius)
+		, bUseBrushTransparency(false)
 	{}
 
 	FSlateBrushOutlineSettings(float InUniformRadius, const FSlateColor& InColor, float InWidth)
@@ -162,6 +166,7 @@ struct SLATECORE_API FSlateBrushOutlineSettings
 		, Color(InColor)
 		, Width(InWidth)
 		, RoundingType(ESlateBrushRoundingType::FixedRadius)
+		, bUseBrushTransparency(false)
 	{}
 
 	FSlateBrushOutlineSettings(FVector4 InRadius, const FSlateColor& InColor, float InWidth)
@@ -169,6 +174,7 @@ struct SLATECORE_API FSlateBrushOutlineSettings
 		, Color(InColor)
 		, Width(InWidth)
 		, RoundingType(ESlateBrushRoundingType::FixedRadius)
+		, bUseBrushTransparency(false)
 	{}
 
 	/** Radius in Slate Units applied to the outline at each corner. X = Top Left, Y = Top Right, Z = Bottom Right, W = Bottom Left */
@@ -187,6 +193,10 @@ struct SLATECORE_API FSlateBrushOutlineSettings
 	UPROPERTY(EditAnywhere, BlueprintReadWrite, Category=Brush)
 	TEnumAsByte<enum ESlateBrushRoundingType::Type > RoundingType;
 
+	/** True if we should use the owning brush's transparency as our own **/
+	UPROPERTY(EditAnywhere, BlueprintReadWrite, Category=Brush)
+	bool bUseBrushTransparency;
+
 };
 
 namespace SlateBrushDefs
