The "Use External Actors" world option should go in the "World" category, not in "World Partition"

diff --git a/Engine/Source/Editor/DetailCustomizations/Private/WorldSettingsDetails.cpp b/Engine/Source/Editor/DetailCustomizations/Private/WorldSettingsDetails.cpp
--- a/Engine/Source/Editor/DetailCustomizations/Private/WorldSettingsDetails.cpp
+++ b/Engine/Source/Editor/DetailCustomizations/Private/WorldSettingsDetails.cpp
@@ -107,10 +107,11 @@ void FWorldSettingsDetails::AddWorldCustomization(IDetailLayoutBuilder& DetailBu
 	{
 		const bool bIsPartitionedWorld = UWorld::HasSubsystem<UWorldPartitionSubsystem>(CustomizedLevel->GetWorld());
 
-		IDetailCategoryBuilder& WorldPartitionCategory = DetailBuilder.EditCategory("WorldPartition");
 		if (GetDefault<UEditorExperimentalSettings>()->bEnableOneFilePerActorSupport)
 		{
-			WorldPartitionCategory.AddCustomRow(LOCTEXT("LevelUseExternalActorsRow", "LevelUseExternalActors"), true)
+			IDetailCategoryBuilder& WorldCategory = DetailBuilder.EditCategory("World");
+
+			WorldCategory.AddCustomRow(LOCTEXT("LevelUseExternalActorsRow", "LevelUseExternalActors"), true)
 				.NameContent()
 				[
 					SNew(STextBlock)
@@ -130,6 +131,8 @@ void FWorldSettingsDetails::AddWorldCustomization(IDetailLayoutBuilder& DetailBu
 
 		if (bIsPartitionedWorld)
 		{
+			IDetailCategoryBuilder& WorldPartitionCategory = DetailBuilder.EditCategory("WorldPartition");
+
 			WorldPartitionCategory.AddCustomRow(LOCTEXT("DefaultWorldPartitionSettingsRow", "DefaultWorldPartitionSettings"), true)
 				.NameContent()
 				[
