Fixing missing include

diff --git a/Engine/Source/Editor/DetailCustomizations/Private/CurveStructCustomization.h b/Engine/Source/Editor/DetailCustomizations/Private/CurveStructCustomization.h
--- a/Engine/Source/Editor/DetailCustomizations/Private/CurveStructCustomization.h
+++ b/Engine/Source/Editor/DetailCustomizations/Private/CurveStructCustomization.h
@@ -8,6 +8,7 @@
 #include "PropertyHandle.h"
 #include "Curves/CurveOwnerInterface.h"
 #include "IDetailChildrenBuilder.h"
+#include "EditorUndoClient.h"
 
 class FDetailWidgetRow;
 class SCurveEditor;
