Added LinearDeltaSensitivity metadata tag to make it possible for detail panel sliders to not be exponential. Used it in a couple of modeling tools.

#rnx

diff --git a/Engine/Plugins/Experimental/MeshModelingToolsetExp/Source/MeshModelingToolsExp/Public/RevolveBoundaryTool.h b/Engine/Plugins/Experimental/MeshModelingToolsetExp/Source/MeshModelingToolsExp/Public/RevolveBoundaryTool.h
--- a/Engine/Plugins/Experimental/MeshModelingToolsetExp/Source/MeshModelingToolsExp/Public/RevolveBoundaryTool.h
+++ b/Engine/Plugins/Experimental/MeshModelingToolsetExp/Source/MeshModelingToolsExp/Public/RevolveBoundaryTool.h
@@ -58,11 +58,13 @@ public:
 	bool bDisplayInputMesh = false;
 
 	/** Sets the revolution axis origin. */
-	UPROPERTY(EditAnywhere, Category = RevolutionAxis, meta = (DisplayName = "Origin"))
+	UPROPERTY(EditAnywhere, Category = RevolutionAxis, meta = (DisplayName = "Origin",
+		Delta = 5, LinearDeltaSensitivity = 1))
 	FVector AxisOrigin = FVector(0, 0, 0);
 
 	/** Sets the revolution axis pitch and yaw. */
-	UPROPERTY(EditAnywhere, Category = RevolutionAxis, meta = (DisplayName = "Orientation"))
+	UPROPERTY(EditAnywhere, Category = RevolutionAxis, meta = (DisplayName = "Orientation",
+		UIMin = -180, UIMax = 180, ClampMin = -180000, ClampMax = 180000))
 	FVector2D AxisOrientation;
 
 protected:
diff --git a/Engine/Plugins/Runtime/MeshModelingToolset/Source/MeshModelingTools/Public/DrawAndRevolveTool.h b/Engine/Plugins/Runtime/MeshModelingToolset/Source/MeshModelingTools/Public/DrawAndRevolveTool.h
--- a/Engine/Plugins/Runtime/MeshModelingToolset/Source/MeshModelingTools/Public/DrawAndRevolveTool.h
+++ b/Engine/Plugins/Runtime/MeshModelingToolset/Source/MeshModelingTools/Public/DrawAndRevolveTool.h
@@ -46,12 +46,13 @@ public:
 	bool bClosePathToAxis = true;
 	
 	/** Sets the draw plane origin. The revolution axis is the X axis in the plane. */
-	UPROPERTY(EditAnywhere, Category = DrawPlane, meta = (DisplayName = "Origin", EditCondition = "bAllowedToEditDrawPlane", HideEditConditionToggle))
+	UPROPERTY(EditAnywhere, Category = DrawPlane, meta = (DisplayName = "Origin", EditCondition = "bAllowedToEditDrawPlane", HideEditConditionToggle,
+		Delta = 5, LinearDeltaSensitivity = 1))
 	FVector DrawPlaneOrigin = FVector(0, 0, 0);
 
 	/** Sets the draw plane orientation. The revolution axis is the X axis in the plane. */
 	UPROPERTY(EditAnywhere, Category = DrawPlane, meta = (DisplayName = "Orientation", EditCondition = "bAllowedToEditDrawPlane", HideEditConditionToggle, 
-		UIMin = -180, UIMax = 180, ClampMin = -180000, ClampMax = 18000))
+		UIMin = -180, UIMax = 180, ClampMin = -180000, ClampMax = 180000))
 	FRotator DrawPlaneOrientation = FRotator(90, 0, 0);
 
 	/** Enables snapping while editing the path. */
diff --git a/Engine/Source/Editor/DetailCustomizations/Private/ColorGradingVectorCustomization.cpp b/Engine/Source/Editor/DetailCustomizations/Private/ColorGradingVectorCustomization.cpp
--- a/Engine/Source/Editor/DetailCustomizations/Private/ColorGradingVectorCustomization.cpp
+++ b/Engine/Source/Editor/DetailCustomizations/Private/ColorGradingVectorCustomization.cpp
@@ -545,19 +545,17 @@ void FColorGradingVectorCustomization::MakeHeaderRow(FDetailWidgetRow& Row, TSha
 			];
 
 		// Make a widget for each property.  The vector component properties  will be displayed in the header
-		TOptional<float> MinValue, MaxValue, SliderMinValue, SliderMaxValue;
-		float SliderExponent, Delta;
-		int32 ShiftMouseMovePixelPerDelta = 1;
-		bool SupportDynamicSliderMaxValue = false;
-		bool SupportDynamicSliderMinValue = false;
 
 		TSharedRef<IPropertyHandle> ColorGradingPropertyHandleRef = ColorGradingPropertyHandle.Pin().ToSharedRef();
-		FMathStructCustomization::ExtractNumericMetadata<float>(ColorGradingPropertyHandleRef, MinValue, MaxValue, SliderMinValue, SliderMaxValue, SliderExponent, Delta, ShiftMouseMovePixelPerDelta, SupportDynamicSliderMaxValue, SupportDynamicSliderMinValue);
+		FMathStructCustomization::FNumericMetadata<float> Metadata;
+		FMathStructCustomization::ExtractNumericMetadata(ColorGradingPropertyHandleRef, Metadata);
 
 		for (int32 ColorIndex = 0; ColorIndex < SortedChildArray.Num(); ++ColorIndex)
 		{
 			TWeakPtr<IPropertyHandle> WeakHandlePtr = SortedChildArray[ColorIndex];
-			TSharedRef<SNumericEntryBox<float>> NumericEntryBox = MakeNumericEntryBox(ColorIndex, MinValue, MaxValue, SliderMinValue, SliderMaxValue, SliderExponent, Delta, ShiftMouseMovePixelPerDelta, SupportDynamicSliderMaxValue, SupportDynamicSliderMinValue);
+			TSharedRef<SNumericEntryBox<float>> NumericEntryBox = MakeNumericEntryBox(ColorIndex, Metadata.MinValue, Metadata.MaxValue, 
+				Metadata.SliderMinValue, Metadata.SliderMaxValue, Metadata.SliderExponent, Metadata.Delta, 
+				Metadata.ShiftMouseMovePixelPerDelta, Metadata.bSupportDynamicSliderMaxValue, Metadata.bSupportDynamicSliderMinValue);
 			TSharedPtr<SSpinBox<float>> NumericEntrySpinBox = StaticCastSharedPtr<SSpinBox<float>>(NumericEntryBox->GetSpinBox());
 			 
 			NumericEntryBoxWidgetList.Add(NumericEntryBox);
@@ -788,15 +786,11 @@ void FColorGradingCustomBuilder::Tick(float DeltaTime)
 void FColorGradingCustomBuilder::GenerateHeaderRowContent(FDetailWidgetRow& NodeRow)
 {
 	// Make a widget for each property.  The vector component properties  will be displayed in the header
-	TOptional<float> MinValue, MaxValue, SliderMinValue, SliderMaxValue;
-	float SliderExponent, Delta;
-	int32 ShiftMouseMovePixelPerDelta = 1;
-	bool SupportDynamicSliderMaxValue = false;
-	bool SupportDynamicSliderMinValue = false;
-	TSharedRef<IPropertyHandle> ColorGradingPropertyHandleRef = ColorGradingPropertyHandle.Pin().ToSharedRef();
-
-	FMathStructCustomization::ExtractNumericMetadata<float>(ColorGradingPropertyHandleRef, MinValue, MaxValue, SliderMinValue, SliderMaxValue, SliderExponent, Delta, ShiftMouseMovePixelPerDelta, SupportDynamicSliderMaxValue, SupportDynamicSliderMinValue);	
 
+	TSharedRef<IPropertyHandle> ColorGradingPropertyHandleRef = ColorGradingPropertyHandle.Pin().ToSharedRef();
+	FMathStructCustomization::FNumericMetadata<float> Metadata;
+	FMathStructCustomization::ExtractNumericMetadata(ColorGradingPropertyHandleRef, Metadata);
+	
 	EColorGradingModes ColorGradingMode = GetColorGradingMode();
 
 	NodeRow.NameContent()
@@ -812,14 +806,14 @@ void FColorGradingCustomBuilder::GenerateHeaderRowContent(FDetailWidgetRow& Node
 			.Padding(FMargin(2.0f, 2.0f, 2.0f, 2.0f))
 			[
 				SAssignNew(ColorGradingPickerWidget, SColorGradingPicker)
-				.ValueMin(MinValue)
-				.ValueMax(MaxValue)
-				.SliderValueMin(SliderMinValue)
-				.SliderValueMax(SliderMaxValue)
-				.MainDelta(Delta)
-				.SupportDynamicSliderMaxValue(SupportDynamicSliderMaxValue)
-				.SupportDynamicSliderMinValue(SupportDynamicSliderMinValue)
-				.MainShiftMouseMovePixelPerDelta(ShiftMouseMovePixelPerDelta)
+				.ValueMin(Metadata.MinValue)
+				.ValueMax(Metadata.MaxValue)
+				.SliderValueMin(Metadata.SliderMinValue)
+				.SliderValueMax(Metadata.SliderMaxValue)
+				.MainDelta(Metadata.Delta)
+				.SupportDynamicSliderMaxValue(Metadata.bSupportDynamicSliderMaxValue)
+				.SupportDynamicSliderMinValue(Metadata.bSupportDynamicSliderMinValue)
+				.MainShiftMouseMovePixelPerDelta(Metadata.ShiftMouseMovePixelPerDelta)
 				.ColorGradingModes(ColorGradingMode)
 				.OnColorCommitted(this, &FColorGradingCustomBuilder::OnColorGradingPickerChanged)
 				.OnQueryCurrentColor(this, &FColorGradingCustomBuilder::GetCurrentColorGradingValue)
@@ -887,7 +881,9 @@ void FColorGradingCustomBuilder::GenerateHeaderRowContent(FDetailWidgetRow& Node
 	{
 		TWeakPtr<IPropertyHandle> WeakHandlePtr = SortedChildArray[ColorIndex];
 
-		TSharedRef<SNumericEntryBox<float>> NumericEntryBox = MakeNumericEntryBox(ColorIndex, MinValue, MaxValue, SliderMinValue, SliderMaxValue, SliderExponent, Delta, ShiftMouseMovePixelPerDelta, SupportDynamicSliderMaxValue, SupportDynamicSliderMinValue);
+		TSharedRef<SNumericEntryBox<float>> NumericEntryBox = MakeNumericEntryBox(ColorIndex, Metadata.MinValue, Metadata.MaxValue, 
+			Metadata.SliderMinValue, Metadata.SliderMaxValue, Metadata.SliderExponent, Metadata.Delta, 
+			Metadata.ShiftMouseMovePixelPerDelta, Metadata.bSupportDynamicSliderMaxValue, Metadata.bSupportDynamicSliderMinValue);
 		TSharedPtr<SSpinBox<float>> NumericEntrySpinBox = StaticCastSharedPtr<SSpinBox<float>>(NumericEntryBox->GetSpinBox());
 
 		NumericEntryBoxWidgetList.Add(NumericEntryBox);
diff --git a/Engine/Source/Editor/DetailCustomizations/Private/MathStructCustomizations.cpp b/Engine/Source/Editor/DetailCustomizations/Private/MathStructCustomizations.cpp
--- a/Engine/Source/Editor/DetailCustomizations/Private/MathStructCustomizations.cpp
+++ b/Engine/Source/Editor/DetailCustomizations/Private/MathStructCustomizations.cpp
@@ -75,6 +75,7 @@ void FMathStructCustomization::MakeHeaderRow(TSharedRef<class IPropertyHandle>&
 		ChildHandle->SetInstanceMetaData(TEXT("UIMax"), StructPropertyHandle->GetMetaData(TEXT("UIMax")));
 		ChildHandle->SetInstanceMetaData(TEXT("SliderExponent"), StructPropertyHandle->GetMetaData(TEXT("SliderExponent")));
 		ChildHandle->SetInstanceMetaData(TEXT("Delta"), StructPropertyHandle->GetMetaData(TEXT("Delta")));
+		ChildHandle->SetInstanceMetaData(TEXT("LinearDeltaSensitivity"), StructPropertyHandle->GetMetaData(TEXT("LinearDeltaSensitivity")));
 		ChildHandle->SetInstanceMetaData(TEXT("ShiftMouseMovePixelPerDelta"), StructPropertyHandle->GetMetaData(TEXT("ShiftMouseMovePixelPerDelta")));
 		ChildHandle->SetInstanceMetaData(TEXT("SupportDynamicSliderMaxValue"), StructPropertyHandle->GetMetaData(TEXT("SupportDynamicSliderMaxValue")));
 		ChildHandle->SetInstanceMetaData(TEXT("SupportDynamicSliderMinValue"), StructPropertyHandle->GetMetaData(TEXT("SupportDynamicSliderMinValue")));
@@ -157,9 +158,49 @@ void FMathStructCustomization::GetSortedChildren(TSharedRef<IPropertyHandle> Str
 	}
 }
 
+// Deprecated overload, to be removed.
+template<typename NumericType>
+void FMathStructCustomization::ExtractNumericMetadata(TSharedRef<IPropertyHandle>& PropertyHandle, TOptional<NumericType>& MinValue,
+		TOptional<NumericType>& MaxValue, TOptional<NumericType>& SliderMinValue, TOptional<NumericType>& SliderMaxValue,
+		NumericType& SliderExponent, NumericType& Delta, int32& ShiftMouseMovePixelPerDelta,
+		bool& bSupportDynamicSliderMaxValue, bool& bSupportDynamicSliderMinValue)
+{
+	FNumericMetadata<NumericType> Metadata;
+	ExtractNumericMetadata(PropertyHandle, Metadata);
+
+	MinValue = Metadata.MinValue;
+	MaxValue = Metadata.MaxValue;
+	SliderMinValue = Metadata.SliderMinValue;
+	SliderMaxValue = Metadata.SliderMaxValue;
+	SliderExponent = Metadata.SliderExponent;
+	Delta = Metadata.Delta;
+	ShiftMouseMovePixelPerDelta = Metadata.ShiftMouseMovePixelPerDelta;
+	bSupportDynamicSliderMaxValue = Metadata.bSupportDynamicSliderMaxValue;
+	bSupportDynamicSliderMinValue = Metadata.bSupportDynamicSliderMinValue;
+}
+// Explicitly instantiate the deprecated overload for the four types that we had implicit instantiations
+// for at time of deprecation (float, double, int32, uint8), since we no longer implicitly instantiate
+// it in this file due to using the other overload.
+template void FMathStructCustomization::ExtractNumericMetadata<float>(TSharedRef<IPropertyHandle>& PropertyHandle,
+	TOptional<float>& MinValue, TOptional<float>& MaxValue, TOptional<float>& SliderMinValue,
+	TOptional<float>& SliderMaxValue, float& SliderExponent, float& Delta,
+	int32& ShiftMouseMovePixelPerDelta, bool& bSupportDynamicSliderMaxValue, bool& bSupportDynamicSliderMinValue);
+template void FMathStructCustomization::ExtractNumericMetadata<double>(TSharedRef<IPropertyHandle>& PropertyHandle,
+	TOptional<double>& MinValue, TOptional<double>& MaxValue, TOptional<double>& SliderMinValue,
+	TOptional<double>& SliderMaxValue, double& SliderExponent, double& Delta,
+	int32& ShiftMouseMovePixelPerDelta, bool& bSupportDynamicSliderMaxValue, bool& bSupportDynamicSliderMinValue);
+template void FMathStructCustomization::ExtractNumericMetadata<int32>(TSharedRef<IPropertyHandle>& PropertyHandle,
+	TOptional<int32>& MinValue, TOptional<int32>& MaxValue, TOptional<int32>& SliderMinValue,
+	TOptional<int32>& SliderMaxValue, int32& SliderExponent, int32& Delta,
+	int32& ShiftMouseMovePixelPerDelta, bool& bSupportDynamicSliderMaxValue, bool& bSupportDynamicSliderMinValue);
+template void FMathStructCustomization::ExtractNumericMetadata<uint8>(TSharedRef<IPropertyHandle>& PropertyHandle,
+	TOptional<uint8>& MinValue, TOptional<uint8>& MaxValue, TOptional<uint8>& SliderMinValue,
+	TOptional<uint8>& SliderMaxValue, uint8& SliderExponent, uint8& Delta,
+	int32& ShiftMouseMovePixelPerDelta, bool& bSupportDynamicSliderMaxValue, bool& bSupportDynamicSliderMinValue);
+
 
 template<typename NumericType>
-void FMathStructCustomization::ExtractNumericMetadata(TSharedRef<IPropertyHandle>& PropertyHandle, TOptional<NumericType>& MinValue, TOptional<NumericType>& MaxValue, TOptional<NumericType>& SliderMinValue, TOptional<NumericType>& SliderMaxValue, NumericType& SliderExponent, NumericType& Delta, int32 &ShiftMouseMovePixelPerDelta, bool& SupportDynamicSliderMaxValue, bool& SupportDynamicSliderMinValue)
+void FMathStructCustomization::ExtractNumericMetadata(TSharedRef<IPropertyHandle>& PropertyHandle, FNumericMetadata<NumericType>& MetadataOut)
 {
 	FProperty* Property = PropertyHandle->GetProperty();
 
@@ -167,6 +208,7 @@ void FMathStructCustomization::ExtractNumericMetadata(TSharedRef<IPropertyHandle
 	const FString& MetaUIMaxString = Property->GetMetaData(TEXT("UIMax"));
 	const FString& SliderExponentString = Property->GetMetaData(TEXT("SliderExponent"));
 	const FString& DeltaString = Property->GetMetaData(TEXT("Delta"));
+	const FString& LinearDeltaSensitivityString = Property->GetMetaData(TEXT("LinearDeltaSensitivity"));
 	const FString& ShiftMouseMovePixelPerDeltaString = Property->GetMetaData(TEXT("ShiftMouseMovePixelPerDelta"));
 	const FString& SupportDynamicSliderMaxValueString = Property->GetMetaData(TEXT("SupportDynamicSliderMaxValue"));
 	const FString& SupportDynamicSliderMinValueString = Property->GetMetaData(TEXT("SupportDynamicSliderMinValue"));
@@ -195,29 +237,37 @@ void FMathStructCustomization::ExtractNumericMetadata(TSharedRef<IPropertyHandle
 	TTypeFromString<NumericType>::FromString(UIMin, *UIMinString);
 	TTypeFromString<NumericType>::FromString(UIMax, *UIMaxString);
 
-	SliderExponent = NumericType(1);
+	MetadataOut.SliderExponent = NumericType(1);
 
 	if (SliderExponentString.Len())
 	{
-		TTypeFromString<NumericType>::FromString(SliderExponent, *SliderExponentString);
+		TTypeFromString<NumericType>::FromString(MetadataOut.SliderExponent, *SliderExponentString);
 	}
 
-	Delta = NumericType(0);
+	MetadataOut.Delta = NumericType(0);
 
 	if (DeltaString.Len())
 	{
-		TTypeFromString<NumericType>::FromString(Delta, *DeltaString);
+		TTypeFromString<NumericType>::FromString(MetadataOut.Delta, *DeltaString);
 	}
 
-	ShiftMouseMovePixelPerDelta = 1;
+	MetadataOut.LinearDeltaSensitivity = 0;
+	if (LinearDeltaSensitivityString.Len())
+	{
+		TTypeFromString<int32>::FromString(MetadataOut.LinearDeltaSensitivity, *LinearDeltaSensitivityString);
+	}
+	// LinearDeltaSensitivity only works in SSpinBox if delta is provided, so add it in if it wasn't.
+	MetadataOut.Delta = (MetadataOut.LinearDeltaSensitivity != 0 && MetadataOut.Delta == NumericType(0)) ? NumericType(1) : MetadataOut.Delta;
+
+	MetadataOut.ShiftMouseMovePixelPerDelta = 1;
 	if (ShiftMouseMovePixelPerDeltaString.Len())
 	{
-		TTypeFromString<int32>::FromString(ShiftMouseMovePixelPerDelta, *ShiftMouseMovePixelPerDeltaString);
+		TTypeFromString<int32>::FromString(MetadataOut.ShiftMouseMovePixelPerDelta, *ShiftMouseMovePixelPerDeltaString);
 		//The value should be greater or equal to 1
 		// 1 is neutral since it is a multiplier of the mouse drag pixel
-		if (ShiftMouseMovePixelPerDelta < 1)
+		if (MetadataOut.ShiftMouseMovePixelPerDelta < 1)
 		{
-			ShiftMouseMovePixelPerDelta = 1;
+			MetadataOut.ShiftMouseMovePixelPerDelta = 1;
 		}
 	}
 
@@ -229,18 +279,18 @@ void FMathStructCustomization::ExtractNumericMetadata(TSharedRef<IPropertyHandle
 	const NumericType ActualUIMin = FMath::Max(UIMin, ClampMin);
 	const NumericType ActualUIMax = FMath::Min(UIMax, ClampMax);
 
-	MinValue = ClampMinString.Len() ? ClampMin : TOptional<NumericType>();
-	MaxValue = ClampMaxString.Len() ? ClampMax : TOptional<NumericType>();
-	SliderMinValue = (UIMinString.Len()) ? ActualUIMin : TOptional<NumericType>();
-	SliderMaxValue = (UIMaxString.Len()) ? ActualUIMax : TOptional<NumericType>();
+	MetadataOut.MinValue = ClampMinString.Len() ? ClampMin : TOptional<NumericType>();
+	MetadataOut.MaxValue = ClampMaxString.Len() ? ClampMax : TOptional<NumericType>();
+	MetadataOut.SliderMinValue = (UIMinString.Len()) ? ActualUIMin : TOptional<NumericType>();
+	MetadataOut.SliderMaxValue = (UIMaxString.Len()) ? ActualUIMax : TOptional<NumericType>();
 
 	if (ActualUIMin >= ActualUIMax && (MetaUIMinString.Len() || MetaUIMaxString.Len()))
 	{
 		//UE_LOG(LogPropertyNode, Warning, TEXT("UI Min (%s) >= UI Max (%s) for Ranged Numeric"), *UIMinString, *UIMaxString);
 	}
 	
-	SupportDynamicSliderMaxValue = SupportDynamicSliderMaxValueString.Len() > 0 && SupportDynamicSliderMaxValueString.ToBool();
-	SupportDynamicSliderMinValue = SupportDynamicSliderMinValueString.Len() > 0 && SupportDynamicSliderMinValueString.ToBool();
+	MetadataOut.bSupportDynamicSliderMaxValue = SupportDynamicSliderMaxValueString.Len() > 0 && SupportDynamicSliderMaxValueString.ToBool();
+	MetadataOut.bSupportDynamicSliderMinValue = SupportDynamicSliderMinValueString.Len() > 0 && SupportDynamicSliderMinValueString.ToBool();
 }
 
 
@@ -249,13 +299,8 @@ TSharedRef<SWidget> FMathStructCustomization::MakeNumericWidget(
 	TSharedRef<IPropertyHandle>& StructurePropertyHandle,
 	TSharedRef<IPropertyHandle>& PropertyHandle)
 {
-	TOptional<NumericType> MinValue, MaxValue, SliderMinValue, SliderMaxValue;
-	NumericType SliderExponent, Delta;
-	int32 ShiftMouseMovePixelPerDelta = 1;
-	bool SupportDynamicSliderMaxValue = false;
-	bool SupportDynamicSliderMinValue = false;
-
-	ExtractNumericMetadata(StructurePropertyHandle, MinValue, MaxValue, SliderMinValue, SliderMaxValue, SliderExponent, Delta, ShiftMouseMovePixelPerDelta, SupportDynamicSliderMaxValue, SupportDynamicSliderMinValue);
+	FNumericMetadata<NumericType> Metadata;
+	ExtractNumericMetadata(StructurePropertyHandle, Metadata);
 
 	TWeakPtr<IPropertyHandle> WeakHandlePtr = PropertyHandle;
 
@@ -272,17 +317,19 @@ TSharedRef<SWidget> FMathStructCustomization::MakeNumericWidget(
 			.OnEndSliderMovement(this, &FMathStructCustomization::OnEndSliderMovement<NumericType>)
 			// Only allow spin on handles with one object.  Otherwise it is not clear what value to spin
 			.AllowSpin(PropertyHandle->GetNumOuterObjects() < 2)
-			.ShiftMouseMovePixelPerDelta(ShiftMouseMovePixelPerDelta)
-			.SupportDynamicSliderMaxValue(SupportDynamicSliderMaxValue)
-			.SupportDynamicSliderMinValue(SupportDynamicSliderMinValue)
+			.ShiftMouseMovePixelPerDelta(Metadata.ShiftMouseMovePixelPerDelta)
+			.SupportDynamicSliderMaxValue(Metadata.bSupportDynamicSliderMaxValue)
+			.SupportDynamicSliderMinValue(Metadata.bSupportDynamicSliderMinValue)
 			.OnDynamicSliderMaxValueChanged(this, &FMathStructCustomization::OnDynamicSliderMaxValueChanged<NumericType>)
 			.OnDynamicSliderMinValueChanged(this, &FMathStructCustomization::OnDynamicSliderMinValueChanged<NumericType>)
-			.MinValue(MinValue)
-			.MaxValue(MaxValue)
-			.MinSliderValue(SliderMinValue)
-			.MaxSliderValue(SliderMaxValue)
-			.SliderExponent(SliderExponent)
-			.Delta(Delta)
+			.MinValue(Metadata.MinValue)
+			.MaxValue(Metadata.MaxValue)
+			.MinSliderValue(Metadata.SliderMinValue)
+			.MaxSliderValue(Metadata.SliderMaxValue)
+			.SliderExponent(Metadata.SliderExponent)
+			.Delta(Metadata.Delta)
+			// LinearDeltaSensitivity must be left unset if not provided, rather than being set to some default
+			.LinearDeltaSensitivity(Metadata.LinearDeltaSensitivity != 0 ? Metadata.LinearDeltaSensitivity : TAttribute<int32>())
 			.ToolTipText(this, &FMathStructCustomization::OnGetValueToolTip<NumericType>, WeakHandlePtr);
 }
 
diff --git a/Engine/Source/Editor/DetailCustomizations/Public/Customizations/MathStructCustomizations.h b/Engine/Source/Editor/DetailCustomizations/Public/Customizations/MathStructCustomizations.h
--- a/Engine/Source/Editor/DetailCustomizations/Public/Customizations/MathStructCustomizations.h
+++ b/Engine/Source/Editor/DetailCustomizations/Public/Customizations/MathStructCustomizations.h
@@ -51,9 +51,32 @@ public:
 	*/
 	bool IsValueEnabled(TWeakPtr<IPropertyHandle> WeakHandlePtr) const;
 	
+	// Argument struct for ExtractNumericMetadata, to make it easier to add new metadata to extract.
+	template <typename NumericType>
+	struct FNumericMetadata
+	{
+		TOptional<NumericType> MinValue;
+		TOptional<NumericType> MaxValue;
+		TOptional<NumericType> SliderMinValue;
+		TOptional<NumericType> SliderMaxValue;
+		NumericType SliderExponent;
+		NumericType Delta;
+		int32 LinearDeltaSensitivity;
+		int32 ShiftMouseMovePixelPerDelta;
+		bool bSupportDynamicSliderMaxValue;
+		bool bSupportDynamicSliderMinValue;
+	};
+
 	/** Utility function that will extract common Math related numeric metadata */	
 	template <typename NumericType>
-	DETAILCUSTOMIZATIONS_API static void ExtractNumericMetadata(TSharedRef<IPropertyHandle>& PropertyHandle, TOptional<NumericType>& MinValue, TOptional<NumericType>& MaxValue, TOptional<NumericType>& SliderMinValue, TOptional<NumericType>& SliderMaxValue, NumericType& SliderExponent, NumericType& Delta, int32 &ShiftMouseMovePixelPerDelta, bool& SupportDynamicSliderMaxValue, bool& SupportDynamicSliderMinValue);
+	DETAILCUSTOMIZATIONS_API static void ExtractNumericMetadata(TSharedRef<IPropertyHandle>& PropertyHandle, FNumericMetadata<NumericType>& MetadataOut);
+
+	template <typename NumericType>
+	UE_DEPRECATED(5.0, "Use ExtractNumericMetadata overload with struct argument instead.")
+	DETAILCUSTOMIZATIONS_API static void ExtractNumericMetadata(TSharedRef<IPropertyHandle>& PropertyHandle, TOptional<NumericType>& MinValue, 
+		TOptional<NumericType>& MaxValue, TOptional<NumericType>& SliderMinValue, TOptional<NumericType>& SliderMaxValue,
+		NumericType& SliderExponent, NumericType& Delta, int32& ShiftMouseMovePixelPerDelta, 
+		bool& bSupportDynamicSliderMaxValue, bool& bSupportDynamicSliderMinValue);
 
 protected:
 
@@ -168,4 +191,4 @@ protected:
 
 	/** True if the ratio is locked when scaling occurs (uniform scaling) */
 	bool bPreserveScaleRatio;
-};
+};
\ No newline at end of file
diff --git a/Engine/Source/Editor/PropertyEditor/Private/UserInterface/PropertyEditor/SPropertyEditorNumeric.h b/Engine/Source/Editor/PropertyEditor/Private/UserInterface/PropertyEditor/SPropertyEditorNumeric.h
--- a/Engine/Source/Editor/PropertyEditor/Private/UserInterface/PropertyEditor/SPropertyEditorNumeric.h
+++ b/Engine/Source/Editor/PropertyEditor/Private/UserInterface/PropertyEditor/SPropertyEditorNumeric.h
@@ -256,6 +256,7 @@ public:
 			const FString& MetaUIMinString = GetMetaDataFromKey("UIMin");
 			const FString& MetaUIMaxString = GetMetaDataFromKey("UIMax");
 			const FString& SliderExponentString = GetMetaDataFromKey("SliderExponent");
+			const FString& LinearDeltaSensitivityString = GetMetaDataFromKey("LinearDeltaSensitivity");
 			const FString& DeltaString = GetMetaDataFromKey("Delta");
 			const FString& ClampMinString = GetMetaDataFromKey("ClampMin");
 			const FString& ClampMaxString = GetMetaDataFromKey("ClampMax");
@@ -294,6 +295,14 @@ public:
 				TTypeFromString<NumericType>::FromString(Delta, *DeltaString);
 			}
 
+			int32 LinearDeltaSensitivity = 0;
+			if (LinearDeltaSensitivityString.Len())
+			{
+				TTypeFromString<int32>::FromString(LinearDeltaSensitivity, *LinearDeltaSensitivityString);
+			}
+			// LinearDeltaSensitivity only works in SSpinBox if delta is provided, so add it in if it wasn't.
+			Delta = (LinearDeltaSensitivity != 0 && Delta == NumericType(0)) ? NumericType(1) : Delta;
+
 			if (ClampMin >= ClampMax && (ClampMinString.Len() || ClampMaxString.Len()))
 			{
 				UE_LOG(LogPropertyNode, Warning, TEXT("Clamp Min (%s) >= Clamp Max (%s) for Ranged Numeric property %s"), *ClampMinString, *ClampMaxString, *Property->GetPathName());
@@ -371,6 +380,8 @@ public:
 				.MaxSliderValue(SliderMaxValue)
 				.SliderExponent(SliderExponent)
 				.Delta(Delta)
+				// LinearDeltaSensitivity needs to be left unset if not provided, rather than being set to some default
+				.LinearDeltaSensitivity(LinearDeltaSensitivity != 0 ? LinearDeltaSensitivity : TAttribute<int32>())
 				.UndeterminedString(LOCTEXT("MultipleValues", "Multiple Values"))
 				.OnValueChanged(this, &SPropertyEditorNumeric<NumericType>::OnValueChanged)
 				.OnValueCommitted(this, &SPropertyEditorNumeric<NumericType>::OnValueCommitted)
diff --git a/Engine/Source/Runtime/CoreUObject/Public/UObject/ObjectMacros.h b/Engine/Source/Runtime/CoreUObject/Public/UObject/ObjectMacros.h
--- a/Engine/Source/Runtime/CoreUObject/Public/UObject/ObjectMacros.h
+++ b/Engine/Source/Runtime/CoreUObject/Public/UObject/ObjectMacros.h
@@ -1204,6 +1204,9 @@ namespace UM
 		/// [PropertyMetadata] Used by FDirectoryPath properties. Indicates that the path will be picked using the Slate-style directory picker inside the game Content dir.
 		ContentDir,
 
+		/// [PropertyMetadata] Used for float and integer properties.  Specifies the step size when the slider is used (e.g. for Delta = 10, the value will scroll among steps of 10).
+		Delta,
+
 		/// [PropertyMetadata] This property is deprecated, any blueprint references to it cause a compilation warning.
 		DeprecatedProperty,
 
@@ -1280,6 +1283,13 @@ namespace UM
 		/// [PropertyMetadata] Signifies that the bool property is only displayed inline as an edit condition toggle in other properties, and should not be shown on its own row.
 		InlineEditConditionToggle,
 
+		/// [PropertyMetadata] Somewhat misleading name (because a higher value is less sensitive), but maps directly to the LinearDeltaSensitivity 
+		/// parameter in sliders created for float and integer properties. When the property does not have both a UIMin and UIMax to bound it, this
+		/// parameter prevents the slider sensitivity from always behaving exponentially (ie slower around 0 and faster elsewhere). Instead, the mouse 
+		/// displacement is simply divided by this value before being multiplied by Delta (which defaults to 1 when LinearDeltaSensitivity is specified,
+		/// and is still raised to SliderExponent). The value should be a positive int (0 does nothing, negative sign ignored if present).
+		LinearDeltaSensitivity,
+
 		/// [ClassMetadata] [PropertyMetadata] Used by TObjectPtr properties.  Allows the referenced object to be eager/lazy/background loaded in editor.
 		// LoadBehavior, (Commented out so as to avoid duplicate name with version in the Class section, but still show in the property section)
 
diff --git a/Engine/Source/Runtime/Slate/Public/Widgets/Input/SSpinBox.h b/Engine/Source/Runtime/Slate/Public/Widgets/Input/SSpinBox.h
--- a/Engine/Source/Runtime/Slate/Public/Widgets/Input/SSpinBox.h
+++ b/Engine/Source/Runtime/Slate/Public/Widgets/Input/SSpinBox.h
@@ -620,7 +620,7 @@ public:
 				{
 					// If this control has a specified delta and sensitivity then we use that instead of the current value for determining how much to change.
 					const float Sign = (MouseEvent.GetCursorDelta().X > 0) ? 1.f : -1.f;
-					if (LinearDeltaSensitivity.IsSet() && Delta.IsSet() && Delta.Get() > 0)
+					if (LinearDeltaSensitivity.IsSet() && LinearDeltaSensitivity.Get() != 0 && Delta.IsSet() && Delta.Get() > 0)
 					{
 						const float MouseDelta = FMath::Abs(MouseEvent.GetCursorDelta().X / LinearDeltaSensitivity.Get());
 						NewValue = InternalValue + (Sign * MouseDelta * FMath::Pow((float)Delta.Get(), SliderExponent.Get()));
