FMathStructCustomizations::OnValueChanged now passes in both EPropertyValueSetFlags::InteractiveChange and EPropertyValueSetFlags::NotTransactable to IPropertyHandle::SetValue when dragging a spinbox. This is most visible when editing a vector from the header X/Y/Z boxes versus the boxes when the vector is expanded.

This was causing issues for a licensee when editing a blueprint instance, presumably due to the blueprint being reinstantiated after every change.

UDN case: https://epicgames.lightning.force.com/lightning/r/Case/5004z00001aV5jwAAC/view

diff --git a/Engine/Source/Editor/DetailCustomizations/Private/MathStructCustomizations.cpp b/Engine/Source/Editor/DetailCustomizations/Private/MathStructCustomizations.cpp
--- a/Engine/Source/Editor/DetailCustomizations/Private/MathStructCustomizations.cpp
+++ b/Engine/Source/Editor/DetailCustomizations/Private/MathStructCustomizations.cpp
@@ -465,7 +465,7 @@ void FMathStructCustomization::OnValueChanged(NumericType NewValue, TWeakPtr<IPr
 {
 	if (bIsUsingSlider)
 	{
-		EPropertyValueSetFlags::Type Flags = EPropertyValueSetFlags::InteractiveChange;
+		EPropertyValueSetFlags::Type Flags = EPropertyValueSetFlags::InteractiveChange | EPropertyValueSetFlags::NotTransactable;
 		SetValue(NewValue, Flags, WeakHandlePtr);
 	}
 }
