Pose Asset not dirtied when source animation is updated
#feature added FGuid property to PoseAsset to compare against the Source Animation's Raw Data GUID - if they don't match during ::PostLoad a warning will be logged. If they mismatch with the Animation Editor open a warning message will be shown in the viewport - directing the user to update the poses manually
#misc some shuffling in details customization layout to make UI more readable - added object version to gate logging the warning until the next resave

diff --git a/Engine/Source/Editor/DetailCustomizations/Private/PoseAssetDetails.cpp b/Engine/Source/Editor/DetailCustomizations/Private/PoseAssetDetails.cpp
--- a/Engine/Source/Editor/DetailCustomizations/Private/PoseAssetDetails.cpp
+++ b/Engine/Source/Editor/DetailCustomizations/Private/PoseAssetDetails.cpp
@@ -205,7 +205,7 @@ void FPoseAssetDetails::CustomizeDetails(IDetailLayoutBuilder& DetailBuilder)
 		.AutoWidth()
 		.HAlign(HAlign_Left)
 		.VAlign(VAlign_Center)
-		.Padding(3)
+		.Padding(FMargin(4.0f,5.0f))
 		[
 			SNew(STextBlock)
 			.Text(LOCTEXT("AdditiveBasePoseLabel", "Base Pose"))
@@ -214,9 +214,8 @@ void FPoseAssetDetails::CustomizeDetails(IDetailLayoutBuilder& DetailBuilder)
 
 		// if additive, let them choose base pose
 		+ SHorizontalBox::Slot()
-		.FillWidth(1)
-		.HAlign(HAlign_Fill)
-		.Padding(3)
+		.AutoWidth()
+		.Padding(FMargin(4.0f,5.0f))
 		[
 			SAssignNew(BasePoseComboBox, SSearchableComboBox)
 			.OptionsSource(&BasePoseComboList)
@@ -225,7 +224,6 @@ void FPoseAssetDetails::CustomizeDetails(IDetailLayoutBuilder& DetailBuilder)
 			.OnComboBoxOpening(this, &FPoseAssetDetails::OnBasePoseComboOpening)
 			.InitiallySelectedItem(InitialSelectedPose)
 			.IsEnabled(this, &FPoseAssetDetails::CanSelectBasePose)
-			.ContentPadding(3)
 			.Content()
 			[
 				SNew(STextBlock)
@@ -235,22 +233,10 @@ void FPoseAssetDetails::CustomizeDetails(IDetailLayoutBuilder& DetailBuilder)
 				.Font(IDetailLayoutBuilder::GetDetailFont())
 			]
 		]
-	];
 
-	AdditiveCategory.AddCustomRow(LOCTEXT("AdditiveSettingCategoryLabel_Apply", "AdditiveSetting_ApplyButton"))
-	.NameContent()
-	[
-		SNew(STextBlock)
-		.Text(LOCTEXT("DummyText"," "))
-	]
-	.ValueContent()
-	.MinDesiredWidth(200)
-	[
-		SNew(SBox)
-		.Padding(5)
-		.HAlign(HAlign_Center)
-		.VAlign(VAlign_Center)
-		.WidthOverride(200)
+		+ SHorizontalBox::Slot()
+		.AutoWidth()
+		.Padding(FMargin(4.0f,5.0f))
 		[
 			// apply button 
 			SNew(SButton)
@@ -258,7 +244,8 @@ void FPoseAssetDetails::CustomizeDetails(IDetailLayoutBuilder& DetailBuilder)
 			.ToolTipText(LOCTEXT("ApplySettingButton_Tooltip", "Apply Additive Setting changes"))
 			.OnClicked(this, &FPoseAssetDetails::OnApplyAdditiveSettings)
 			.HAlign(HAlign_Center)
-			.IsEnabled(this, &FPoseAssetDetails::CanApplySettings)
+			.VAlign(VAlign_Center)
+			.Visibility(this, &FPoseAssetDetails::CanApplySettings)
 		]
 	];
 
@@ -277,30 +264,25 @@ void FPoseAssetDetails::CustomizeDetails(IDetailLayoutBuilder& DetailBuilder)
 	.ValueContent()
 	.MinDesiredWidth(200)
 	[
-		SNew(SVerticalBox)
-		+ SVerticalBox::Slot()
-		.AutoHeight()
+		SNew(SHorizontalBox)
+		+ SHorizontalBox::Slot()
+		.AutoWidth()
 		[
 			SNew(SObjectPropertyEntryBox)
 			.AllowedClass(UAnimSequence::StaticClass())
 			.OnShouldFilterAsset(this, &FPoseAssetDetails::ShouldFilterAsset)
 			.PropertyHandle(SourceAnimationPropertyHandle)
 		]
-		+ SVerticalBox::Slot()
-		.AutoHeight()
-		[
-			SNew(SBox)
-			.Padding(5)
+		+ SHorizontalBox::Slot()
+		.Padding(FMargin(0.f, 5.f))
+		.AutoWidth()
+		[			
+			SNew(SButton)			
+			.Text(LOCTEXT("UpdateSource_Lable", "Update Source"))
+			.ToolTipText(LOCTEXT("UpdateSource_Tooltip", "Update Poses From Source Animation"))
+			.OnClicked(this, &FPoseAssetDetails::OnUpdatePoseSourceAnimation)
+			.IsEnabled(this, &FPoseAssetDetails::IsUpdateSourceEnabled)
 			.HAlign(HAlign_Center)
-			.VAlign(VAlign_Center)
-			.WidthOverride(100)
-			[
-				SNew(SButton)
-				.Text(LOCTEXT("UpdateSource_Lable", "Update Source"))
-				.ToolTipText(LOCTEXT("UpdateSource_Tooltip", "Update Pose From Source Animation"))
-				.OnClicked(this, &FPoseAssetDetails::OnUpdatePoseSourceAnimation)
-				.HAlign(HAlign_Center)
-			]
 		]
 	];
 }
@@ -578,15 +560,15 @@ void FPoseAssetDetails::CachePoseAssetData()
 	 return FReply::Handled();
  }
 
- bool FPoseAssetDetails::CanApplySettings() const
+ EVisibility FPoseAssetDetails::CanApplySettings() const
  {
 	 if (PoseAsset.IsValid())
 	 {
 		 bool bIsAdditiveAsset = PoseAsset->IsValidAdditive();
-		 return (bCachedAdditive != bIsAdditiveAsset || (bIsAdditiveAsset && CachedBasePoseIndex != PoseAsset->GetBasePoseIndex()));
+		 return (bCachedAdditive != bIsAdditiveAsset || (bIsAdditiveAsset && CachedBasePoseIndex != PoseAsset->GetBasePoseIndex())) ? EVisibility::Visible : EVisibility::Collapsed;
 	 }
 
-	 return false;
+	 return EVisibility::Collapsed;
  }
 
  FReply FPoseAssetDetails::OnApplyAdditiveSettings()
@@ -678,4 +660,18 @@ void FPoseAssetDetails::CachePoseAssetData()
 	 return FReply::Handled();
  }
 
+bool FPoseAssetDetails::IsUpdateSourceEnabled() const
+{
+	const UObject* ObjectSet;
+	SourceAnimationPropertyHandle->GetValue(ObjectSet);
+
+	const UAnimSequence* AnimSequenceSelected = Cast<UAnimSequence>(ObjectSet);
+	if (AnimSequenceSelected && PoseAsset.IsValid())
+	{
+		return (PoseAsset->SourceAnimation && PoseAsset->SourceAnimation->GetRawDataGuid() != PoseAsset->SourceAnimationRawDataGUID);		
+	}
+
+	return false;
+}
+
 #undef LOCTEXT_NAMESPACE
diff --git a/Engine/Source/Editor/DetailCustomizations/Private/PoseAssetDetails.h b/Engine/Source/Editor/DetailCustomizations/Private/PoseAssetDetails.h
--- a/Engine/Source/Editor/DetailCustomizations/Private/PoseAssetDetails.h
+++ b/Engine/Source/Editor/DetailCustomizations/Private/PoseAssetDetails.h
@@ -75,7 +75,7 @@ private:
 	int32 CachedBasePoseIndex;
 	void CachePoseAssetData();
 
-	bool CanApplySettings() const;
+	EVisibility CanApplySettings() const;
 	FReply OnApplyAdditiveSettings();
 
 	TSharedPtr<IPropertyHandle> SourceAnimationPropertyHandle;
@@ -83,6 +83,7 @@ private:
 	void OnSourceAnimationChanged(const FAssetData& AssetData);
 	bool ShouldFilterAsset(const FAssetData& AssetData);
 	FReply OnUpdatePoseSourceAnimation();
+	bool IsUpdateSourceEnabled() const;
 
 	FText GetButtonText() const;
 };
diff --git a/Engine/Source/Editor/Persona/Private/AnimationEditorViewportClient.cpp b/Engine/Source/Editor/Persona/Private/AnimationEditorViewportClient.cpp
--- a/Engine/Source/Editor/Persona/Private/AnimationEditorViewportClient.cpp
+++ b/Engine/Source/Editor/Persona/Private/AnimationEditorViewportClient.cpp
@@ -1125,6 +1125,14 @@ FText FAnimationViewportClient::GetDisplayInfo(bool bDisplayAllInfo) const
 		TextValue = ConcatenateLine(TextValue, FText::Format(LOCTEXT("FramerateFormat", "Framerate: {0}"), AnimSequence->GetSamplingFrameRate().ToPrettyText()));
 	}
 
+	if (const UPoseAsset* PoseAsset = Cast<UPoseAsset>(GetAnimPreviewScene()->GetPreviewAnimationAsset()))
+	{
+		if (PoseAsset->SourceAnimation && PoseAsset->SourceAnimation->GetRawDataGuid() != PoseAsset->SourceAnimationRawDataGUID)
+		{
+			TextValue = ConcatenateLine(TextValue, LOCTEXT("PoseAssetOutOfDateWarning", "<AnimViewport.WarningText>Poses are out-of-sync with the source animation. To update them click \"Update Source\"</>"));
+		}
+	}
+
 	return TextValue;
 }
 void FAnimationViewportClient::DrawNodeDebugLines(TArray<FText>& Lines, FCanvas* Canvas, FSceneView* View)
diff --git a/Engine/Source/Runtime/Core/Public/UObject/UE5ReleaseStreamObjectVersion.h b/Engine/Source/Runtime/Core/Public/UObject/UE5ReleaseStreamObjectVersion.h
--- a/Engine/Source/Runtime/Core/Public/UObject/UE5ReleaseStreamObjectVersion.h
+++ b/Engine/Source/Runtime/Core/Public/UObject/UE5ReleaseStreamObjectVersion.h
@@ -84,6 +84,9 @@ struct CORE_API FUE5ReleaseStreamObjectVersion
 		// Unified volumetric cloud component quality sample count slider between main and reflection views for consistency
 		VolumetricCloudSampleCountUnification,
 
+		// Pose asset GUID generated from source AnimationSequence
+		PoseAssetRawDataGUID,
+
 		// -----<new versions can be added above this line>-------------------------------------------------
 		VersionPlusOne,
 		LatestVersion = VersionPlusOne - 1
diff --git a/Engine/Source/Runtime/Engine/Classes/Animation/PoseAsset.h b/Engine/Source/Runtime/Engine/Classes/Animation/PoseAsset.h
--- a/Engine/Source/Runtime/Engine/Classes/Animation/PoseAsset.h
+++ b/Engine/Source/Runtime/Engine/Classes/Animation/PoseAsset.h
@@ -186,6 +186,10 @@ public:
 #if WITH_EDITORONLY_DATA
 	UPROPERTY(Category=Source, EditAnywhere)
 	TObjectPtr<UAnimSequence> SourceAnimation;
+
+	/** GUID cached when the contained poses were last updated according to SourceAnimation - used to keep track of out-of-date/sync data*/ 
+	UPROPERTY()	
+	FGuid SourceAnimationRawDataGUID;
 #endif // WITH_EDITORONLY_DATA
 
 	/**
diff --git a/Engine/Source/Runtime/Engine/Private/Animation/PoseAsset.cpp b/Engine/Source/Runtime/Engine/Private/Animation/PoseAsset.cpp
--- a/Engine/Source/Runtime/Engine/Private/Animation/PoseAsset.cpp
+++ b/Engine/Source/Runtime/Engine/Private/Animation/PoseAsset.cpp
@@ -11,8 +11,10 @@
 #include "Animation/AnimationPoseData.h"
 #include "Animation/AnimData/AnimDataModel.h"
 #include "Animation/AnimSequenceHelpers.h"
+#include "Framework/Notifications/NotificationManager.h"
 #include "UObject/ObjectSaveContext.h"
 #include "UObject/UE5ReleaseStreamObjectVersion.h"
+#include "Widgets/Notifications/SNotificationList.h"
 
 #define LOCTEXT_NAMESPACE "PoseAsset"
 
@@ -908,6 +910,18 @@ void UPoseAsset::PostLoad()
     {
 		PostProcessData();
     }
+
+	if (GetLinkerCustomVersion(FUE5ReleaseStreamObjectVersion::GUID) >= FUE5ReleaseStreamObjectVersion::PoseAssetRawDataGUID)
+	{
+		if (SourceAnimation != nullptr && (!SourceAnimationRawDataGUID.IsValid() || SourceAnimationRawDataGUID != SourceAnimation->GetRawDataGuid()))
+		{
+			FFormatNamedArguments Args;
+			Args.Add(TEXT("AssetName"), FText::FromString(GetPathName()));
+			Args.Add(TEXT("SourceAsset"), FText::FromString(SourceAnimation->GetPathName()));
+			const FText ResultText = FText::Format(LOCTEXT("PoseAssetSourceOutOfDate", "PoseAsset {AssetName} is out-of-date with its source animation {SourceAsset}"), Args);
+			UE_LOG(LogAnimation, Warning,TEXT("%s"), *ResultText.ToString());
+		}
+	}	
 #endif // WITH_EDITOR
 
 	// fix curve names
@@ -1282,6 +1296,11 @@ void UPoseAsset::PostEditChangeProperty(FPropertyChangedEvent& PropertyChangedEv
 				ConvertToAdditivePose(GetBasePoseIndex());
 			}
 		}
+
+		if (PropertyChangedEvent.Property->GetFName() == GET_MEMBER_NAME_CHECKED(UPoseAsset, SourceAnimation))
+		{
+			SourceAnimationRawDataGUID.Invalidate();
+		}
 	}
 }
 
@@ -1348,6 +1367,7 @@ void UPoseAsset::CreatePoseFromAnimation(class UAnimSequence* AnimSequence, cons
 		{
 			SetSkeleton(TargetSkeleton);
 			SourceAnimation = AnimSequence;
+			SourceAnimationRawDataGUID = AnimSequence->GetRawDataGuid();
 
 			// reinitialize, now we're making new pose from this animation
 			Reinitialize();
