Fix actor guid reset when resetting to BP defaults.

diff --git a/Engine/Source/Editor/DetailCustomizations/Private/ActorDetails.cpp b/Engine/Source/Editor/DetailCustomizations/Private/ActorDetails.cpp
--- a/Engine/Source/Editor/DetailCustomizations/Private/ActorDetails.cpp
+++ b/Engine/Source/Editor/DetailCustomizations/Private/ActorDetails.cpp
@@ -617,6 +617,28 @@ void FActorDetails::AddActorCategory( IDetailLayoutBuilder& DetailBuilder, const
 
 	IDetailCategoryBuilder& ActorCategory = DetailBuilder.EditCategory("Actor", FText::GetEmpty(), ECategoryPriority::Uncommon );
 
+	if (GetSelectedActors().Num() == 1)
+	{
+		if (AActor* Actor = GEditor->GetSelectedActors()->GetTop<AActor>())
+		{
+			const FText ActorGuidText = FText::FromString(Actor->GetActorGuid().ToString());
+			ActorCategory.AddCustomRow( LOCTEXT("ActorGuid", "ActorGuid") )
+				.NameContent()
+				[
+					SNew(STextBlock)
+					.Text(LOCTEXT("ActorGuid", "Actor Guid"))
+					.ToolTipText(LOCTEXT("ActorGuid_ToolTip", "Actor Guid"))
+					.Font(IDetailLayoutBuilder::GetDetailFont())
+				]
+				.ValueContent()
+				[
+					SNew(STextBlock)
+						.Text(ActorGuidText)
+						.Font(IDetailLayoutBuilder::GetDetailFont())
+						.IsEnabled(false)
+				];
+		}
+	};
 
 #if 1
 	// Create the info buttons per level
diff --git a/Engine/Source/Runtime/Engine/Classes/GameFramework/Actor.h b/Engine/Source/Runtime/Engine/Classes/GameFramework/Actor.h
--- a/Engine/Source/Runtime/Engine/Classes/GameFramework/Actor.h
+++ b/Engine/Source/Runtime/Engine/Classes/GameFramework/Actor.h
@@ -817,8 +817,11 @@ private:
 protected:
 	/**
 	 * The GUID for this actor.
+	 *
+	 * Note: Don't use VisibleAnywhere here to avoid getting the CPF_Edit flag and get this property reset when resetting to defaults.
+	 *       See FActorDetails::AddActorCategory and EditorUtilities::CopySingleProperty for details.
 	 */
-	UPROPERTY(VisibleAnywhere, AdvancedDisplay, Category=Actor, NonPIEDuplicateTransient, TextExportTransient, NonTransactional)
+	UPROPERTY(BluePrintReadOnly, AdvancedDisplay, Category=Actor, NonPIEDuplicateTransient, TextExportTransient, NonTransactional)
 	FGuid ActorGuid;
 
 	/** DataLayers the actor belongs to.*/
