Add button to landscape details panel to create RVT volumes for all set RVT assets.
Small quality of life improvement for RVT setup.

diff --git a/Engine/Source/Editor/DetailCustomizations/DetailCustomizations.Build.cs b/Engine/Source/Editor/DetailCustomizations/DetailCustomizations.Build.cs
--- a/Engine/Source/Editor/DetailCustomizations/DetailCustomizations.Build.cs
+++ b/Engine/Source/Editor/DetailCustomizations/DetailCustomizations.Build.cs
@@ -66,6 +66,7 @@ public class DetailCustomizations : ModuleRules
 				"RenderCore",
 				"ToolWidgets",
 				"MaterialEditor",
+				"VirtualTexturingEditor",
 			}
 		);
 
diff --git a/Engine/Source/Editor/DetailCustomizations/Private/LandscapeProxyUIDetails.cpp b/Engine/Source/Editor/DetailCustomizations/Private/LandscapeProxyUIDetails.cpp
--- a/Engine/Source/Editor/DetailCustomizations/Private/LandscapeProxyUIDetails.cpp
+++ b/Engine/Source/Editor/DetailCustomizations/Private/LandscapeProxyUIDetails.cpp
@@ -1,19 +1,25 @@
 // Copyright Epic Games, Inc. All Rights Reserved.
 
 #include "LandscapeProxyUIDetails.h"
-#include "Widgets/DeclarativeSyntaxSupport.h"
-#include "PropertyHandle.h"
+
+#include "Components/RuntimeVirtualTextureComponent.h"
+#include "DetailCategoryBuilder.h"
 #include "DetailLayoutBuilder.h"
+#include "DetailWidgetRow.h"
+#include "Editor.h"
+#include "Misc/MessageDialog.h"
+#include "PropertyHandle.h"
 #include "Runtime/Landscape/Classes/Landscape.h"
 #include "Runtime/Landscape/Classes/LandscapeInfo.h"
+#include "RuntimeVirtualTextureSetBounds.h"
+#include "ScopedTransaction.h"
 #include "Settings/EditorExperimentalSettings.h"
-#include "DetailCategoryBuilder.h"
-#include "DetailWidgetRow.h"
+#include "VT/RuntimeVirtualTexture.h"
+#include "VT/RuntimeVirtualTextureVolume.h"
+#include "Widgets/DeclarativeSyntaxSupport.h"
 #include "Widgets/Input/SButton.h"
-#include "Widgets/Text/STextBlock.h"
 #include "Widgets/Input/SCheckBox.h"
-#include "Misc/MessageDialog.h"
-#include "Editor.h"
+#include "Widgets/Text/STextBlock.h"
 
 #define LOCTEXT_NAMESPACE "FLandscapeProxyUIDetails"
 
@@ -40,9 +46,10 @@ void FLandscapeProxyUIDetails::CustomizeDetails( IDetailLayoutBuilder& DetailBui
 
 	if (EditingObjects.Num() == 1)
 	{
-		if (ALandscapeProxy* Proxy = Cast<ALandscapeProxy>(EditingObjects[0]))
+		LandscapeProxy = Cast<ALandscapeProxy>(EditingObjects[0]);
+		if (LandscapeProxy != nullptr)
 		{
-			if (ULandscapeInfo* LandscapeInfo = Proxy->GetLandscapeInfo())
+			if (ULandscapeInfo* LandscapeInfo = LandscapeProxy->GetLandscapeInfo())
 			{
 				IDetailCategoryBuilder& CategoryBuilder = DetailBuilder.EditCategory("Information", FText::GetEmpty(), ECategoryPriority::Important);
 				
@@ -54,7 +61,7 @@ void FLandscapeProxyUIDetails::CustomizeDetails( IDetailLayoutBuilder& DetailBui
 				]
 				.ValueContent()
 				[
-					GenerateTextWidget(FText::Format(LOCTEXT("LandscapeComponentResolutionValue", "{0} x {0}"), Proxy->ComponentSizeQuads+1), true) // Verts
+					GenerateTextWidget(FText::Format(LOCTEXT("LandscapeComponentResolutionValue", "{0} x {0}"), LandscapeProxy->ComponentSizeQuads+1), true) // Verts
 				];
 
 				RowDisplayText = LOCTEXT("LandscapeComponentCount", "Component Count");
@@ -65,7 +72,7 @@ void FLandscapeProxyUIDetails::CustomizeDetails( IDetailLayoutBuilder& DetailBui
 				]
 				.ValueContent()
 				[
-					GenerateTextWidget(FText::Format(LOCTEXT("LandscapeComponentCountValue", "{0}"), Proxy->LandscapeComponents.Num()), true)
+					GenerateTextWidget(FText::Format(LOCTEXT("LandscapeComponentCountValue", "{0}"), LandscapeProxy->LandscapeComponents.Num()), true)
 				];
 
 				RowDisplayText = LOCTEXT("LandscapeComponentSubsections", "Component Subsections");
@@ -76,10 +83,10 @@ void FLandscapeProxyUIDetails::CustomizeDetails( IDetailLayoutBuilder& DetailBui
 				]
 				.ValueContent()
 				[
-					GenerateTextWidget(FText::Format(LOCTEXT("LandscapeComponentSubSectionsValue", "{0} x {0}"), Proxy->NumSubsections), true)
+					GenerateTextWidget(FText::Format(LOCTEXT("LandscapeComponentSubSectionsValue", "{0} x {0}"), LandscapeProxy->NumSubsections), true)
 				];
 
-				FIntRect Rect = Proxy->GetBoundingRect();
+				FIntRect Rect = LandscapeProxy->GetBoundingRect();
 				FIntPoint Size = Rect.Size();
 				RowDisplayText = LOCTEXT("LandscapeResolution", "Resolution (Verts)");
 				CategoryBuilder.AddCustomRow(RowDisplayText)
@@ -130,8 +137,89 @@ void FLandscapeProxyUIDetails::CustomizeDetails( IDetailLayoutBuilder& DetailBui
 				];
 			}
 
+
+			// Apply custom widget for CreateVolume.
+			TSharedRef<IPropertyHandle> CreateVolumePropertyHandle = DetailBuilder.GetProperty(TEXT("bSetCreateRuntimeVirtualTextureVolumes"));
+			DetailBuilder.EditDefaultProperty(CreateVolumePropertyHandle)->CustomWidget()
+			.NameContent()
+			[
+				SNew(STextBlock)
+				.Font(IDetailLayoutBuilder::GetDetailFont())
+				.Text(LOCTEXT("Button_CreateVolumes", "Create Volumes"))
+				.ToolTipText(LOCTEXT("Button_CreateVolumes_Tooltip", "Create volumes for the selected Runtime Virtual Textures."))
+			]
+			.ValueContent()
+			.MinDesiredWidth(125.f)
+			[
+				SNew(SButton)
+				.VAlign(VAlign_Center)
+				.HAlign(HAlign_Center)
+				.ContentPadding(2)
+				.Text(LOCTEXT("Button_CreateVolumes", "Create Volumes"))
+				.OnClicked(this, &FLandscapeProxyUIDetails::CreateRuntimeVirtualTextureVolume)
+				.IsEnabled(this, &FLandscapeProxyUIDetails::IsCreateRuntimeVirtualTextureVolumeEnabled)
+			];
+
+		}
+	}
+}
+
+static void GetMissingRuntimeVirtualTextureVolumes(ALandscapeProxy* LandscapeProxy, TArray<URuntimeVirtualTexture*>& OutVirtualTextures)
+{
+	UWorld* World = LandscapeProxy != nullptr ? LandscapeProxy->GetWorld() : nullptr;
+	if (World == nullptr)
+	{
+		return;
+	}
+
+	TArray<URuntimeVirtualTexture*> FoundVolumes;
+	for (TObjectIterator<URuntimeVirtualTextureComponent> It(RF_ClassDefaultObject, false, EInternalObjectFlags::PendingKill); It; ++It)
+	{
+		if (It->GetWorld() == World)
+		{
+			if (URuntimeVirtualTexture* VirtualTexture = It->GetVirtualTexture())
+			{
+				FoundVolumes.Add(VirtualTexture);
+			}
+		}
+	}
+
+	for (URuntimeVirtualTexture* VirtualTexture : LandscapeProxy->RuntimeVirtualTextures)
+	{
+		if (VirtualTexture != nullptr && FoundVolumes.Find(VirtualTexture) == INDEX_NONE)
+		{
+			OutVirtualTextures.Add(VirtualTexture);
 		}
 	}
 }
 
+bool FLandscapeProxyUIDetails::IsCreateRuntimeVirtualTextureVolumeEnabled() const
+{
+	TArray<URuntimeVirtualTexture*> VirtualTextureVolumesToCreate;
+	GetMissingRuntimeVirtualTextureVolumes(LandscapeProxy, VirtualTextureVolumesToCreate);
+	return VirtualTextureVolumesToCreate.Num() > 0;
+}
+
+FReply FLandscapeProxyUIDetails::CreateRuntimeVirtualTextureVolume()
+{
+	TArray<URuntimeVirtualTexture*> VirtualTextureVolumesToCreate;
+	GetMissingRuntimeVirtualTextureVolumes(LandscapeProxy, VirtualTextureVolumesToCreate);
+	if (VirtualTextureVolumesToCreate.Num() == 0)
+	{
+		return FReply::Unhandled();
+	}
+
+	const FScopedTransaction Transaction(LOCTEXT("Transaction_CreateVolumes", "Create Runtime Virtual Texture Volumes"));
+	
+	for (URuntimeVirtualTexture* VirtualTexture : VirtualTextureVolumesToCreate)
+	{
+		ARuntimeVirtualTextureVolume* NewVolume = LandscapeProxy->GetWorld()->SpawnActor<ARuntimeVirtualTextureVolume>();
+		NewVolume->VirtualTextureComponent->SetVirtualTexture(VirtualTexture);
+		NewVolume->VirtualTextureComponent->SetBoundsAlignActor(LandscapeProxy);
+		RuntimeVirtualTexture::SetBounds(NewVolume->VirtualTextureComponent);
+	}
+
+	return FReply::Handled();
+}
+
 #undef LOCTEXT_NAMESPACE
diff --git a/Engine/Source/Editor/DetailCustomizations/Private/LandscapeProxyUIDetails.h b/Engine/Source/Editor/DetailCustomizations/Private/LandscapeProxyUIDetails.h
--- a/Engine/Source/Editor/DetailCustomizations/Private/LandscapeProxyUIDetails.h
+++ b/Engine/Source/Editor/DetailCustomizations/Private/LandscapeProxyUIDetails.h
@@ -3,10 +3,8 @@
 #pragma once
 
 #include "CoreMinimal.h"
-#include "UObject/WeakObjectPtr.h"
-#include "Types/SlateEnums.h"
 #include "IDetailCustomization.h"
-#include "Internationalization/Text.h"
+#include "Input/Reply.h"
 
 class IDetailLayoutBuilder;
 class ALandscapeProxy;
@@ -24,6 +22,12 @@ public:
 private:
 	/** Use MakeInstance to create an instance of this class */
 	FLandscapeProxyUIDetails();
+
+		/** Returns true if SetBounds button is enabled */
+	bool IsCreateRuntimeVirtualTextureVolumeEnabled() const;
+	/** Callback for Set Bounds button */
+	FReply CreateRuntimeVirtualTextureVolume();
 	
 	IDetailLayoutBuilder* DetailLayoutBuilder = nullptr;
+	ALandscapeProxy* LandscapeProxy = nullptr;
 };
diff --git a/Engine/Source/Editor/VirtualTexturingEditor/Private/RuntimeVirtualTextureSetBounds.h b/Engine/Source/Editor/VirtualTexturingEditor/Private/RuntimeVirtualTextureSetBounds.h
deleted file mode 100644
diff --git a/Engine/Source/Editor/VirtualTexturingEditor/Public/RuntimeVirtualTextureSetBounds.h b/Engine/Source/Editor/VirtualTexturingEditor/Public/RuntimeVirtualTextureSetBounds.h
new file mode 100644
--- /dev/null
+++ b/Engine/Source/Editor/VirtualTexturingEditor/Public/RuntimeVirtualTextureSetBounds.h
@@ -0,0 +1,13 @@
+// Copyright Epic Games, Inc. All Rights Reserved.
+
+#pragma once
+
+#include "CoreMinimal.h"
+
+class URuntimeVirtualTextureComponent;
+
+namespace RuntimeVirtualTexture
+{
+	/** Set the transform on a URuntimeVirtualTextureComponent so that it includes the bounds of all associated primitives in the current world. */
+	void VIRTUALTEXTURINGEDITOR_API SetBounds(URuntimeVirtualTextureComponent* InComponent);
+};
diff --git a/Engine/Source/Runtime/Engine/Classes/Components/RuntimeVirtualTextureComponent.h b/Engine/Source/Runtime/Engine/Classes/Components/RuntimeVirtualTextureComponent.h
--- a/Engine/Source/Runtime/Engine/Classes/Components/RuntimeVirtualTextureComponent.h
+++ b/Engine/Source/Runtime/Engine/Classes/Components/RuntimeVirtualTextureComponent.h
@@ -43,7 +43,7 @@ protected:
 
 	/** Hide primitives in the main pass. Hidden primitives will be those that draw to this virtual texture with 'Draw in Main Pass' set to 'From Virtual Texture'. */
 	UPROPERTY(EditAnywhere, AdvancedDisplay, Category = VirtualTexture)
-	bool bHidePrimitives = false;
+	bool bHidePrimitives = true;
 
 	/** Texture object containing streamed low mips. This can reduce rendering update cost. */
 	UPROPERTY(EditAnywhere, BlueprintReadOnly, NonPIEDuplicateTransient, Category = VirtualTextureBuild)
@@ -86,6 +86,9 @@ public:
 	UFUNCTION(BlueprintCallable, Category = "VirtualTexture")
 	void Invalidate(FBoxSphereBounds const& WorldBounds);
 
+	/** Set the runtime virtual texture object on this component. */
+	void SetVirtualTexture(URuntimeVirtualTexture* InVirtualTexture);
+
 	/** Get the runtime virtual texture object on this component. */
 	URuntimeVirtualTexture* GetVirtualTexture() const { return VirtualTexture; }
 
@@ -127,6 +130,8 @@ public:
 #endif
 
 #if WITH_EDITOR
+	/** Get the BoundsAlignActor on this component. */
+	void SetBoundsAlignActor(AActor* InActor) { BoundsAlignActor = InActor; }
 	/** Get the BoundsAlignActor on this component. */
 	TSoftObjectPtr<AActor>& GetBoundsAlignActor() { return BoundsAlignActor; }
 	/** Get if SnapBoundsToLandscape is set on this component. */
diff --git a/Engine/Source/Runtime/Engine/Private/Components/RuntimeVirtualTextureComponent.cpp b/Engine/Source/Runtime/Engine/Private/Components/RuntimeVirtualTextureComponent.cpp
--- a/Engine/Source/Runtime/Engine/Private/Components/RuntimeVirtualTextureComponent.cpp
+++ b/Engine/Source/Runtime/Engine/Private/Components/RuntimeVirtualTextureComponent.cpp
@@ -40,6 +40,12 @@ void URuntimeVirtualTextureComponent::OnUnregister()
 
 #endif
 
+void URuntimeVirtualTextureComponent::SetVirtualTexture(URuntimeVirtualTexture* InVirtualTexture) 
+{
+	VirtualTexture = InVirtualTexture; 
+	MarkRenderStateDirty();
+}
+
 void URuntimeVirtualTextureComponent::GetHidePrimitiveSettings(bool& OutHidePrimitiveEditor, bool& OutHidePrimitiveGame) const
 {
 	OutHidePrimitiveEditor = bHidePrimitives;
diff --git a/Engine/Source/Runtime/Landscape/Classes/LandscapeProxy.h b/Engine/Source/Runtime/Landscape/Classes/LandscapeProxy.h
--- a/Engine/Source/Runtime/Landscape/Classes/LandscapeProxy.h
+++ b/Engine/Source/Runtime/Landscape/Classes/LandscapeProxy.h
@@ -492,6 +492,10 @@ public:
 	UPROPERTY(EditAnywhere, BlueprintReadWrite, Category = VirtualTexture, meta = (DisplayName = "Draw in Virtual Textures"))
 	TArray<TObjectPtr<URuntimeVirtualTexture>> RuntimeVirtualTextures;
 
+	/** Placeholder for details customization button. */
+	UPROPERTY(VisibleAnywhere, Transient, Category = VirtualTexture)
+	bool bSetCreateRuntimeVirtualTextureVolumes;
+
 	/** 
 	 * Number of mesh levels to use when rendering landscape into runtime virtual texture.
 	 * Lower values reduce vertex count when rendering to the runtime virtual texture but decrease accuracy when using values that require vertex interpolation.
