Editor - Fix an issue in the CustomPrimitiveDataCustomization where the details panel would be refreshed every change due to a mismatch in the number of the materials used and the number of valid material pointers returned.

diff --git a/Engine/Source/Editor/DetailCustomizations/Private/CustomPrimitiveDataCustomization.cpp b/Engine/Source/Editor/DetailCustomizations/Private/CustomPrimitiveDataCustomization.cpp
--- a/Engine/Source/Editor/DetailCustomizations/Private/CustomPrimitiveDataCustomization.cpp
+++ b/Engine/Source/Editor/DetailCustomizations/Private/CustomPrimitiveDataCustomization.cpp
@@ -492,6 +492,7 @@ void FCustomPrimitiveDataCustomization::Cleanup()
 	DataArrayHandle = NULL;
 
 	ComponentsToWatch.Empty();
+	ComponentMaterialCounts.Empty();
 	MaterialsToWatch.Empty();
 	VectorParameterData.Empty();
 	ScalarParameterData.Empty();
@@ -503,7 +504,7 @@ void FCustomPrimitiveDataCustomization::PopulateParameterData(UPrimitiveComponen
 	const int32 NumMaterials = PrimitiveComponent->GetNumMaterials();
 
 	TSet<TSoftObjectPtr<UMaterial>>& CachedComponentMaterials = ComponentsToWatch.FindOrAdd(PrimitiveComponent);
-
+	ComponentMaterialCounts.FindOrAdd(PrimitiveComponent) = NumMaterials;
 	for (int32 i = 0; i < NumMaterials; ++i)
 	{
 		UMaterialInterface* MaterialInterface = PrimitiveComponent->GetMaterial(i);
@@ -627,7 +628,7 @@ void FCustomPrimitiveDataCustomization::OnObjectPropertyChanged(UObject* Object,
 	UPrimitiveComponent* PrimComponent = Cast<UPrimitiveComponent>(Object);
 	const EPropertyChangeType::Type IgnoreFlags = EPropertyChangeType::Interactive | EPropertyChangeType::Redirected;
 
-	if (!(PropertyChangedEvent.ChangeType & IgnoreFlags) && ComponentsToWatch.Contains(PrimComponent)
+	if (!(PropertyChangedEvent.ChangeType & IgnoreFlags) && ComponentsToWatch.Contains(PrimComponent) && ComponentMaterialCounts.Contains(PrimComponent)
 		&& IsSelected(PrimComponent)) // Need to test this in case we're hitting a stale hash in ComponentsToWatch (#jira UE-136687)
 	{
 		bool bMaterialChange = false;
@@ -646,10 +647,10 @@ void FCustomPrimitiveDataCustomization::OnObjectPropertyChanged(UObject* Object,
 			// NOTE: Optimally would be done via an "OnMaterialChanged" for each component, however,
 			// the property name checks above should handle most cases
 			TSet<TSoftObjectPtr<UMaterial>>& CachedComponentMaterials = ComponentsToWatch[PrimComponent];
-
+			int32 CachedComponentMaterialCount = ComponentMaterialCounts[PrimComponent];
 			const int32 NumMaterials = PrimComponent->GetNumMaterials();
 
-			if (NumMaterials != CachedComponentMaterials.Num())
+			if (NumMaterials != CachedComponentMaterialCount)
 			{
 				bMaterialChange = true;
 			}
diff --git a/Engine/Source/Editor/DetailCustomizations/Private/CustomPrimitiveDataCustomization.h b/Engine/Source/Editor/DetailCustomizations/Private/CustomPrimitiveDataCustomization.h
--- a/Engine/Source/Editor/DetailCustomizations/Private/CustomPrimitiveDataCustomization.h
+++ b/Engine/Source/Editor/DetailCustomizations/Private/CustomPrimitiveDataCustomization.h
@@ -47,6 +47,7 @@ private:
 	TSharedPtr<IPropertyHandleArray> DataArrayHandle;
 
 	TMap<TWeakObjectPtr<UPrimitiveComponent>, TSet<TSoftObjectPtr<UMaterial>>> ComponentsToWatch;
+	TMap<TWeakObjectPtr<UPrimitiveComponent>, uint32> ComponentMaterialCounts;
 	TSet<TSoftObjectPtr<UMaterialInterface>> MaterialsToWatch;
 
 	TMap<uint8, TArray<FParameterData>> VectorParameterData;
