Hide clear buttons for anim sequence reference in montage segment

diff --git a/Engine/Source/Editor/DetailCustomizations/Private/AnimMontageSegmentDetails.cpp b/Engine/Source/Editor/DetailCustomizations/Private/AnimMontageSegmentDetails.cpp
--- a/Engine/Source/Editor/DetailCustomizations/Private/AnimMontageSegmentDetails.cpp
+++ b/Engine/Source/Editor/DetailCustomizations/Private/AnimMontageSegmentDetails.cpp
@@ -86,13 +86,12 @@ void FAnimMontageSegmentDetails::CustomizeDetails( IDetailLayoutBuilder& DetailB
 	TSharedPtr<SWidget> ValueWidget;
 	FDetailWidgetRow Row;
 	PropertyRow.GetDefaultWidgets(NameWidget, ValueWidget, Row);
-
-	bool bAllowClear = !(ObjectProperty->PropertyFlags & CPF_NoClear);
+	PropertyRow.OverrideResetToDefault(FResetToDefaultOverride::Hide());
 
 	SAssignNew(ValueWidget, SObjectPropertyEntryBox)
 		.PropertyHandle(TargetPropertyHandle)
 		.AllowedClass(ObjectProperty->PropertyClass)
-		.AllowClear(bAllowClear)
+		.AllowClear(false)
 		.OnShouldFilterAsset(FOnShouldFilterAsset::CreateSP(this, &FAnimMontageSegmentDetails::OnShouldFilterAnimAsset));
 
 	PropertyRow.CustomWidget()
