Fix an ensure that is triggered when opening a force feedback effect. This ensure was caused when more properties were added to the Attenuation settings, but the property handle ensure was not updated.

#rnx

diff --git a/Engine/Source/Editor/DetailCustomizations/Private/AttenuationSettingsCustomizations.cpp b/Engine/Source/Editor/DetailCustomizations/Private/AttenuationSettingsCustomizations.cpp
--- a/Engine/Source/Editor/DetailCustomizations/Private/AttenuationSettingsCustomizations.cpp
+++ b/Engine/Source/Editor/DetailCustomizations/Private/AttenuationSettingsCustomizations.cpp
@@ -725,7 +725,7 @@ void FForceFeedbackAttenuationSettingsCustomization::CustomizeChildren(TSharedRe
 	EnumRestriction->AddHiddenValue(AttenuationDistanceModelEnum->GetNameStringByValue((uint8)EAttenuationDistanceModel::NaturalSound));
 	DistanceAlgorithmHandle->AddRestriction(EnumRestriction.ToSharedRef());
 
-	if (PropertyHandles.Num() != 7)
+	if (PropertyHandles.Num() != 10)
 	{
 		FString PropertyList;
 		for (auto It(PropertyHandles.CreateConstIterator()); It; ++It)
