Montage segment changes:
- Deprecated public access to AnimReference
    * Used for new behaviour that keeps track of PlayLength from referenced asset
    * In case the cached PlayLength does not match the current _and_ the cache playlength == than the segment length it can be deemed out-of-date due to a reimport
    * Out-of-date segments show in the OrangeAccent on the timeline rather than GreenAccent
    * Out-of-date segments will introduce log warnings on-load
- Patched up all paths using AnimReference (get/set)
- Modernized parts of EditorAnim* files
- Face-lifted details customization for AnimSegment
    * Scrollboxes now have (dynamic) UI/Clamp according to the context
    * Added behaviour for interactive operation when dragging the timeline (dragable) bars
- Implemented FAnimModel_AnimMontage::OnDataModelChanged to handle curve changes only - as play length etc is handled by montage editor itself

diff --git a/Engine/Plugins/Experimental/Animation/ContextualAnimation/Source/ContextualAnimation/Private/ContextualAnimUtilities.cpp b/Engine/Plugins/Experimental/Animation/ContextualAnimation/Source/ContextualAnimation/Private/ContextualAnimUtilities.cpp
--- a/Engine/Plugins/Experimental/Animation/ContextualAnimation/Source/ContextualAnimation/Private/ContextualAnimUtilities.cpp
+++ b/Engine/Plugins/Experimental/Animation/ContextualAnimation/Source/ContextualAnimation/Private/ContextualAnimUtilities.cpp
@@ -69,7 +69,7 @@ FTransform UContextualAnimUtilities::ExtractRootTransformFromAnimation(const UAn
 	{
 		if (const FAnimSegment* Segment = AnimMontage->SlotAnimTracks[0].AnimTrack.GetSegmentAtTime(Time))
 		{
-			if (const UAnimSequence* AnimSequence = Cast<UAnimSequence>(Segment->AnimReference))
+			if (const UAnimSequence* AnimSequence = Cast<UAnimSequence>(Segment->GetAnimReference()))
 			{
 				const float AnimSequenceTime = Segment->ConvertTrackPosToAnimPos(Time);
 				return AnimSequence->ExtractRootTrackTransform(AnimSequenceTime, nullptr);
diff --git a/Engine/Plugins/Experimental/Animation/MotionWarping/Source/MotionWarping/Private/MotionWarpingComponent.cpp b/Engine/Plugins/Experimental/Animation/MotionWarping/Source/MotionWarping/Private/MotionWarpingComponent.cpp
--- a/Engine/Plugins/Experimental/Animation/MotionWarping/Source/MotionWarping/Private/MotionWarpingComponent.cpp
+++ b/Engine/Plugins/Experimental/Animation/MotionWarping/Source/MotionWarping/Private/MotionWarpingComponent.cpp
@@ -94,7 +94,7 @@ FTransform UMotionWarpingUtilities::ExtractRootTransformFromAnimation(const UAni
 	{
 		if(const FAnimSegment* Segment = AnimMontage->SlotAnimTracks[0].AnimTrack.GetSegmentAtTime(Time))
 		{
-			if (const UAnimSequence* AnimSequence = Cast<UAnimSequence>(Segment->AnimReference))
+			if (const UAnimSequence* AnimSequence = Cast<UAnimSequence>(Segment->GetAnimReference()))
 			{
 				const float AnimSequenceTime = Segment->ConvertTrackPosToAnimPos(Time);
 				return AnimSequence->ExtractRootTrackTransform(AnimSequenceTime, nullptr);
@@ -295,21 +295,22 @@ void UMotionWarpingComponent::Update()
 			{
 				const FAnimTrack& AnimTrack = Montage->SlotAnimTracks[SlotIdx].AnimTrack;
 				const FAnimSegment* AnimSegment = AnimTrack.GetSegmentAtTime(PreviousPosition);
-				if (AnimSegment && AnimSegment->AnimReference)
+				const UAnimSequenceBase* AnimReference = AnimSegment->GetAnimReference();
+				if (AnimSegment && AnimReference)
 				{
-					for (const FAnimNotifyEvent& NotifyEvent : AnimSegment->AnimReference->Notifies)
+					for (const FAnimNotifyEvent& NotifyEvent : AnimReference->Notifies)
 					{
 						const UAnimNotifyState_MotionWarping* MotionWarpingNotify = NotifyEvent.NotifyStateClass ? Cast<UAnimNotifyState_MotionWarping>(NotifyEvent.NotifyStateClass) : nullptr;
 						if (MotionWarpingNotify)
 						{
 							if (MotionWarpingNotify->RootMotionModifier == nullptr)
 							{
-								UE_LOG(LogMotionWarping, Warning, TEXT("MotionWarpingComponent::Update. A motion warping window in %s doesn't have a valid root motion modifier!"), *GetNameSafe(AnimSegment->AnimReference));
+								UE_LOG(LogMotionWarping, Warning, TEXT("MotionWarpingComponent::Update. A motion warping window in %s doesn't have a valid root motion modifier!"), *GetNameSafe(AnimReference));
 								continue;
 							}
 
-							const float NotifyStartTime = FMath::Clamp(NotifyEvent.GetTriggerTime(), 0.f, AnimSegment->AnimReference->GetPlayLength());
-							const float NotifyEndTime = FMath::Clamp(NotifyEvent.GetEndTriggerTime(), 0.f, AnimSegment->AnimReference->GetPlayLength());
+							const float NotifyStartTime = FMath::Clamp(NotifyEvent.GetTriggerTime(), 0.f, AnimReference->GetPlayLength());
+							const float NotifyEndTime = FMath::Clamp(NotifyEvent.GetEndTriggerTime(), 0.f, AnimReference->GetPlayLength());
 
 							// Convert notify times from AnimSequence times to montage times
 							const float StartTime = (NotifyStartTime - AnimSegment->AnimStartTime) + AnimSegment->StartPos;
diff --git a/Engine/Source/Editor/AnimationEditor/Private/AnimationEditor.cpp b/Engine/Source/Editor/AnimationEditor/Private/AnimationEditor.cpp
--- a/Engine/Source/Editor/AnimationEditor/Private/AnimationEditor.cpp
+++ b/Engine/Source/Editor/AnimationEditor/Private/AnimationEditor.cpp
@@ -768,7 +768,7 @@ void FAnimationEditor::ConditionalRefreshEditor(UObject* InObject)
 
 			for (FAnimSegment& Segment : Slot.AnimTrack.AnimSegments)
 			{
-				if (Segment.AnimReference == InObject)
+				if (Segment.GetAnimReference() == InObject)
 				{
 					bInterestingAsset = true;
 					break;
diff --git a/Engine/Source/Editor/DetailCustomizations/Private/AnimMontageSegmentDetails.cpp b/Engine/Source/Editor/DetailCustomizations/Private/AnimMontageSegmentDetails.cpp
--- a/Engine/Source/Editor/DetailCustomizations/Private/AnimMontageSegmentDetails.cpp
+++ b/Engine/Source/Editor/DetailCustomizations/Private/AnimMontageSegmentDetails.cpp
@@ -17,8 +17,11 @@
 #include "PropertyCustomizationHelpers.h"
 #include "AnimPreviewInstance.h"
 #include "Slate/SceneViewport.h"
-#define LOCTEXT_NAMESPACE "AnimMontageSegmentDetails"
 #include "Settings/SkeletalMeshEditorSettings.h"
+#include "Widgets/Input/SNumericEntryBox.h"
+#include "Animation/EditorAnimSegment.h"
+
+#define LOCTEXT_NAMESPACE "AnimMontageSegmentDetails"
 
 /////////////////////////////////////////////////////////////////////////
 FAnimationSegmentViewportClient::FAnimationSegmentViewportClient(FPreviewScene& InPreviewScene, const TWeakPtr<SEditorViewport>& InEditorViewportWidget)
@@ -74,13 +77,13 @@ void FAnimMontageSegmentDetails::CustomizeDetails( IDetailLayoutBuilder& DetailB
 {
 	IDetailCategoryBuilder& SegmentCategory = DetailBuilder.EditCategory("Animation Segment", LOCTEXT("AnimationSegmentCategoryTitle", "Animation Segment") );
 
-	TSharedRef<IPropertyHandle> TargetPropertyHandle = DetailBuilder.GetProperty("AnimSegment.AnimReference");
+	AnimSegmentHandle = DetailBuilder.GetProperty(GET_MEMBER_NAME_CHECKED(UEditorAnimSegment, AnimSegment));
+	TSharedPtr<IPropertyHandle> TargetPropertyHandle = AnimSegmentHandle->GetChildHandle("AnimReference");
 	FProperty* TargetProperty = TargetPropertyHandle->GetProperty();
 
 	const FObjectPropertyBase* ObjectProperty = CastFieldChecked<const FObjectPropertyBase>(TargetProperty);
 
 	IDetailPropertyRow& PropertyRow = SegmentCategory.AddProperty(TargetPropertyHandle);
-	PropertyRow.DisplayName(LOCTEXT("AnimationReferenceLabel", "Animation Reference"));
 
 	TSharedPtr<SWidget> NameWidget;
 	TSharedPtr<SWidget> ValueWidget;
@@ -92,6 +95,7 @@ void FAnimMontageSegmentDetails::CustomizeDetails( IDetailLayoutBuilder& DetailB
 		.PropertyHandle(TargetPropertyHandle)
 		.AllowedClass(ObjectProperty->PropertyClass)
 		.AllowClear(false)
+		.OnObjectChanged(this, &FAnimMontageSegmentDetails::SetAnimationAsset)
 		.OnShouldFilterAsset(FOnShouldFilterAsset::CreateSP(this, &FAnimMontageSegmentDetails::OnShouldFilterAnimAsset));
 
 	PropertyRow.CustomWidget()
@@ -108,33 +112,162 @@ void FAnimMontageSegmentDetails::CustomizeDetails( IDetailLayoutBuilder& DetailB
 			ValueWidget.ToSharedRef()
 		];
 
-	SegmentCategory.AddProperty("AnimSegment.AnimStartTime").DisplayName( LOCTEXT("StartTimeLabel", "Start Time") );
-	SegmentCategory.AddProperty("AnimSegment.AnimEndTime").DisplayName( LOCTEXT("EndTimeLabel", "End Time") );
+	if (AnimSegmentHandle.IsValid())
+	{
+		AnimStartTimeProperty = AnimSegmentHandle->GetChildHandle(GET_MEMBER_NAME_CHECKED(FAnimSegment, AnimStartTime));
 
-	SegmentCategory.AddProperty("AnimSegment.AnimPlayRate").DisplayName( LOCTEXT("PlayRateLabel", "Play Rate") );
-	SegmentCategory.AddProperty("AnimSegment.LoopingCount").DisplayName( LOCTEXT("LoopCountLabel", "Loop Count") );
+		IDetailPropertyRow& StartPropertyRow = SegmentCategory.AddProperty(AnimStartTimeProperty);
+		StartPropertyRow.CustomWidget()
+		.NameContent()
+		[
+			AnimStartTimeProperty->CreatePropertyNameWidget()
+		]
+		.ValueContent()
+		[
+			SNew(SNumericEntryBox<float>)
+			.Font(DetailBuilder.GetDetailFont())
+			.AllowSpin(true)
+			.MinSliderValue(0.f)
+			.MinValue(0.f)
+			.MaxSliderValue_Raw(this, &FAnimMontageSegmentDetails::GetAnimationAssetPlayLength)
+			.MaxValue_Raw(this, &FAnimMontageSegmentDetails::GetAnimationAssetPlayLength)
+			.Value_Raw(this, &FAnimMontageSegmentDetails::GetStartTime)
+			.OnValueChanged_Raw(this, &FAnimMontageSegmentDetails::OnStartTimeChanged, ETextCommit::Default, true)
+			.OnValueCommitted_Raw(this, &FAnimMontageSegmentDetails::OnStartTimeChanged, false)
+		];
 
-	TSharedPtr<IPropertyHandle> InPropertyHandle = DetailBuilder.GetProperty("AnimSegment.AnimReference");
-	UObject *Object = NULL;
-	InPropertyHandle->GetValue(Object);
+		FResetToDefaultOverride Handler = FResetToDefaultOverride::Create
+		(
+			FIsResetToDefaultVisible::CreateLambda([this](TSharedPtr<IPropertyHandle> InAnimEndTimeProperty)
+			{
+				float Value = 0.f;
+				if (InAnimEndTimeProperty->GetValue(Value) == FPropertyAccess::Success)
+				{
+					return !FMath::IsNearlyEqual(Value, GetAnimationAssetPlayLength().Get(Value));
+				}
+
+				return false;
+			}),
+			FResetToDefaultHandler::CreateLambda([this](TSharedPtr<IPropertyHandle> InAnimEndTimeProperty)
+			{			
+				InAnimEndTimeProperty->SetValue(GetAnimationAssetPlayLength().Get(0.f));
+			})
+		);
+		
+		AnimEndTimeProperty = AnimSegmentHandle->GetChildHandle(GET_MEMBER_NAME_CHECKED(FAnimSegment, AnimEndTime));
+		SegmentCategory.AddProperty(AnimEndTimeProperty).CustomWidget()
+		.NameContent()
+		[
+			AnimEndTimeProperty->CreatePropertyNameWidget()
+		]
+		.OverrideResetToDefault(Handler)
+		.ValueContent()
+		[
+			SNew(SNumericEntryBox<float>)
+			.Font(DetailBuilder.GetDetailFont())
+			.AllowSpin(true)
+			.MinSliderValue(0.f)
+			.MinValue(0.f)
+			.MaxSliderValue_Raw(this, &FAnimMontageSegmentDetails::GetAnimationAssetPlayLength)
+			.MaxValue_Raw(this, &FAnimMontageSegmentDetails::GetAnimationAssetPlayLength)
+			.Value_Raw(this, &FAnimMontageSegmentDetails::GetEndTime)
+			.OnValueChanged_Raw(this, &FAnimMontageSegmentDetails::OnEndTimeChanged, ETextCommit::Default, true)
+			.OnValueCommitted_Raw(this, &FAnimMontageSegmentDetails::OnEndTimeChanged, false)
+		];
 
-	UAnimSequenceBase *AnimRef = Cast<UAnimSequenceBase>(Object);
-	USkeleton *Skeleton = NULL;
-	if(AnimRef != NULL)
-	{
-		Skeleton = AnimRef->GetSkeleton();
+		const TSharedPtr<IPropertyHandle> AnimPlayRateProperty = AnimSegmentHandle->GetChildHandle(GET_MEMBER_NAME_CHECKED(FAnimSegment, AnimPlayRate));
+		
+		IDetailPropertyRow& PlayRatePropertyRow = SegmentCategory.AddProperty(AnimPlayRateProperty);
+
+		PlayRatePropertyRow.CustomWidget()
+		.NameContent()
+		[
+			AnimPlayRateProperty->CreatePropertyNameWidget()
+		]
+		.ValueContent()
+		[
+			SNew(SNumericEntryBox<float>)
+			.Font(DetailBuilder.GetDetailFont())
+			.AllowSpin(true)
+			.MinSliderValue(-1.f)
+			.MinValue(-32.f)
+			.MaxSliderValue(1.f)
+			.MaxValue(32.f)
+			.Value_Raw(this, &FAnimMontageSegmentDetails::GetPlayRate)
+			.OnValueChanged_Lambda([AnimPlayRateProperty](float InValue)
+			{
+				if(AnimPlayRateProperty.IsValid() && !FMath::IsNearlyZero(InValue))
+				{
+					AnimPlayRateProperty->SetValue(InValue, EPropertyValueSetFlags::InteractiveChange);
+				}
+			})
+			.OnValueCommitted_Lambda([AnimPlayRateProperty](float InValue, ETextCommit::Type InCommitType)
+			{
+                if(AnimPlayRateProperty.IsValid() && !FMath::IsNearlyZero(InValue))
+                {
+					AnimPlayRateProperty->SetValue(InValue);
+                }
+			})
+		];
+
+		const TSharedPtr<IPropertyHandle> LoopingCountProperty = AnimSegmentHandle->GetChildHandle(GET_MEMBER_NAME_CHECKED(FAnimSegment, LoopingCount));
+		
+		IDetailPropertyRow& LoopingPropertyRow = SegmentCategory.AddProperty(LoopingCountProperty);
+
+		LoopingPropertyRow.CustomWidget()
+		.NameContent()
+		[
+			LoopingCountProperty->CreatePropertyNameWidget()
+		]
+		.ValueContent()
+		[
+			SNew(SNumericEntryBox<int32>)
+			.Font(DetailBuilder.GetDetailFont())
+			.AllowSpin(true)
+			.MinSliderValue(1)
+			.MinValue(1)
+			.MaxSliderValue(4)
+			.MaxValue(32)
+			.Value_Lambda([this]
+			{
+				if(const FAnimSegment* AnimSegmentPtr = GetAnimationSegment())
+				{
+					return AnimSegmentPtr->LoopingCount;
+				}
+				return 0;
+			})
+			.OnValueChanged_Lambda([LoopingCountProperty](int32 InValue)
+			{
+				if(LoopingCountProperty.IsValid() && InValue != 0)
+				{
+					LoopingCountProperty->SetValue(InValue, EPropertyValueSetFlags::InteractiveChange);
+				}
+			})
+			.OnValueCommitted_Lambda([LoopingCountProperty](int32 InValue, ETextCommit::Type InCommitType)
+			{
+				if(LoopingCountProperty.IsValid() && InValue != 0)
+				{
+					LoopingCountProperty->SetValue(InValue);
+				}
+			})
+		];
 	}
 
+	SegmentCategory.AddProperty(AnimSegmentHandle->GetChildHandle(GET_MEMBER_NAME_CHECKED(FAnimSegment, StartPos)));
+
+	const UAnimSequenceBase* AnimRef = GetAnimationAsset();
+	const USkeleton* Skeleton = AnimRef ? AnimRef->GetSkeleton() : nullptr;
+
 	SegmentCategory.AddCustomRow(FText::GetEmpty(), false)
 	[
 		SNew(SAnimationSegmentViewport)
-		.Skeleton(Skeleton)
-		.AnimRef(AnimRef)
-		.AnimRefPropertyHandle(DetailBuilder.GetProperty("AnimSegment.AnimReference"))
-		.StartTimePropertyHandle(DetailBuilder.GetProperty("AnimSegment.AnimStartTime"))
-		.EndTimePropertyHandle(DetailBuilder.GetProperty("AnimSegment.AnimEndTime"))
-		.PlayRatePropertyHandle(DetailBuilder.GetProperty("AnimSegment.AnimPlayRate"))
-	];	
+		.AnimRef_Lambda([this]() { return GetAnimationAsset(); })
+		.StartTime_Raw(this, &FAnimMontageSegmentDetails::GetStartTime)
+		.EndTime_Raw(this, &FAnimMontageSegmentDetails::GetEndTime)
+		.PlayRate_Raw(this, &FAnimMontageSegmentDetails::GetPlayRate)
+		.OnStartTimeChanged_Lambda([this](float InValue, bool bInteractive) { OnStartTimeChanged(InValue, ETextCommit::Default, bInteractive); })
+		.OnEndTimeChanged_Lambda([this](float InValue, bool bInteractive) { OnEndTimeChanged(InValue, ETextCommit::Default, bInteractive); })
+	];
 }
 
 bool FAnimMontageSegmentDetails::OnShouldFilterAnimAsset(const FAssetData& AssetData) const
@@ -142,6 +275,106 @@ bool FAnimMontageSegmentDetails::OnShouldFilterAnimAsset(const FAssetData& Asset
 	return AssetData.GetClass() == UAnimMontage::StaticClass();
 }
 
+const UAnimSequenceBase* FAnimMontageSegmentDetails::GetAnimationAsset() const
+{
+	if (const FAnimSegment* AnimSegment = GetAnimationSegment())
+	{
+		return AnimSegment->GetAnimReference().Get();
+	}
+
+	return nullptr;
+}
+
+void FAnimMontageSegmentDetails::SetAnimationAsset(const FAssetData& InAssetData)
+{
+	if (FAnimSegment* AnimSegment = GetAnimationSegment())
+	{
+		UAnimSequenceBase* AnimSequenceBase = Cast<UAnimSequenceBase>(InAssetData.GetAsset());
+		AnimSegment->SetAnimReference(AnimSequenceBase);
+	}
+}
+
+void FAnimMontageSegmentDetails::OnStartTimeChanged(float InValue, ETextCommit::Type InCommitType, bool bInteractive)
+{
+	if (FAnimSegment* AnimSegment = GetAnimationSegment())
+	{
+		const float ValueToSet = FMath::Min(AnimSegment->AnimEndTime, InValue);
+		AnimStartTimeProperty->SetValue(ValueToSet, bInteractive ? EPropertyValueSetFlags::InteractiveChange : EPropertyValueSetFlags::DefaultFlags);			
+	}
+}
+
+TOptional<float> FAnimMontageSegmentDetails::GetStartTime() const
+{
+	TOptional<float> TimeValue;
+	
+	if (const FAnimSegment* AnimSegment = GetAnimationSegment())
+	{
+		TimeValue = AnimSegment->AnimStartTime;
+	}
+
+	return TimeValue;
+}
+
+void FAnimMontageSegmentDetails::OnEndTimeChanged(float InValue, ETextCommit::Type InCommitType, bool bInteractive)
+{
+	if (FAnimSegment* AnimSegment = GetAnimationSegment())
+	{
+		const float ValueToSet = FMath::Max(AnimSegment->AnimStartTime, InValue);
+		AnimEndTimeProperty->SetValue(ValueToSet, bInteractive ? EPropertyValueSetFlags::InteractiveChange : EPropertyValueSetFlags::DefaultFlags);		
+	}
+}
+
+TOptional<float> FAnimMontageSegmentDetails::GetEndTime() const
+{
+	TOptional<float> TimeValue;
+	if (const FAnimSegment* AnimSegment = GetAnimationSegment())
+	{
+		TimeValue = AnimSegment->AnimEndTime;
+	}
+
+	return TimeValue;
+}
+
+TOptional<float> FAnimMontageSegmentDetails::GetPlayRate() const
+{
+	TOptional<float> Value;
+	if (const FAnimSegment* AnimSegment = GetAnimationSegment())
+	{
+		Value = AnimSegment->AnimPlayRate;
+	}
+
+	return Value;
+}
+
+FAnimSegment* FAnimMontageSegmentDetails::GetAnimationSegment() const
+{
+	if (AnimSegmentHandle.IsValid())
+	{
+		void* StructPtr = nullptr;	
+		AnimSegmentHandle->GetValueData(StructPtr);
+		return (FAnimSegment*)StructPtr;
+	}
+
+	return nullptr;
+}
+
+bool FAnimMontageSegmentDetails::CanEditSegmentProperties() const
+{
+	return true;
+}
+
+TOptional<float> FAnimMontageSegmentDetails::GetAnimationAssetPlayLength() const
+{
+	TOptional<float> Value;
+	
+	if (const UAnimSequenceBase* SequenceBase = GetAnimationAsset())
+	{
+		Value = SequenceBase->GetPlayLength();
+	}
+
+	return Value;
+}
+
 /////////////////////////////////////////////////
 ////////////////////////////////////////////////
 
@@ -182,19 +415,18 @@ void SAnimationSegmentViewport::CleanupComponent(USceneComponent* Component)
 }
 
 SAnimationSegmentViewport::SAnimationSegmentViewport()
-	: PreviewScene(FPreviewScene::ConstructionValues())
+	: CurrentAnimSequenceBase(nullptr), PreviewScene(FPreviewScene::ConstructionValues())
 {
 }
 
 void SAnimationSegmentViewport::Construct(const FArguments& InArgs)
 {
-	TargetSkeleton = InArgs._Skeleton;
-	AnimRef = InArgs._AnimRef;
-	
-	AnimRefPropertyHandle = InArgs._AnimRefPropertyHandle;
-	StartTimePropertyHandle = InArgs._StartTimePropertyHandle;
-	EndTimePropertyHandle = InArgs._EndTimePropertyHandle;
-	PlayRatePropertyHandle = InArgs._PlayRatePropertyHandle;
+	AnimationRefAttribute = InArgs._AnimRef;
+	StartTimeAttribute = InArgs._StartTime;
+	EndTimeAttribute = InArgs._EndTime;
+	PlayRateAttribute = InArgs._PlayRate;
+	OnStartTimeChanged = InArgs._OnStartTimeChanged;
+	OnEndTimeChanged = InArgs._OnEndTimeChanged;
 
 	this->ChildSlot
 	[
@@ -228,7 +460,8 @@ void SAnimationSegmentViewport::Construct(const FArguments& InArgs)
 			.ViewInputMax(this, &SAnimationSegmentViewport::GetViewMaxInput)
 			.PreviewInstance(this, &SAnimationSegmentViewport::GetPreviewInstance)
 			.DraggableBars(this, &SAnimationSegmentViewport::GetBars)
-			.OnBarDrag(this, &SAnimationSegmentViewport::OnBarDrag)
+			.OnBarDrag(this, &SAnimationSegmentViewport::OnBarDrag, true)
+			.OnBarCommit(this, &SAnimationSegmentViewport::OnBarDrag, false)
 			.bAllowZoom(true)
 		]
 	];
@@ -259,58 +492,45 @@ void SAnimationSegmentViewport::Construct(const FArguments& InArgs)
 
 void SAnimationSegmentViewport::InitSkeleton()
 {
-	UObject *Object = NULL;
-	AnimRefPropertyHandle->GetValue(Object);
-	UAnimSequenceBase *AnimSequence = Cast<UAnimSequenceBase>(Object);
-	USkeleton *Skeleton = NULL;
-	if(AnimSequence != NULL)
+	UAnimSequenceBase* AnimSequenceBase = const_cast<UAnimSequenceBase*>(AnimationRefAttribute.Get());
+	if (PreviewComponent && AnimSequenceBase && AnimSequenceBase != CurrentAnimSequenceBase)
 	{
-		Skeleton = AnimSequence->GetSkeleton();
-	}
-
-	if( PreviewComponent != NULL && Skeleton != NULL )
-	{
-		USkeletalMesh* PreviewMesh = Skeleton->GetAssetPreviewMesh(AnimSequence);
-		if (PreviewMesh)
+		USkeleton* Skeleton = AnimSequenceBase->GetSkeleton();
+		USkeletalMesh* PreviewMesh = Skeleton ? Skeleton->GetAssetPreviewMesh(AnimSequenceBase) : nullptr;
+		if (Skeleton && PreviewMesh)
 		{
-			UAnimSingleNodeInstance * Preview = PreviewComponent->PreviewInstance;
-			if((Preview == NULL || Preview->GetCurrentAsset() != AnimSequence) ||
-				(PreviewComponent->SkeletalMesh != PreviewMesh))
+			UAnimSingleNodeInstance* Preview = PreviewComponent->PreviewInstance;
+			if((Preview == nullptr || Preview->GetCurrentAsset() != AnimSequenceBase) || (PreviewComponent->SkeletalMesh != PreviewMesh))
 			{
-				float PlayRate;
-				PlayRatePropertyHandle->GetValue(PlayRate);
+				const float PlayRate = PlayRateAttribute.Get().Get(1.f);
 
 				PreviewComponent->SetSkeletalMesh(PreviewMesh);
-				PreviewComponent->EnablePreview(true, AnimSequence);
+				PreviewComponent->EnablePreview(true, AnimSequenceBase);
 				PreviewComponent->PreviewInstance->SetLooping(true);
 				PreviewComponent->SetPlayRate(PlayRate);
 
 				//Place the camera at a good viewer position
-				FVector NewPosition = LevelViewportClient->GetViewLocation();
-				NewPosition.Normalize();
+				const FVector NewPosition = LevelViewportClient->GetViewLocation().GetSafeNormal();
 				LevelViewportClient->SetViewLocation(NewPosition * (PreviewMesh->GetImportedBounds().SphereRadius*1.5f));
+
+				CurrentAnimSequenceBase = AnimSequenceBase;
 			}
 		}
 	}
-
-	TargetSkeleton = Skeleton;
 }
 
 void SAnimationSegmentViewport::Tick( const FGeometry& AllottedGeometry, const double InCurrentTime, const float InDeltaTime )
 {
 	class UDebugSkelMeshComponent* Component = PreviewComponent;
-
-	FString TargetSkeletonName = TargetSkeleton ? TargetSkeleton->GetName() : FName(NAME_None).ToString();
-
-	if (Component != NULL)
+	const FString TargetSkeletonName = CurrentAnimSequenceBase && CurrentAnimSequenceBase->GetSkeleton() ? CurrentAnimSequenceBase->GetSkeleton()->GetName() : TEXT("None");
+	if (Component != nullptr)
 	{
 		// Reinit the skeleton if the anim ref has changed
 		InitSkeleton();
 
-		float Start, End, PlayRate;
-		StartTimePropertyHandle->GetValue(Start);
-		EndTimePropertyHandle->GetValue(End);
-		PlayRatePropertyHandle->GetValue(PlayRate);
+		const float Start = StartTimeAttribute.Get().Get(0.f);
+		const float End = EndTimeAttribute.Get().Get(0.f);
+		const float PlayRate = PlayRateAttribute.Get().Get(0.f);
 
 		if (Component->PreviewInstance->GetCurrentTime() > End || Component->PreviewInstance->GetCurrentTime() < Start)
 		{
@@ -390,9 +610,8 @@ float SAnimationSegmentViewport::GetViewMaxInput() const
 
 TArray<float> SAnimationSegmentViewport::GetBars() const
 {
-	float Start, End;
-	StartTimePropertyHandle->GetValue(Start);
-	EndTimePropertyHandle->GetValue(End);
+	const float Start = StartTimeAttribute.Get().Get(0.f);
+	const float End = EndTimeAttribute.Get().Get(0.f);
 
 	TArray<float> Bars;
 	Bars.Add(Start);
@@ -401,15 +620,15 @@ TArray<float> SAnimationSegmentViewport::GetBars() const
 	return Bars;
 }
 
-void SAnimationSegmentViewport::OnBarDrag(int32 Index, float Position)
+void SAnimationSegmentViewport::OnBarDrag(int32 Index, float Position, bool bInteractive)
 {
 	if(Index==0)
 	{
-		StartTimePropertyHandle->SetValue(Position);
+		OnStartTimeChanged.ExecuteIfBound(Position, bInteractive);
 	}
 	else if(Index==1)
 	{
-		EndTimePropertyHandle->SetValue(Position);
+		OnEndTimeChanged.ExecuteIfBound(Position, bInteractive);
 	}
 }
 
@@ -452,6 +671,7 @@ void SAnimationSegmentScrubPanel::Construct( const SAnimationSegmentScrubPanel::
 				.IsRealtimeStreamingMode(this, &SAnimationSegmentScrubPanel::IsRealtimeStreamingMode)
 				.DraggableBars(InArgs._DraggableBars)
 				.OnBarDrag(InArgs._OnBarDrag)
+				.OnBarCommit(InArgs._OnBarCommit)
 				.OnTickPlayback(InArgs._OnTickPlayback)
 			]
 		];
diff --git a/Engine/Source/Editor/DetailCustomizations/Private/AnimMontageSegmentDetails.h b/Engine/Source/Editor/DetailCustomizations/Private/AnimMontageSegmentDetails.h
--- a/Engine/Source/Editor/DetailCustomizations/Private/AnimMontageSegmentDetails.h
+++ b/Engine/Source/Editor/DetailCustomizations/Private/AnimMontageSegmentDetails.h
@@ -31,7 +31,27 @@ public:
 
 	/** IDetailCustomization interface */
 	virtual void CustomizeDetails( IDetailLayoutBuilder& DetailBuilder ) override;
+protected:
 	bool OnShouldFilterAnimAsset(const FAssetData& AssetData) const;
+	
+	const UAnimSequenceBase* GetAnimationAsset() const;
+	FAnimSegment* GetAnimationSegment() const;
+	bool CanEditSegmentProperties() const;
+
+	void SetAnimationAsset(const FAssetData& InAssetData);
+	void OnStartTimeChanged(float InValue, ETextCommit::Type InCommitType, bool bInteractive = false);	
+	void OnEndTimeChanged(float InValue, ETextCommit::Type InCommitType, bool bInteractive = false);
+
+	TOptional<float> GetStartTime() const;
+	TOptional<float> GetEndTime() const;	
+	TOptional<float> GetAnimationAssetPlayLength() const;
+	TOptional<float> GetPlayRate() const;
+
+protected:
+	TSharedPtr<IPropertyHandle> AnimSegmentHandle;
+	TSharedPtr<IPropertyHandle> AnimationAssetHandle;
+	TSharedPtr<IPropertyHandle> AnimStartTimeProperty;
+	TSharedPtr<IPropertyHandle> AnimEndTimeProperty;
 };
 
 
@@ -39,20 +59,16 @@ public:
 class SAnimationSegmentViewport : public SCompoundWidget
 {
 public:
+	DECLARE_DELEGATE_TwoParams(FOnValueChanged, float, bool);
+	
 	SLATE_BEGIN_ARGS( SAnimationSegmentViewport )
-		: _Skeleton(NULL)
-		, _AnimRef(NULL)
-		, _IsEditable(true)
 	{}
-		
-		SLATE_ARGUMENT( USkeleton*, Skeleton )
-		SLATE_ARGUMENT( UAnimSequenceBase*, AnimRef)
-		SLATE_ARGUMENT( TSharedPtr<IPropertyHandle>, AnimRefPropertyHandle )
-		SLATE_ARGUMENT( TSharedPtr<IPropertyHandle>, StartTimePropertyHandle )
-		SLATE_ARGUMENT( TSharedPtr<IPropertyHandle>, EndTimePropertyHandle )
-		SLATE_ARGUMENT(TSharedPtr<IPropertyHandle>, PlayRatePropertyHandle);
-
-		SLATE_ATTRIBUTE( bool, IsEditable )	
+		SLATE_ATTRIBUTE(const UAnimSequenceBase*, AnimRef)
+		SLATE_ATTRIBUTE(TOptional<float>, StartTime)
+		SLATE_ATTRIBUTE(TOptional<float>, EndTime)
+		SLATE_ATTRIBUTE(TOptional<float>, PlayRate)
+		SLATE_EVENT(FOnValueChanged, OnStartTimeChanged)
+		SLATE_EVENT(FOnValueChanged, OnEndTimeChanged)
 	SLATE_END_ARGS()
 	
 public:
@@ -65,25 +81,25 @@ public:
 	void RefreshViewport();
 
 private:
-
 	void InitSkeleton();
 
 	TSharedPtr<FEditorViewportClient> LevelViewportClient;
-	
-	TSharedPtr<IPropertyHandle> AnimRefPropertyHandle;
-	TSharedPtr<IPropertyHandle> StartTimePropertyHandle;
-	TSharedPtr<IPropertyHandle> EndTimePropertyHandle;
-	TSharedPtr<IPropertyHandle> PlayRatePropertyHandle;
 
+	TAttribute<const UAnimSequenceBase*> AnimationRefAttribute;	
+	TAttribute<TOptional<float>> StartTimeAttribute;
+	TAttribute<TOptional<float>> EndTimeAttribute;
+	TAttribute<TOptional<float>> PlayRateAttribute;
+	FOnValueChanged OnStartTimeChanged;
+	FOnValueChanged OnEndTimeChanged;
+	
 	/** Slate viewport for rendering and I/O */
 	//TSharedPtr<class FSceneViewport> Viewport;
 	TSharedPtr<SViewport> ViewportWidget;
 
 	TSharedPtr<class FSceneViewport> SceneViewport;
 
-	/** Skeleton */
-	USkeleton* TargetSkeleton;
-	UAnimSequenceBase *AnimRef;
+
+	TObjectPtr<UAnimSequenceBase> CurrentAnimSequenceBase;
 
 	FPreviewScene PreviewScene;
 	class UDebugSkelMeshComponent* PreviewComponent;
@@ -103,7 +119,7 @@ public:
 
 	/** Optional, additional values to draw on the timeline **/
 	TArray<float> GetBars() const;
-	void OnBarDrag(int32 index, float newPos);
+	void OnBarDrag(int32 index, float newPos, bool bInteractive);
 	
 
 };
@@ -128,6 +144,7 @@ public:
 		SLATE_ARGUMENT( bool, bAllowZoom )
 		SLATE_ATTRIBUTE( TArray<float>, DraggableBars )
 		SLATE_EVENT( FOnScrubBarDrag, OnBarDrag )
+		SLATE_EVENT( FOnScrubBarCommit, OnBarCommit )
 		SLATE_EVENT( FOnTickPlayback, OnTickPlayback )
 	SLATE_END_ARGS()
 	
diff --git a/Engine/Source/Editor/KismetWidgets/Private/SScrubControlPanel.cpp b/Engine/Source/Editor/KismetWidgets/Private/SScrubControlPanel.cpp
--- a/Engine/Source/Editor/KismetWidgets/Private/SScrubControlPanel.cpp
+++ b/Engine/Source/Editor/KismetWidgets/Private/SScrubControlPanel.cpp
@@ -72,6 +72,7 @@ void SScrubControlPanel::Construct( const SScrubControlPanel::FArguments& InArgs
 				.bAllowZoom(InArgs._bAllowZoom)
 				.DraggableBars(InArgs._DraggableBars)
 				.OnBarDrag(InArgs._OnBarDrag)
+				.OnBarCommit(InArgs._OnBarCommit)
 			]
 		]
 
diff --git a/Engine/Source/Editor/KismetWidgets/Private/SScrubWidget.cpp b/Engine/Source/Editor/KismetWidgets/Private/SScrubWidget.cpp
--- a/Engine/Source/Editor/KismetWidgets/Private/SScrubWidget.cpp
+++ b/Engine/Source/Editor/KismetWidgets/Private/SScrubWidget.cpp
@@ -59,6 +59,7 @@ void SScrubWidget::Construct( const SScrubWidget::FArguments& InArgs )
 
 	DraggableBars = InArgs._DraggableBars;
 	OnBarDrag = InArgs._OnBarDrag;
+	OnBarCommit = InArgs._OnBarCommit;
 	bDisplayDrag = InArgs._DisplayDrag;
 	bDisplayAnimScrubBarEditing = InArgs._bDisplayAnimScrubBarEditing;
 	bMouseMovedDuringPanning = false;
@@ -264,6 +265,14 @@ FReply SScrubWidget::OnMouseButtonUp( const FGeometry& MyGeometry, const FPointe
 		if(DraggingBar)
 		{
 			DraggingBar = false;
+
+			if(OnBarCommit.IsBound())
+			{
+				const FVector2D CursorPos = MyGeometry.AbsoluteToLocal( MouseEvent.GetScreenSpacePosition() );
+				const FTrackScaleInfo ScaleInfo(ViewInputMin.Get(), ViewInputMax.Get(), 0.f, 0.f, MyGeometry.GetLocalSize());
+				const float NewDataPos = FMath::Clamp( ScaleInfo.LocalXToInput(CursorPos.X), ViewInputMin.Get(), ViewInputMax.Get() );				
+				OnBarCommit.ExecuteIfBound(DraggableBarIndex, NewDataPos);
+			}			
 		}
 		else if( bDragging )
 		{
diff --git a/Engine/Source/Editor/KismetWidgets/Public/SScrubControlPanel.h b/Engine/Source/Editor/KismetWidgets/Public/SScrubControlPanel.h
--- a/Engine/Source/Editor/KismetWidgets/Public/SScrubControlPanel.h
+++ b/Engine/Source/Editor/KismetWidgets/Public/SScrubControlPanel.h
@@ -80,6 +80,7 @@ public:
 		/** Optional, additional values to draw on the timeline **/
 		SLATE_ATTRIBUTE( TArray<float>, DraggableBars )
 		SLATE_EVENT( FOnScrubBarDrag, OnBarDrag )
+		SLATE_EVENT( FOnScrubBarCommit, OnBarCommit )
 		/** Called each frame during playback */
 		SLATE_EVENT( FOnTickPlayback, OnTickPlayback )
 	SLATE_END_ARGS()
diff --git a/Engine/Source/Editor/KismetWidgets/Public/SScrubWidget.h b/Engine/Source/Editor/KismetWidgets/Public/SScrubWidget.h
--- a/Engine/Source/Editor/KismetWidgets/Public/SScrubWidget.h
+++ b/Engine/Source/Editor/KismetWidgets/Public/SScrubWidget.h
@@ -23,6 +23,7 @@ DECLARE_DELEGATE_TwoParams( FOnCropAnimSequence, bool, float )
 DECLARE_DELEGATE_TwoParams( FOnAddAnimSequence, bool, int32 )
 DECLARE_DELEGATE_TwoParams( FOnAppendAnimSequence, bool, int32 )
 DECLARE_DELEGATE_TwoParams( FOnScrubBarDrag, int32, float)
+DECLARE_DELEGATE_TwoParams( FOnScrubBarCommit, int32, float)
 DECLARE_DELEGATE_OneParam( FOnReZeroAnimSequence, int32 )
 
 class KISMETWIDGETS_API SScrubWidget : public SCompoundWidget
@@ -72,7 +73,8 @@ public:
 		SLATE_EVENT( FOnReZeroAnimSequence, OnReZeroAnimSequence )
 		/** Optional, additional values to draw on the timeline **/
 		SLATE_ATTRIBUTE( TArray<float>, DraggableBars )
-		SLATE_EVENT( FOnScrubBarDrag, OnBarDrag)
+		SLATE_EVENT( FOnScrubBarDrag, OnBarDrag )
+		SLATE_EVENT( FOnScrubBarCommit, OnBarCommit )
 	SLATE_END_ARGS()
 
 
@@ -140,6 +142,7 @@ private:
 	/** Dragagble bars are generic lines drawn on the scrub widget that can be dragged with the mouse. This is very bare bones and just represents drawing/moving float values */
 	TAttribute<TArray<float>>	DraggableBars;
 	FOnScrubBarDrag OnBarDrag;
+	FOnScrubBarCommit OnBarCommit;
 
 	/** Distance Dragged **/
 	float DistanceDragged;
diff --git a/Engine/Source/Editor/Persona/Private/AnimTimeline/AnimModel_AnimMontage.cpp b/Engine/Source/Editor/Persona/Private/AnimTimeline/AnimModel_AnimMontage.cpp
--- a/Engine/Source/Editor/Persona/Private/AnimTimeline/AnimModel_AnimMontage.cpp
+++ b/Engine/Source/Editor/Persona/Private/AnimTimeline/AnimModel_AnimMontage.cpp
@@ -312,4 +312,30 @@ void FAnimModel_AnimMontage::ToggleSectionTimingDisplay()
 	bSectionTimingEnabled = !bSectionTimingEnabled;
 }
 
-#undef LOCTEXT_NAMESPACE
\ No newline at end of file
+void FAnimModel_AnimMontage::OnDataModelChanged(const EAnimDataModelNotifyType& NotifyType, UAnimDataModel* Model, const FAnimDataModelNotifPayload& PayLoad)
+{
+	NotifyCollector.Handle(NotifyType);
+
+	switch(NotifyType)
+	{ 
+	case EAnimDataModelNotifyType::CurveAdded:
+	case EAnimDataModelNotifyType::CurveRemoved:
+		{
+			if (NotifyCollector.IsNotWithinBracket())
+			{
+				RefreshTracks();
+			}
+			break;
+		}
+	case EAnimDataModelNotifyType::BracketClosed:
+		{
+			if (NotifyCollector.IsNotWithinBracket())
+			{
+				RefreshTracks();
+			}
+			break;
+		}
+	}
+}
+
+#undef LOCTEXT_NAMESPACE
diff --git a/Engine/Source/Editor/Persona/Private/AnimTimeline/AnimModel_AnimMontage.h b/Engine/Source/Editor/Persona/Private/AnimTimeline/AnimModel_AnimMontage.h
--- a/Engine/Source/Editor/Persona/Private/AnimTimeline/AnimModel_AnimMontage.h
+++ b/Engine/Source/Editor/Persona/Private/AnimTimeline/AnimModel_AnimMontage.h
@@ -61,6 +61,9 @@ public:
 	bool IsSectionTimingDisplayEnabled() const;
 	void ToggleSectionTimingDisplay();
 
+protected:	
+	virtual void OnDataModelChanged(const EAnimDataModelNotifyType& NotifyType, UAnimDataModel* Model, const FAnimDataModelNotifPayload& PayLoad) override;
+
 private:
 	/** The anim montage we wrap */
 	UAnimMontage* AnimMontage;
diff --git a/Engine/Source/Editor/Persona/Private/AnimTimeline/AnimModel_AnimSequenceBase.h b/Engine/Source/Editor/Persona/Private/AnimTimeline/AnimModel_AnimSequenceBase.h
--- a/Engine/Source/Editor/Persona/Private/AnimTimeline/AnimModel_AnimSequenceBase.h
+++ b/Engine/Source/Editor/Persona/Private/AnimTimeline/AnimModel_AnimSequenceBase.h
@@ -67,7 +67,7 @@ protected:
 	void RefreshAttributeTracks();
 
 	/** Callback for any change made to the UAnimDataModel embedded in the AnimSequenceBase instance this represents */
-	void OnDataModelChanged(const EAnimDataModelNotifyType& NotifyType, UAnimDataModel* Model, const FAnimDataModelNotifPayload& PayLoad);
+	virtual void OnDataModelChanged(const EAnimDataModelNotifyType& NotifyType, UAnimDataModel* Model, const FAnimDataModelNotifPayload& PayLoad);
 
 private:
 	/** UI handlers */
@@ -103,6 +103,6 @@ private:
 
 	/** Display flags for notifies track */
 	bool NotifiesTimingElementNodeDisplayFlags[ETimingElementType::Max];
-
+protected:
 	UE::Anim::FAnimDataModelNotifyCollector NotifyCollector;
 };
\ No newline at end of file
diff --git a/Engine/Source/Editor/Persona/Private/SAnimCompositePanel.cpp b/Engine/Source/Editor/Persona/Private/SAnimCompositePanel.cpp
--- a/Engine/Source/Editor/Persona/Private/SAnimCompositePanel.cpp
+++ b/Engine/Source/Editor/Persona/Private/SAnimCompositePanel.cpp
@@ -214,10 +214,18 @@ FLinearColor SAnimCompositePanel::HandleGetNodeColor(const FAnimSegment& InSegme
 {
 	static const FLinearColor DisabledColor(64, 64, 64);
 
-	if(InSegment.AnimReference != nullptr)
+	const FSlateColor OrangeAccent = FAppStyle::Get().GetSlateColor("Colors.AccentOrange");
+	FLinearColor OutOfDateNodeColor = OrangeAccent.GetSpecifiedColor();
+
+    if (InSegment.IsPlayLengthOutOfDate())
+    {
+	    return OutOfDateNodeColor;
+    }
+
+	if(const TObjectPtr<UAnimSequenceBase> AnimReference = InSegment.GetAnimReference())
 	{
 		FAssetToolsModule& AssetToolsModule = FModuleManager::LoadModuleChecked<FAssetToolsModule>(TEXT("AssetTools"));
-		TWeakPtr<IAssetTypeActions> AssetTypeActions = AssetToolsModule.Get().GetAssetTypeActionsForClass(InSegment.AnimReference->GetClass());
+		TWeakPtr<IAssetTypeActions> AssetTypeActions = AssetToolsModule.Get().GetAssetTypeActionsForClass(AnimReference->GetClass());
 		check(AssetTypeActions.IsValid());
 		return AssetTypeActions.Pin()->GetTypeColor().ReinterpretAsLinear();
 	}
diff --git a/Engine/Source/Editor/Persona/Private/SAnimMontagePanel.cpp b/Engine/Source/Editor/Persona/Private/SAnimMontagePanel.cpp
--- a/Engine/Source/Editor/Persona/Private/SAnimMontagePanel.cpp
+++ b/Engine/Source/Editor/Persona/Private/SAnimMontagePanel.cpp
@@ -254,8 +254,8 @@ bool SAnimMontagePanel::IsDiffererentFromParent(FName SlotName, int32 SegmentIdx
 
 			if (ParentTrack && ParentTrack->AnimSegments.IsValidIndex(SegmentIdx))
 			{
-				UAnimSequenceBase* SourceAsset = ParentTrack->AnimSegments[SegmentIdx].AnimReference;
-				return (SourceAsset != Segment.AnimReference);
+				UAnimSequenceBase* SourceAsset = ParentTrack->AnimSegments[SegmentIdx].GetAnimReference();
+				return (SourceAsset != Segment.GetAnimReference());
 			}
 		}
 	}
@@ -277,7 +277,7 @@ void SAnimMontagePanel::ReplaceAnimationMapping(FName SlotName, int32 SegmentIdx
 
 			if (ParentTrack && ParentTrack->AnimSegments.IsValidIndex(SegmentIdx))
 			{
-				UAnimSequenceBase* SourceAsset = ParentTrack->AnimSegments[SegmentIdx].AnimReference;
+				UAnimSequenceBase* SourceAsset = ParentTrack->AnimSegments[SegmentIdx].GetAnimReference();
 				if (Montage->RemapAsset(SourceAsset, NewSequenceBase))
 				{
 					// success
@@ -592,8 +592,12 @@ void SAnimMontagePanel::Update()
 		OnSetMontagePreviewSlot(CurrentPreviewSlot);
 
 		int32 ColorIdx=0;
-		
-		FLinearColor NodeColor = FLinearColor(0.f, 0.5f, 0.0f, 0.5f);
+
+		const FSlateColor GreenAccent = FAppStyle::Get().GetSlateColor("Colors.AccentGreen");
+		FLinearColor NodeColor = GreenAccent.GetSpecifiedColor();
+
+		const FSlateColor OrangeAccent = FAppStyle::Get().GetSlateColor("Colors.AccentOrange");
+		FLinearColor OutOfDateNodeColor = OrangeAccent.GetSpecifiedColor();
 
 		TSharedPtr<SVerticalBox> MontageSlots;
 		PanelArea->SetContent(
@@ -639,7 +643,15 @@ void SAnimMontagePanel::Update()
 						.ViewInputMin(ViewInputMin)
 						.ViewInputMax(ViewInputMax)
 						.bChildAnimMontage(bChildAnimMontage)
-						.OnGetNodeColor_Lambda([NodeColor](const FAnimSegment& InSegment){ return NodeColor; })
+						.OnGetNodeColor_Lambda([NodeColor, OutOfDateNodeColor](const FAnimSegment& InSegment)
+						{
+							if (InSegment.IsPlayLengthOutOfDate())
+							{
+								return OutOfDateNodeColor;
+							}
+							
+							return NodeColor;
+						})
 						.OnPreAnimUpdate(this, &SAnimMontagePanel::PreAnimUpdate)
 						.OnPostAnimUpdate(this, &SAnimMontagePanel::PostAnimUpdate)
 						.OnAnimReplaceMapping(this, &SAnimMontagePanel::ReplaceAnimationMapping)
@@ -662,7 +674,15 @@ void SAnimMontagePanel::Update()
 						.ViewInputMin(ViewInputMin)
 						.ViewInputMax(ViewInputMax)
 						.bChildAnimMontage(bChildAnimMontage)
-						.OnGetNodeColor_Lambda([NodeColor](const FAnimSegment& InSegment){ return NodeColor; })
+						.OnGetNodeColor_Lambda([NodeColor, OutOfDateNodeColor](const FAnimSegment& InSegment)
+						{
+							if (InSegment.IsPlayLengthOutOfDate())
+							{
+								return OutOfDateNodeColor;
+							}
+							
+							return NodeColor;
+						})
 						.TrackMaxValue(this, &SAnimMontagePanel::GetSequenceLength)
 						.TrackNumDiscreteValues(Montage->GetNumberOfSampledKeys())
 						.OnAnimSegmentNodeClicked(this, &SAnimMontagePanel::ShowSegmentInDetailsView, SlotAnimIdx)
diff --git a/Engine/Source/Editor/Persona/Private/SAnimSegmentsPanel.cpp b/Engine/Source/Editor/Persona/Private/SAnimSegmentsPanel.cpp
--- a/Engine/Source/Editor/Persona/Private/SAnimSegmentsPanel.cpp
+++ b/Engine/Source/Editor/Persona/Private/SAnimSegmentsPanel.cpp
@@ -203,8 +203,7 @@ FString	SAnimSegmentsPanel::GetAnimSegmentName(int32 AnimSegmentIndex) const
 	if (ValidIndex(AnimSegmentIndex))
 	{
 		FString TitleLabel;
-		UAnimSequenceBase* AnimReference = AnimTrack->AnimSegments[AnimSegmentIndex].AnimReference;
-		if(AnimReference)
+		if(const UAnimSequenceBase* AnimReference = AnimTrack->AnimSegments[AnimSegmentIndex].GetAnimReference())
 		{
 			FString AssetName = AnimReference->GetName();
 			if (AnimTrack->AnimSegments[AnimSegmentIndex].IsValid() == false)
@@ -231,8 +230,7 @@ FText SAnimSegmentsPanel::GetAnimSegmentDetailedInfo(int32 AnimSegmentIndex) con
 	if (ValidIndex(AnimSegmentIndex))
 	{
 		FAnimSegment& AnimSegment = AnimTrack->AnimSegments[AnimSegmentIndex];
-		UAnimSequenceBase * Anim = AnimSegment.AnimReference;
-		if ( Anim != NULL )
+		if (const UAnimSequenceBase* AnimReference = AnimTrack->AnimSegments[AnimSegmentIndex].GetAnimReference())
 		{
 			static const FNumberFormattingOptions FormatOptions = FNumberFormattingOptions()
 				.SetMinimumFractionalDigits(2)
@@ -240,11 +238,12 @@ FText SAnimSegmentsPanel::GetAnimSegmentDetailedInfo(int32 AnimSegmentIndex) con
 
 			if (AnimTrack->AnimSegments[AnimSegmentIndex].IsValid())
 			{
-				return FText::Format(LOCTEXT("AnimSegmentPanel_GetAnimSegmentDetailedInfoFmt", "{0} {1}"), FText::FromString(Anim->GetName()), FText::AsNumber(AnimSegment.GetLength(), &FormatOptions));
+				return FText::Format(LOCTEXT("AnimSegmentPanel_GetAnimSegmentDetailedInfoFmt", "{0} {1} {2}"), FText::FromString(AnimReference->GetName()), FText::AsNumber(AnimSegment.GetLength(), &FormatOptions), 
+					AnimTrack->AnimSegments[AnimSegmentIndex].IsPlayLengthOutOfDate() ? LOCTEXT("AnimSegmentPanel_GetAnimSegmentDetailedInfoFmt_Warning_PlayTimeIncorrect", "(segment length does not match animation play length)") : FText::FromString(FString()));
 			}
 			else
 			{
-				return FText::Format(LOCTEXT("AnimSegmentPanel_GetAnimSegmentDetailedInfoFmt_Error_RecursiveReference", "{0} {1} - ERROR: Recursive Reference Found"), FText::FromString(Anim->GetName()), FText::AsNumber(AnimSegment.GetLength(), &FormatOptions));  
+				return FText::Format(LOCTEXT("AnimSegmentPanel_GetAnimSegmentDetailedInfoFmt_Error_RecursiveReference", "{0} {1} - ERROR: Recursive Reference Found"), FText::FromString(AnimReference->GetName()), FText::AsNumber(AnimSegment.GetLength(), &FormatOptions));  
 			}
 			
 		}
@@ -318,11 +317,8 @@ void SAnimSegmentsPanel::AddAnimSegment( UAnimSequenceBase* NewSequenceBase, flo
 	OnPreAnimUpdateDelegate.ExecuteIfBound();
 
 	FAnimSegment NewSegment;
-	NewSegment.AnimReference = NewSequenceBase;
-	NewSegment.AnimStartTime = 0.f;
-	NewSegment.AnimEndTime = NewSequenceBase->GetPlayLength();
-	NewSegment.AnimPlayRate = 1.f;
-	NewSegment.LoopingCount = 1;
+	NewSegment.SetAnimReference(NewSequenceBase, true);
+
 	NewSegment.StartPos = NewStartPos;
 
 	AnimTrack->AnimSegments.Add(NewSegment);
@@ -334,7 +330,7 @@ void SAnimSegmentsPanel::ReplaceAnimSegment(int32 AnimSegmentIndex, UAnimSequenc
 	const FScopedTransaction Transaction(LOCTEXT("AnimSegmentPanel_ReplaceSegment", "Replace Segment"));
 	if (AnimTrack->AnimSegments.IsValidIndex(AnimSegmentIndex))
 	{
-		UAnimSequenceBase* OldSequenceBase = AnimTrack->AnimSegments[AnimSegmentIndex].AnimReference;
+		UAnimSequenceBase* OldSequenceBase = AnimTrack->AnimSegments[AnimSegmentIndex].GetAnimReference();
 		if (OldSequenceBase != NewSequenceBase)
 		{
 			OnPreAnimUpdateDelegate.ExecuteIfBound();
@@ -395,8 +391,7 @@ void SAnimSegmentsPanel::OpenAsset(int32 AnimSegmentIndex)
 {
 	if (ValidIndex(AnimSegmentIndex))
 	{
-		UAnimSequenceBase* Asset = AnimTrack->AnimSegments[AnimSegmentIndex].AnimReference;
-		if (Asset)
+		if (UAnimSequenceBase* Asset = AnimTrack->AnimSegments[AnimSegmentIndex].GetAnimReference())
 		{
 			GEditor->GetEditorSubsystem<UAssetEditorSubsystem>()->OpenEditorForAsset(Asset);
 		}
@@ -407,7 +402,7 @@ void SAnimSegmentsPanel::FillSubMenu(FMenuBuilder& MenuBuilder, int32 AnimSegmen
 {
 	if (ValidIndex(AnimSegmentIndex))
 	{
-		UAnimSequenceBase* OldSequenceBase = AnimTrack->AnimSegments[AnimSegmentIndex].AnimReference;
+		UAnimSequenceBase* OldSequenceBase = AnimTrack->AnimSegments[AnimSegmentIndex].GetAnimReference();
 
 		if (ensureAlways(OldSequenceBase))
 		{
diff --git a/Engine/Source/Editor/Persona/Private/SAnimTimingPanel.cpp b/Engine/Source/Editor/Persona/Private/SAnimTimingPanel.cpp
--- a/Engine/Source/Editor/Persona/Private/SAnimTimingPanel.cpp
+++ b/Engine/Source/Editor/Persona/Private/SAnimTimingPanel.cpp
@@ -184,7 +184,7 @@ FVector2D SAnimTimingNode::ComputeDesiredSize(float) const
 {
 	// Desired height is always the same (a little less than the track height) but the width depends on the text we display
 	const TSharedRef< FSlateFontMeasure > FontMeasureService = FSlateApplication::Get().GetRenderer()->GetFontMeasureService();
-	static FSlateFontInfo LabelFont = FCoreStyle::GetDefaultFontStyle("Regular", 10);
+	static FSlateFontInfo LabelFont = FCoreStyle::GetDefaultFontStyle("Regular", AnimTimingConstants::FontSize);
 	float TextWidth = FontMeasureService->Measure(FString::FromInt(Element->TriggerIdx), LabelFont).X;
 	return FVector2D(FMath::Max(AnimTimingConstants::DefaultNodeSize, TextWidth), AnimTimingConstants::DefaultNodeSize);
 }
diff --git a/Engine/Source/Editor/UnrealEd/Classes/Animation/EditorAnimBaseObj.h b/Engine/Source/Editor/UnrealEd/Classes/Animation/EditorAnimBaseObj.h
--- a/Engine/Source/Editor/UnrealEd/Classes/Animation/EditorAnimBaseObj.h
+++ b/Engine/Source/Editor/UnrealEd/Classes/Animation/EditorAnimBaseObj.h
@@ -37,6 +37,4 @@ public:
 	virtual void PostEditChangeProperty(FPropertyChangedEvent& PropertyChangedEvent) override;
 
 	virtual bool PropertyChangeRequiresRebuild(FPropertyChangedEvent& PropertyChangedEvent) { return true;}
-
-	//void NotifyUser();
 };
diff --git a/Engine/Source/Editor/UnrealEd/Classes/Animation/EditorAnimCompositeSegment.h b/Engine/Source/Editor/UnrealEd/Classes/Animation/EditorAnimCompositeSegment.h
--- a/Engine/Source/Editor/UnrealEd/Classes/Animation/EditorAnimCompositeSegment.h
+++ b/Engine/Source/Editor/UnrealEd/Classes/Animation/EditorAnimCompositeSegment.h
@@ -29,11 +29,9 @@ public:
 
 	int AnimSegmentIndex;
 
-	virtual void InitAnimSegment(int AnimSegmentIndex);
+	virtual void InitAnimSegment(int32 AnimSegmentIndex);
 	virtual bool ApplyChangesToMontage() override;
 
+	virtual void PostEditChangeProperty(FPropertyChangedEvent& PropertyChangedEvent) override;
 	virtual bool PropertyChangeRequiresRebuild(FPropertyChangedEvent& PropertyChangedEvent) override;
-
-private:
-	//FAnimTrack* GetAnimTrack
 };
diff --git a/Engine/Source/Editor/UnrealEd/Classes/Animation/EditorAnimSegment.h b/Engine/Source/Editor/UnrealEd/Classes/Animation/EditorAnimSegment.h
--- a/Engine/Source/Editor/UnrealEd/Classes/Animation/EditorAnimSegment.h
+++ b/Engine/Source/Editor/UnrealEd/Classes/Animation/EditorAnimSegment.h
@@ -30,8 +30,9 @@ public:
 	int AnimSlotIndex;
 	int AnimSegmentIndex;
 
-	virtual void InitAnimSegment(int AnimSlotIndex, int AnimSegmentIndex);
+	virtual void InitAnimSegment(int32 AnimSlotIndex, int32 AnimSegmentIndex);
 	virtual bool ApplyChangesToMontage() override;
 
+	virtual void PostEditChangeProperty(FPropertyChangedEvent& PropertyChangedEvent) override;
 	virtual bool PropertyChangeRequiresRebuild(FPropertyChangedEvent& PropertyChangedEvent) override;
 };
diff --git a/Engine/Source/Editor/UnrealEd/Classes/Animation/EditorCompositeSection.h b/Engine/Source/Editor/UnrealEd/Classes/Animation/EditorCompositeSection.h
--- a/Engine/Source/Editor/UnrealEd/Classes/Animation/EditorCompositeSection.h
+++ b/Engine/Source/Editor/UnrealEd/Classes/Animation/EditorCompositeSection.h
@@ -26,6 +26,6 @@ public:
 
 	int32 SectionIndex;
 	
-	virtual void InitSection(int SectionIndex);
+	virtual void InitSection(int32 SectionIndex);
 	virtual bool ApplyChangesToMontage() override;
 };
diff --git a/Engine/Source/Editor/UnrealEd/Private/Animation/EditorAnimBaseObj.cpp b/Engine/Source/Editor/UnrealEd/Private/Animation/EditorAnimBaseObj.cpp
--- a/Engine/Source/Editor/UnrealEd/Private/Animation/EditorAnimBaseObj.cpp
+++ b/Engine/Source/Editor/UnrealEd/Private/Animation/EditorAnimBaseObj.cpp
@@ -19,6 +19,7 @@ void UEditorAnimBaseObj::InitFromAnim(UAnimSequenceBase* AnimObjectIn, FOnAnimOb
 	AnimObject = AnimObjectIn;
 	OnChange = OnChangeIn;
 }
+
 bool UEditorAnimBaseObj::ApplyChangesToMontage()
 {
 	return false;
@@ -36,7 +37,7 @@ void UEditorAnimBaseObj::PostEditChangeProperty(FPropertyChangedEvent& PropertyC
 {
 	Super::PostEditChangeProperty(PropertyChangedEvent);
 
-	if(OnChange.IsBound())
+	if(OnChange.IsBound() && PropertyChangedEvent.ChangeType != EPropertyChangeType::Interactive)
 	{
 		if (ApplyChangesToMontage())
 		{
@@ -45,15 +46,4 @@ void UEditorAnimBaseObj::PostEditChangeProperty(FPropertyChangedEvent& PropertyC
 	}
 }
 
-/*
-void UEditorAnimBaseObj::NotifyUser()
-{
-	FNotificationInfo Info( FString::Printf( *LOCTEXT( "InvalidAnimSeg", "Invalid Animation is selected" ).ToString() ) );
-
-	Info.bUseLargeFont = false;
-	Info.ExpireDuration = 5.0f;
-
-	TSharedPtr<SNotificationItem> Notification = FSlateNotificationManager::Get().AddNotification( Info );
-}*/
-
 #undef LOCTEXT_NAMESPACE
diff --git a/Engine/Source/Editor/UnrealEd/Private/Animation/EditorAnimCompositeSegment.cpp b/Engine/Source/Editor/UnrealEd/Private/Animation/EditorAnimCompositeSegment.cpp
--- a/Engine/Source/Editor/UnrealEd/Private/Animation/EditorAnimCompositeSegment.cpp
+++ b/Engine/Source/Editor/UnrealEd/Private/Animation/EditorAnimCompositeSegment.cpp
@@ -15,7 +15,7 @@ UEditorAnimCompositeSegment::UEditorAnimCompositeSegment(const FObjectInitialize
 	AnimSegmentIndex = 0;
 }
 
-void UEditorAnimCompositeSegment::InitAnimSegment(int AnimSegmentIndexIn)
+void UEditorAnimCompositeSegment::InitAnimSegment(int32 AnimSegmentIndexIn)
 {
 	AnimSegmentIndex = AnimSegmentIndexIn;
 	if(UAnimComposite* Composite = Cast<UAnimComposite>(AnimObject))
@@ -27,20 +27,32 @@ void UEditorAnimCompositeSegment::InitAnimSegment(int AnimSegmentIndexIn)
 	}
 }
 
+void UEditorAnimCompositeSegment::PostEditChangeProperty(FPropertyChangedEvent& PropertyChangedEvent)
+{
+	PRAGMA_DISABLE_DEPRECATION_WARNINGS
+	if (PropertyChangedEvent.GetPropertyName() == GET_MEMBER_NAME_CHECKED(FAnimSegment, AnimReference))
+	{
+		AnimSegment.UpdateCachedPlayLength();			
+	}
+	PRAGMA_ENABLE_DEPRECATION_WARNINGS
+	
+	Super::PostEditChangeProperty(PropertyChangedEvent);
+}
+
 bool UEditorAnimCompositeSegment::ApplyChangesToMontage()
 {
 	if(UAnimComposite* Composite = Cast<UAnimComposite>(AnimObject))
 	{
 		if(Composite->AnimationTrack.AnimSegments.IsValidIndex(AnimSegmentIndex))
 		{
-			if (AnimSegment.AnimReference && Composite->GetSkeleton()->IsCompatible(AnimSegment.AnimReference->GetSkeleton()))
+			if (AnimSegment.GetAnimReference() && Composite->GetSkeleton()->IsCompatible(AnimSegment.GetAnimReference()->GetSkeleton()))
 			{
 				Composite->AnimationTrack.AnimSegments[AnimSegmentIndex] = AnimSegment;
 				return true;
 			}
 			else
 			{
-				AnimSegment.AnimReference = Composite->AnimationTrack.AnimSegments[AnimSegmentIndex].AnimReference;
+				AnimSegment.SetAnimReference(Composite->AnimationTrack.AnimSegments[AnimSegmentIndex].GetAnimReference(), false);
 				return false;
 			}
 		}
@@ -52,13 +64,9 @@ bool UEditorAnimCompositeSegment::ApplyChangesToMontage()
 bool UEditorAnimCompositeSegment::PropertyChangeRequiresRebuild(FPropertyChangedEvent& PropertyChangedEvent)
 {
 	const FName PropertyName = PropertyChangedEvent.Property ? PropertyChangedEvent.Property->GetFName() : NAME_None;
-	if( PropertyName == FName(TEXT("AnimEndTime")) || PropertyName == FName(TEXT("AnimStartTime")) || PropertyName == FName(TEXT("AnimPlayRate")) || PropertyName == FName(TEXT("LoopingCount")))
-	{
-		// Changing the end or start time of the segment length can't change the order of the montage segments.
-		// Return false so that the montage editor does not fully rebuild its UI and we can keep this UEditorAnimSegment object 
-		// in the details view. (A better solution would be handling the rebuilding in a way that didn't possibly invalidate the UEditorMontageObj in the details view)
-		return false;
-	}
 
-	return true;
+	// Changing the end or start time of the segment length can't change the order of the montage segments.
+	// Return false so that the montage editor does not fully rebuild its UI and we can keep this UEditorAnimSegment object 
+	// in the details view. (A better solution would be handling the rebuilding in a way that didn't possibly invalidate the UEditorMontageObj in the details view)
+	return !(PropertyName == GET_MEMBER_NAME_CHECKED(FAnimSegment, AnimEndTime) || PropertyName == GET_MEMBER_NAME_CHECKED(FAnimSegment, AnimStartTime) || PropertyName == GET_MEMBER_NAME_CHECKED(FAnimSegment, AnimPlayRate) || PropertyName == GET_MEMBER_NAME_CHECKED(FAnimSegment, LoopingCount));
 }
diff --git a/Engine/Source/Editor/UnrealEd/Private/Animation/EditorAnimSegment.cpp b/Engine/Source/Editor/UnrealEd/Private/Animation/EditorAnimSegment.cpp
--- a/Engine/Source/Editor/UnrealEd/Private/Animation/EditorAnimSegment.cpp
+++ b/Engine/Source/Editor/UnrealEd/Private/Animation/EditorAnimSegment.cpp
@@ -17,7 +17,7 @@ UEditorAnimSegment::UEditorAnimSegment(const FObjectInitializer& ObjectInitializ
 	AnimSegmentIndex = 0;
 }
 
-void UEditorAnimSegment::InitAnimSegment(int AnimSlotIndexIn,int AnimSegmentIndexIn)
+void UEditorAnimSegment::InitAnimSegment(int32 AnimSlotIndexIn, int32 AnimSegmentIndexIn)
 {
 	AnimSlotIndex = AnimSlotIndexIn;
 	AnimSegmentIndex = AnimSegmentIndexIn;
@@ -30,13 +30,25 @@ void UEditorAnimSegment::InitAnimSegment(int AnimSlotIndexIn,int AnimSegmentInde
 	}
 }
 
+void UEditorAnimSegment::PostEditChangeProperty(FPropertyChangedEvent& PropertyChangedEvent)
+{
+	PRAGMA_DISABLE_DEPRECATION_WARNINGS
+	if (PropertyChangedEvent.GetPropertyName() == GET_MEMBER_NAME_CHECKED(FAnimSegment, AnimReference))
+	{
+		AnimSegment.UpdateCachedPlayLength();			
+	}
+	PRAGMA_ENABLE_DEPRECATION_WARNINGS
+	
+	Super::PostEditChangeProperty(PropertyChangedEvent);
+}
+
 bool UEditorAnimSegment::ApplyChangesToMontage()
 {
 	if(UAnimMontage* Montage = Cast<UAnimMontage>(AnimObject))
 	{
 		if(Montage->SlotAnimTracks.IsValidIndex(AnimSlotIndex) && Montage->SlotAnimTracks[AnimSlotIndex].AnimTrack.AnimSegments.IsValidIndex(AnimSegmentIndex) )
 		{
-			if (AnimSegment.AnimReference && Montage->GetSkeleton()->IsCompatible(AnimSegment.AnimReference->GetSkeleton()))
+			if (AnimSegment.GetAnimReference() && Montage->GetSkeleton()->IsCompatible(AnimSegment.GetAnimReference()->GetSkeleton()))
 			{
 				Montage->SlotAnimTracks[AnimSlotIndex].AnimTrack.AnimSegments[AnimSegmentIndex] = AnimSegment;
 				Montage->UpdateLinkableElements(AnimSlotIndex, AnimSegmentIndex);
@@ -52,7 +64,7 @@ bool UEditorAnimSegment::ApplyChangesToMontage()
 			}
 			else
 			{
-				AnimSegment.AnimReference = Montage->SlotAnimTracks[AnimSlotIndex].AnimTrack.AnimSegments[AnimSegmentIndex].AnimReference;
+				AnimSegment.SetAnimReference(Montage->SlotAnimTracks[AnimSlotIndex].AnimTrack.AnimSegments[AnimSegmentIndex].GetAnimReference());
 				return false;
 			}
 		}
@@ -64,13 +76,9 @@ bool UEditorAnimSegment::ApplyChangesToMontage()
 bool UEditorAnimSegment::PropertyChangeRequiresRebuild(FPropertyChangedEvent& PropertyChangedEvent)
 {
 	const FName PropertyName = PropertyChangedEvent.Property ? PropertyChangedEvent.Property->GetFName() : NAME_None;
-	if( PropertyName == FName(TEXT("AnimEndTime")) || PropertyName == FName(TEXT("AnimStartTime")) || PropertyName == FName(TEXT("AnimPlayRate")) || PropertyName == FName(TEXT("LoopingCount")))
-	{
-		// Changing the end or start time of the segment length can't change the order of the montage segments.
-		// Return false so that the montage editor does not fully rebuild its UI and we can keep this UEditorAnimSegment object 
-		// in the details view. (A better solution would be handling the rebuilding in a way that didn't possibly invalidate the UEditorMontageObj in the details view)
-		return false;
-	}
 
-	return true;
+	// Changing the end or start time of the segment length can't change the order of the montage segments.
+	// Return false so that the montage editor does not fully rebuild its UI and we can keep this UEditorAnimSegment object 
+	// in the details view. (A better solution would be handling the rebuilding in a way that didn't possibly invalidate the UEditorMontageObj in the details view)
+	return !(PropertyName == GET_MEMBER_NAME_CHECKED(FAnimSegment, AnimEndTime) || PropertyName == GET_MEMBER_NAME_CHECKED(FAnimSegment, AnimStartTime) || PropertyName == GET_MEMBER_NAME_CHECKED(FAnimSegment, AnimPlayRate) || PropertyName == GET_MEMBER_NAME_CHECKED(FAnimSegment, LoopingCount));
 }
diff --git a/Engine/Source/Editor/UnrealEd/Private/Animation/EditorCompositeSection.cpp b/Engine/Source/Editor/UnrealEd/Private/Animation/EditorCompositeSection.cpp
--- a/Engine/Source/Editor/UnrealEd/Private/Animation/EditorCompositeSection.cpp
+++ b/Engine/Source/Editor/UnrealEd/Private/Animation/EditorCompositeSection.cpp
@@ -35,7 +35,7 @@ void CopyMetaData(FCompositeSection& Source, FCompositeSection& Dest, UObject* D
 	}
 }
 
-void UEditorCompositeSection::InitSection(int SectionIndexIn)
+void UEditorCompositeSection::InitSection(int32 SectionIndexIn)
 {
 	SectionIndex = SectionIndexIn;
 	if(UAnimMontage* Montage = Cast<UAnimMontage>(AnimObject))
@@ -47,6 +47,7 @@ void UEditorCompositeSection::InitSection(int SectionIndexIn)
 		}
 	}
 }
+
 bool UEditorCompositeSection::ApplyChangesToMontage()
 {
 	if(UAnimMontage* Montage = Cast<UAnimMontage>(AnimObject))
diff --git a/Engine/Source/Editor/UnrealEd/Private/Factories/AnimCompositeFactory.cpp b/Engine/Source/Editor/UnrealEd/Private/Factories/AnimCompositeFactory.cpp
--- a/Engine/Source/Editor/UnrealEd/Private/Factories/AnimCompositeFactory.cpp
+++ b/Engine/Source/Editor/UnrealEd/Private/Factories/AnimCompositeFactory.cpp
@@ -81,12 +81,7 @@ UObject* UAnimCompositeFactory::FactoryCreateNew(UClass* Class, UObject* InParen
 			TargetSkeleton = SourceSkeleton;
 
 			FAnimSegment NewSegment;
-			NewSegment.AnimReference = SourceAnimation;
-			NewSegment.AnimStartTime = 0.f;
-			NewSegment.AnimEndTime = SourceAnimation->GetPlayLength();
-			NewSegment.AnimPlayRate = 1.f;
-			NewSegment.LoopingCount = 1;
-			NewSegment.StartPos = 0.f;
+			NewSegment.SetAnimReference(SourceAnimation, true);
 
 			AnimComposite->AnimationTrack.AnimSegments.Add(NewSegment);
 			AnimComposite->SetCompositeLength(AnimComposite->AnimationTrack.GetLength());
diff --git a/Engine/Source/Editor/UnrealEd/Private/Factories/AnimMontageFactory.cpp b/Engine/Source/Editor/UnrealEd/Private/Factories/AnimMontageFactory.cpp
--- a/Engine/Source/Editor/UnrealEd/Private/Factories/AnimMontageFactory.cpp
+++ b/Engine/Source/Editor/UnrealEd/Private/Factories/AnimMontageFactory.cpp
@@ -81,12 +81,7 @@ UObject* UAnimMontageFactory::FactoryCreateNew(UClass* Class, UObject* InParent,
 			TargetSkeleton = SourceSkeleton;
 
 			FAnimSegment NewSegment;
-			NewSegment.AnimReference = SourceAnimation;
-			NewSegment.AnimStartTime = 0.f;
-			NewSegment.AnimEndTime = SourceAnimation->GetPlayLength();
-			NewSegment.AnimPlayRate = 1.f;
-			NewSegment.LoopingCount = 1;
-			NewSegment.StartPos = 0.f;
+			NewSegment.SetAnimReference(SourceAnimation, true);
 
 			FSlotAnimationTrack& NewTrack = AnimMontage->SlotAnimTracks[0];
 			NewTrack.AnimTrack.AnimSegments.Add(NewSegment);
diff --git a/Engine/Source/Runtime/Engine/Classes/Animation/AnimComposite.h b/Engine/Source/Runtime/Engine/Classes/Animation/AnimComposite.h
--- a/Engine/Source/Runtime/Engine/Classes/Animation/AnimComposite.h
+++ b/Engine/Source/Runtime/Engine/Classes/Animation/AnimComposite.h
@@ -55,6 +55,7 @@ public:
 	virtual void InvalidateRecursiveAsset() override;
 	virtual bool ContainRecursive(TArray<UAnimCompositeBase*>& CurrentAccumulatedList) override;
 	virtual void SetCompositeLength(float InLength) override;
+	virtual void PostLoad() override;
 	//~End UAnimCompositeBase Interface
 };
 
diff --git a/Engine/Source/Runtime/Engine/Classes/Animation/AnimCompositeBase.h b/Engine/Source/Runtime/Engine/Classes/Animation/AnimCompositeBase.h
--- a/Engine/Source/Runtime/Engine/Classes/Animation/AnimCompositeBase.h
+++ b/Engine/Source/Runtime/Engine/Classes/Animation/AnimCompositeBase.h
@@ -59,31 +59,65 @@ struct FAnimSegment
 {
 	GENERATED_USTRUCT_BODY()
 
+	PRAGMA_DISABLE_DEPRECATION_WARNINGS
+	FAnimSegment(const FAnimSegment&) = default;
+	FAnimSegment(FAnimSegment&&) = default;
+	FAnimSegment& operator=(const FAnimSegment&) = default;
+	FAnimSegment& operator=(FAnimSegment&&) = default;
+	PRAGMA_ENABLE_DEPRECATION_WARNINGS
+
+	UE_DEPRECATED(5.1, "Public access to AnimReference has been deprecated, use Set/Get-AnimReference instead")
 	/** Anim Reference to play - only allow AnimSequence or AnimComposite **/
-	UPROPERTY(EditAnywhere, Category=AnimSegment)
+	UPROPERTY(EditAnywhere, Category=AnimSegment, meta=(DisplayName = "Animation Reference"))
 	TObjectPtr<UAnimSequenceBase> AnimReference;
 
+#if WITH_EDITORONLY_DATA
+	UPROPERTY()
+	float CachedPlayLength = 0.f;
+#endif
+
+#if WITH_EDITOR
+	friend class UEditorAnimSegment;
+	friend class UEditorAnimCompositeSegment;
+	ENGINE_API void UpdateCachedPlayLength();
+#endif // WITH_EDITOR
+public:
+
+	ENGINE_API void SetAnimReference(UAnimSequenceBase* InAnimReference, bool bInitialize = false);
+	const TObjectPtr<UAnimSequenceBase>& GetAnimReference() const 
+	{
+		PRAGMA_DISABLE_DEPRECATION_WARNINGS
+		return AnimReference;
+		PRAGMA_ENABLE_DEPRECATION_WARNINGS
+	}
+
+#if WITH_EDITOR
+	ENGINE_API bool IsPlayLengthOutOfDate() const;
+#endif // WITH_EDITOR
+	
 	/** Start Pos within this AnimCompositeBase */
-	UPROPERTY(VisibleAnywhere, Category=AnimSegment)
-	float	StartPos;
+	UPROPERTY(VisibleAnywhere, Category=AnimSegment, meta=(DisplayName = "Starting Position"))
+	float StartPos;
 
 	/** Time to start playing AnimSequence at. */
-	UPROPERTY(EditAnywhere, Category=AnimSegment)
-	float	AnimStartTime;
+	UPROPERTY(EditAnywhere, Category=AnimSegment, meta=(DisplayName = "Start Time"))
+	float AnimStartTime;
 
 	/** Time to end playing the AnimSequence at. */
-	UPROPERTY(EditAnywhere, Category=AnimSegment)
-	float	AnimEndTime;
+	UPROPERTY(EditAnywhere, Category=AnimSegment, meta=(DisplayName = "End Time"))
+	float AnimEndTime;
 
 	/** Playback speed of this animation. If you'd like to reverse, set -1*/
-	UPROPERTY(EditAnywhere, Category=AnimSegment)
-	float	AnimPlayRate;
+	UPROPERTY(EditAnywhere, Category=AnimSegment, meta=(DisplayName = "Play Rate"))
+	float AnimPlayRate;
 
-	UPROPERTY(EditAnywhere, Category=AnimSegment)
-	int32		LoopingCount;
+	UPROPERTY(EditAnywhere, Category=AnimSegment, meta=(DisplayName = "Loop Count"))
+	int32 LoopingCount;
 
 	FAnimSegment()
+		PRAGMA_DISABLE_DEPRECATION_WARNINGS
 		: AnimReference(nullptr)
+		PRAGMA_ENABLE_DEPRECATION_WARNINGS
 		, StartPos(0.f)
 		, AnimStartTime(0.f)
 		, AnimEndTime(0.f)
@@ -96,7 +130,9 @@ struct FAnimSegment
 	/** Ensures PlayRate is non Zero */
 	float GetValidPlayRate() const
 	{
+		PRAGMA_DISABLE_DEPRECATION_WARNINGS
 		float SeqPlayRate = AnimReference ? AnimReference->RateScale : 1.0f;
+		PRAGMA_ENABLE_DEPRECATION_WARNINGS
 		float FinalPlayRate = SeqPlayRate * AnimPlayRate;
 		return (FMath::IsNearlyZero(FinalPlayRate) ? 1.f : FinalPlayRate);
 	}
@@ -185,7 +221,9 @@ struct FAnimSegment
 	/** 
 	 * return true if anim notify is available 
 	 */
+	PRAGMA_DISABLE_DEPRECATION_WARNINGS
 	bool IsNotifyAvailable() const { return IsValid() && AnimReference && AnimReference->IsNotifyAvailable(); }
+	PRAGMA_ENABLE_DEPRECATION_WARNINGS
 private:
 
 	/**
diff --git a/Engine/Source/Runtime/Engine/Private/Animation/AnimComposite.cpp b/Engine/Source/Runtime/Engine/Private/Animation/AnimComposite.cpp
--- a/Engine/Source/Runtime/Engine/Private/Animation/AnimComposite.cpp
+++ b/Engine/Source/Runtime/Engine/Private/Animation/AnimComposite.cpp
@@ -151,3 +151,18 @@ void UAnimComposite::SetCompositeLength(float InLength)
 	PRAGMA_ENABLE_DEPRECATION_WARNINGS
 #endif	
 }
+
+void UAnimComposite::PostLoad()
+{
+	Super::PostLoad();
+
+#if WITH_EDITOR
+	for (const FAnimSegment& AnimSegment : AnimationTrack.AnimSegments)
+	{
+		if(AnimSegment.IsPlayLengthOutOfDate())
+		{
+			UE_LOG(LogAnimation, Warning, TEXT("AnimComposite (%s) contains a Segment for which the playable length %f is out-of-sync with the represented AnimationSequence its length %f (%s). Please up-date the segment and resave."), *GetFullName(), (AnimSegment.AnimEndTime - AnimSegment.AnimStartTime), AnimSegment.GetAnimReference()->GetPlayLength(), *AnimSegment.GetAnimReference()->GetFullName());
+		}
+	}
+#endif
+}
diff --git a/Engine/Source/Runtime/Engine/Private/Animation/AnimCompositeBase.cpp b/Engine/Source/Runtime/Engine/Private/Animation/AnimCompositeBase.cpp
--- a/Engine/Source/Runtime/Engine/Private/Animation/AnimCompositeBase.cpp
+++ b/Engine/Source/Runtime/Engine/Private/Animation/AnimCompositeBase.cpp
@@ -20,7 +20,9 @@ UAnimSequenceBase * FAnimSegment::GetAnimationData(float PositionInTrack, float&
 {
 	if( bValid && IsInRange(PositionInTrack) )
 	{
-		if( AnimReference )
+		PRAGMA_DISABLE_DEPRECATION_WARNINGS
+		if(AnimReference)
+		PRAGMA_ENABLE_DEPRECATION_WARNINGS
 		{
 			const float ValidPlayRate = GetValidPlayRate();
 
@@ -43,11 +45,14 @@ UAnimSequenceBase * FAnimSegment::GetAnimationData(float PositionInTrack, float&
 			{
 				PositionInAnim = AnimEndTime + Delta * ValidPlayRate;
 			}
+
+			PRAGMA_DISABLE_DEPRECATION_WARNINGS
 			return AnimReference;
+			PRAGMA_ENABLE_DEPRECATION_WARNINGS
 		}
 	}
 
-	return NULL;
+	return nullptr;
 }
 
 /** Converts 'Track Position' to position on AnimSequence.
@@ -111,7 +116,9 @@ void FAnimSegment::GetAnimNotifiesFromTrackPositions(const float& PreviousTrackP
 		)
 	{
 		// Only allow AnimSequences for now. Other types will need additional support.
+		PRAGMA_DISABLE_DEPRECATION_WARNINGS
 		UAnimSequenceBase* AnimSequenceBase = AnimReference;
+		PRAGMA_ENABLE_DEPRECATION_WARNINGS
 		if(AnimSequenceBase)
 		{
 			const float ValidPlayRate = GetValidPlayRate();
@@ -171,7 +178,9 @@ void FAnimSegment::GetRootMotionExtractionStepsForTrackRange(TArray<FRootMotionE
 		return;
 	}
 
+	PRAGMA_DISABLE_DEPRECATION_WARNINGS
 	if (!bValid || !AnimReference)
+	PRAGMA_ENABLE_DEPRECATION_WARNINGS
 	{
 		return;
 	}
@@ -208,8 +217,10 @@ void FAnimSegment::GetRootMotionExtractionStepsForTrackRange(TArray<FRootMotionE
 		const float AnimEndPoint = bAnimPlayingBackwards ? AnimStartTime : AnimEndTime;
 
 		// Only allow AnimSequences for now. Other types will need additional support.
+		PRAGMA_DISABLE_DEPRECATION_WARNINGS
 		UAnimSequence* AnimSequence = Cast<UAnimSequence>(AnimReference);
 		UAnimComposite* AnimComposite = Cast<UAnimComposite>(AnimReference);
+		PRAGMA_ENABLE_DEPRECATION_WARNINGS
 
 		if (AnimSequence || AnimComposite)
 		{
@@ -263,7 +274,8 @@ bool FAnimTrack::HasRootMotion() const
 {
 	for (const FAnimSegment& AnimSegment : AnimSegments)
 	{
-		if (AnimSegment.bValid && AnimSegment.AnimReference && AnimSegment.AnimReference->HasRootMotion())
+		const UAnimSequenceBase* AnimReference = AnimSegment.GetAnimReference();
+		if (AnimSegment.bValid && AnimReference && AnimReference->HasRootMotion())
 		{
 			return true;
 		}
@@ -278,7 +290,8 @@ class UAnimSequence* FAnimTrack::GetAdditiveBasePose() const
 	{
 		for (const FAnimSegment& AnimSegment : AnimSegments)
 		{
-			UAnimSequence* BasePose = (AnimSegment.AnimReference) ? (AnimSegment.AnimReference->GetAdditiveBasePose()) : nullptr;
+			UAnimSequenceBase* AnimReference = AnimSegment.GetAnimReference();
+			UAnimSequence* BasePose = AnimReference ? AnimReference->GetAdditiveBasePose() : nullptr;
 			if (BasePose)
 			{
 				return BasePose;
@@ -322,10 +335,9 @@ float FAnimTrack::GetLength() const
 
 	// in the future, if we're more clear about exactly what requirement is for segments, 
 	// this can be optimized. For now this is slow. 
-	for ( int32 I=0; I<AnimSegments.Num(); ++I )
+	for (const FAnimSegment& AnimSegment : AnimSegments)
 	{
-		const struct FAnimSegment& Segment = AnimSegments[I];
-		float EndFrame = Segment.StartPos + Segment.GetLength();
+		const float EndFrame = AnimSegment.StartPos + AnimSegment.GetLength();
 		if ( EndFrame > TotalLength )
 		{
 			TotalLength = EndFrame;
@@ -343,10 +355,10 @@ bool FAnimTrack::IsAdditive() const
 	// and if they mismatch, what can I do? That should be another verification function when this is created
 	// it will look visually wrong if something mismatches, but nothing really is better solution than that. 
 	// in editor, when this is created, the test has to be done to verify all are matches. 
-	for ( int32 I=0; I<AnimSegments.Num(); ++I )
+	for (const FAnimSegment& AnimSegment : AnimSegments)
 	{
-		const struct FAnimSegment & Segment = AnimSegments[I];
-		return ( Segment.AnimReference && Segment.bValid && Segment.AnimReference->IsValidAdditive() ); //-V612
+		const UAnimSequenceBase* AnimReference = AnimSegment.GetAnimReference();
+		return (AnimReference && AnimSegment.bValid && AnimReference->IsValidAdditive() ); //-V612
 	}
 
 	return false;
@@ -360,20 +372,12 @@ bool FAnimTrack::IsRotationOffsetAdditive() const
 	// and if they mismatch, what can I do? That should be another verification function when this is created
 	// it will look visually wrong if something mismatches, but nothing really is better solution than that. 
 	// in editor, when this is created, the test has to be done to verify all are matches. 
-	for ( int32 I=0; I<AnimSegments.Num(); ++I )
+	for (const FAnimSegment& AnimSegment : AnimSegments)
 	{
-		const struct FAnimSegment & Segment = AnimSegments[I];
-		if ( Segment.AnimReference && Segment.AnimReference->IsValidAdditive() )
+		const UAnimSequenceBase* AnimReference = AnimSegment.GetAnimReference();
+		if (AnimReference && AnimReference->IsValidAdditive())
 		{
-			UAnimSequenceBase* SequenceBase = Segment.AnimReference;
-			if (SequenceBase)
-			{
-				return (SequenceBase->GetAdditiveAnimType() == AAT_RotationOffsetMeshSpace);
-			}
-			else
-			{
-				break;
-			}
+			return (AnimReference->GetAdditiveAnimType() == AAT_RotationOffsetMeshSpace);
 		}
 		else
 		{
@@ -394,11 +398,11 @@ int32 FAnimTrack::GetTrackAdditiveType() const
 
 	if( AnimSegments.Num() > 0 )
 	{
-		const struct FAnimSegment & Segment = AnimSegments[0];
-		UAnimSequenceBase* SequenceBase = Segment.AnimReference;
-		if ( SequenceBase )
+		const FAnimSegment& AnimSegment = AnimSegments[0];
+		const UAnimSequenceBase* AnimReference = AnimSegment.GetAnimReference();
+		if (AnimReference)
 		{
-			return SequenceBase->GetAdditiveAnimType();
+			return AnimReference->GetAdditiveAnimType();
 		}
 	}
 	return -1;
@@ -410,17 +414,18 @@ void FAnimTrack::ValidateSegmentTimes()
 	if(AnimSegments.Num() > 0)
 	{
 		AnimSegments[0].StartPos = 0.0f;
-		for(int32 J = 0; J < AnimSegments.Num(); J++)
+		for(int32 SegmentIndex = 0; SegmentIndex < AnimSegments.Num(); SegmentIndex++)
 		{
-			FAnimSegment& Segment = AnimSegments[J];
-			if(J > 0)
+			FAnimSegment& AnimSegment = AnimSegments[SegmentIndex];
+			if(SegmentIndex > 0)
 			{
-				Segment.StartPos = AnimSegments[J - 1].StartPos + AnimSegments[J - 1].GetLength();
+				AnimSegment.StartPos = AnimSegments[SegmentIndex - 1].StartPos + AnimSegments[SegmentIndex - 1].GetLength();
 			}
 
-			if(Segment.AnimReference && Segment.AnimEndTime > Segment.AnimReference->GetPlayLength())
+			const UAnimSequenceBase* AnimReference = AnimSegment.GetAnimReference();
+			if(AnimReference && AnimSegment.AnimEndTime > AnimReference->GetPlayLength())
 			{
-				Segment.AnimEndTime = Segment.AnimReference->GetPlayLength();
+				AnimSegment.AnimEndTime = AnimReference->GetPlayLength();
 			}
 		}
 	}
@@ -459,11 +464,10 @@ int32 FAnimTrack::GetSegmentIndexAtTime(float InTime) const
 #if WITH_EDITOR
 bool FAnimTrack::GetAllAnimationSequencesReferred(TArray<UAnimationAsset*>& AnimationAssets, bool bRecursive/* = true*/) const
 {
-	for ( int32 I=0; I<AnimSegments.Num(); ++I )
+	for (const FAnimSegment& AnimSegment : AnimSegments)
 	{
-		const struct FAnimSegment& Segment = AnimSegments[I];
-		UAnimSequenceBase* AnimSeqBase = Segment.AnimReference;
-		if ( Segment.bValid && AnimSeqBase )
+		UAnimSequenceBase* AnimSeqBase = AnimSegment.GetAnimReference();
+		if ( AnimSegment.bValid && AnimSeqBase )
 		{
 			AnimSeqBase->HandleAnimReferenceCollection(AnimationAssets, bRecursive);
 		}
@@ -475,21 +479,16 @@ bool FAnimTrack::GetAllAnimationSequencesReferred(TArray<UAnimationAsset*>& Anim
 void FAnimTrack::ReplaceReferredAnimations(const TMap<UAnimationAsset*, UAnimationAsset*>& ReplacementMap)
 {
 	TArray<FAnimSegment> NewAnimSegments;
-	for ( int32 I=0; I<AnimSegments.Num(); ++I )
+	for (FAnimSegment& AnimSegment : AnimSegments)
 	{
-		struct FAnimSegment& Segment = AnimSegments[I];
-
-		if (Segment.IsValid())
+		if (AnimSegment.IsValid())
 		{
-			// now fix everythign else
-			UAnimSequenceBase* SequenceBase = Segment.AnimReference;
-			if (SequenceBase)
+			if (UAnimSequenceBase* SequenceBase = AnimSegment.GetAnimReference())
 			{
-				UAnimationAsset* const* ReplacementAsset = ReplacementMap.Find(SequenceBase);
-				if(ReplacementAsset)
+				if(UAnimationAsset* const* ReplacementAsset = ReplacementMap.Find(SequenceBase))
 				{
-					Segment.AnimReference = Cast<UAnimSequenceBase>(*ReplacementAsset);
-					NewAnimSegments.Add(Segment);
+					AnimSegment.SetAnimReference(Cast<UAnimSequenceBase>(*ReplacementAsset));
+					NewAnimSegments.Add(AnimSegment);
 				}
 
 				SequenceBase->ReplaceReferredAnimations(ReplacementMap);
@@ -591,12 +590,11 @@ void FAnimTrack::GetAnimationPose(FAnimationPoseData& OutAnimationPoseData, cons
 
 void FAnimTrack::EnableRootMotionSettingFromMontage(bool bInEnableRootMotion, const ERootMotionRootLock::Type InRootMotionRootLock)
 {
-	for (int32 I = 0; I < AnimSegments.Num(); ++I)
+	for (const FAnimSegment& AnimSegment : AnimSegments)
 	{
-		const FAnimSegment& AnimSegment = AnimSegments[I];
-		if (AnimSegment.AnimReference)
+		if (UAnimSequenceBase* AnimReference = AnimSegment.GetAnimReference())
 		{
-			AnimSegment.AnimReference->EnableRootMotionSettingFromMontage(bInEnableRootMotion, InRootMotionRootLock);
+			AnimReference->EnableRootMotionSettingFromMontage(bInEnableRootMotion, InRootMotionRootLock);
 		}
 	}
 }
@@ -605,10 +603,10 @@ void FAnimTrack::EnableRootMotionSettingFromMontage(bool bInEnableRootMotion, co
 // as a result of anim composite being a part of anim sequence base
 void FAnimTrack::InvalidateRecursiveAsset(class UAnimCompositeBase* CheckAsset)
 {
-	for (int32 I = 0; I < AnimSegments.Num(); ++I)
+	for (FAnimSegment& AnimSegment : AnimSegments)
 	{
-		FAnimSegment& AnimSegment = AnimSegments[I];
-		UAnimCompositeBase* CompositeBase = Cast<UAnimCompositeBase>(AnimSegment.AnimReference);
+		UAnimSequenceBase* SequenceBase = AnimSegment.GetAnimReference();
+		UAnimCompositeBase* CompositeBase = Cast<UAnimCompositeBase>(SequenceBase);
 		if (CompositeBase)
 		{
 			// add owner
@@ -626,7 +624,7 @@ void FAnimTrack::InvalidateRecursiveAsset(class UAnimCompositeBase* CheckAsset)
 		}
 		else
 		{
-			AnimSegment.bValid = IsValidToAdd(AnimSegment.AnimReference);
+			AnimSegment.bValid = IsValidToAdd(SequenceBase);
 		}
 	}
 }
@@ -635,14 +633,12 @@ void FAnimTrack::InvalidateRecursiveAsset(class UAnimCompositeBase* CheckAsset)
 // and return true if it finds nested same assets
 bool FAnimTrack::ContainRecursive(const TArray<UAnimCompositeBase*>& CurrentAccumulatedList)
 {
-	for (int32 I = 0; I < AnimSegments.Num(); ++I)
+	for (const FAnimSegment& AnimSegment : AnimSegments)
 	{
-		FAnimSegment& AnimSegment = AnimSegments[I];
-
 		// we don't want to send this list broad widely (but in depth search)
 		// to do that, we copy the current accumulated list, and send that only, not the siblings
 		TArray<UAnimCompositeBase*> LocalCurrentAccumulatedList = CurrentAccumulatedList;
-		UAnimCompositeBase* CompositeBase = Cast<UAnimCompositeBase>(AnimSegment.AnimReference);
+		UAnimCompositeBase* CompositeBase = Cast<UAnimCompositeBase>(AnimSegment.GetAnimReference());
 		if (CompositeBase && CompositeBase->ContainRecursive(LocalCurrentAccumulatedList))
 		{
 			return true;
@@ -775,4 +771,57 @@ void UAnimCompositeBase::PostLoad()
 	Super::PostLoad();
 
 	InvalidateRecursiveAsset();
-}
\ No newline at end of file
+}
+
+void FAnimSegment::SetAnimReference(UAnimSequenceBase* InAnimReference, bool bInitialize /*= false*/)
+{
+	PRAGMA_DISABLE_DEPRECATION_WARNINGS
+	AnimReference = InAnimReference;
+	PRAGMA_ENABLE_DEPRECATION_WARNINGS
+
+#if WITH_EDITOR
+	UpdateCachedPlayLength();
+#endif // WITH_EDITOR
+
+	PRAGMA_DISABLE_DEPRECATION_WARNINGS
+	if (AnimReference && bInitialize)
+	PRAGMA_ENABLE_DEPRECATION_WARNINGS
+	{		
+		AnimStartTime = 0.f;
+		PRAGMA_DISABLE_DEPRECATION_WARNINGS
+		AnimEndTime = AnimReference->GetPlayLength();
+		PRAGMA_ENABLE_DEPRECATION_WARNINGS
+		AnimPlayRate = 1.f;
+		LoopingCount = 1;
+		StartPos = 0.f;
+	}
+}
+
+#if WITH_EDITOR
+bool FAnimSegment::IsPlayLengthOutOfDate() const
+{
+	PRAGMA_DISABLE_DEPRECATION_WARNINGS
+	if (AnimReference && !FMath::IsNearlyZero(CachedPlayLength))
+	{
+		// When the segment length is equal to _cached_ playlength and the current model playlength is different flag as out-of-date
+		// this can happen when the sequence is reimported without updating the montage and thus ending up with 'invalid' playback range.
+		const float PlayableLength = (AnimEndTime - AnimStartTime);
+		return FMath::IsNearlyEqual(PlayableLength, CachedPlayLength, KINDA_SMALL_NUMBER) && !FMath::IsNearlyEqual(AnimReference->GetPlayLength(), CachedPlayLength, KINDA_SMALL_NUMBER);
+	}
+	PRAGMA_ENABLE_DEPRECATION_WARNINGS
+	
+	return false;
+}
+
+void FAnimSegment::UpdateCachedPlayLength()
+{
+	CachedPlayLength = 0.f;
+	PRAGMA_DISABLE_DEPRECATION_WARNINGS
+	const UAnimDataModel* DataModel = AnimReference ? AnimReference->GetDataModel() : nullptr;	
+	PRAGMA_ENABLE_DEPRECATION_WARNINGS
+	if(DataModel)
+	{
+		CachedPlayLength = DataModel->GetPlayLength();
+	}
+}
+#endif // WITH_EDITOR
\ No newline at end of file
diff --git a/Engine/Source/Runtime/Engine/Private/Animation/AnimInstance.cpp b/Engine/Source/Runtime/Engine/Private/Animation/AnimInstance.cpp
--- a/Engine/Source/Runtime/Engine/Private/Animation/AnimInstance.cpp
+++ b/Engine/Source/Runtime/Engine/Private/Animation/AnimInstance.cpp
@@ -1979,7 +1979,7 @@ bool UAnimInstance::IsPlayingSlotAnimation(const UAnimSequenceBase* Asset, FName
 				if (AnimTrack && AnimTrack->AnimSegments.Num() == 1)
 				{
 					OutMontage = CurMontage;
-					return (AnimTrack->AnimSegments[0].AnimReference == Asset);
+					return (AnimTrack->AnimSegments[0].GetAnimReference() == Asset);
 				}
 			}
 		}
diff --git a/Engine/Source/Runtime/Engine/Private/Animation/AnimLinkableElement.cpp b/Engine/Source/Runtime/Engine/Private/Animation/AnimLinkableElement.cpp
--- a/Engine/Source/Runtime/Engine/Private/Animation/AnimLinkableElement.cpp
+++ b/Engine/Source/Runtime/Engine/Private/Animation/AnimLinkableElement.cpp
@@ -16,7 +16,7 @@ void FAnimLinkableElement::LinkMontage(UAnimMontage* Montage, float AbsMontageTi
 		if(SegmentIndex != INDEX_NONE)
 		{
 			FAnimSegment& Segment = Slot.AnimTrack.AnimSegments[SegmentIndex];
-			LinkedSequence = Segment.AnimReference;
+			LinkedSequence = Segment.GetAnimReference();
 			SegmentBeginTime = Segment.StartPos;
 			SegmentLength = Segment.GetLength();
 
@@ -78,7 +78,7 @@ void FAnimLinkableElement::Update()
 		{
 			// Update timing info from current segment
 			FAnimSegment& Segment = Slot.AnimTrack.AnimSegments[SegmentIndex];
-			LinkedSequence = Segment.AnimReference;
+			LinkedSequence = Segment.GetAnimReference();
 			SegmentBeginTime = Segment.StartPos;
 			SegmentLength = Segment.GetLength();
 
@@ -139,7 +139,7 @@ void FAnimLinkableElement::OnChanged(float NewMontageTime)
 	{
 		// Update to the detected segment
 		FAnimSegment& Segment = Slot.AnimTrack.AnimSegments[SegmentIndex];
-		LinkedSequence = Segment.AnimReference;
+		LinkedSequence = Segment.GetAnimReference();
 		SegmentBeginTime = Segment.StartPos;
 		SegmentLength = Segment.GetLength();
 
@@ -423,7 +423,7 @@ void FAnimLinkableElement::RefreshSegmentOnLoad()
 			{
 				FAnimSegment& Segment = Slot.AnimTrack.AnimSegments[SegmentIndex];
 
-				if(Segment.AnimReference == LinkedSequence)
+				if(Segment.GetAnimReference() == LinkedSequence)
 				{
 					if(CachedLinkMethod == EAnimLinkMethod::Relative)
 					{
diff --git a/Engine/Source/Runtime/Engine/Private/Animation/AnimMontage.cpp b/Engine/Source/Runtime/Engine/Private/Animation/AnimMontage.cpp
--- a/Engine/Source/Runtime/Engine/Private/Animation/AnimMontage.cpp
+++ b/Engine/Source/Runtime/Engine/Private/Animation/AnimMontage.cpp
@@ -373,9 +373,20 @@ void UAnimMontage::PostLoad()
 			UE_LOG(LogAnimMontage, Display, TEXT("UAnimMontage::PostLoad: The actual sequence length for %s does not match the length stored in the asset, please resave the asset."), *GetFullName());
 			SetCompositeLength(CurrentCalculatedLength);
 		}
+
+#if WITH_EDITOR
+		for (const FAnimSegment& AnimSegment : Track.AnimSegments)
+		{
+			if(AnimSegment.IsPlayLengthOutOfDate())
+			{
+				UE_LOG(LogAnimation, Warning, TEXT("AnimMontage (%s) contains a Segment for Slot (%s) for which the playable length %f is out-of-sync with the represented AnimationSequence its length %f (%s). Please up-date the segment and resave."), *GetFullName(), *SlotIter->SlotName.ToString(), (AnimSegment.AnimEndTime - AnimSegment.AnimStartTime), AnimSegment.GetAnimReference()->GetPlayLength(),
+					*AnimSegment.GetAnimReference()->GetFullName());
+			}
+		}
+#endif
 	}
 
-	for(auto& Composite : CompositeSections)
+	for(FCompositeSection& Composite : CompositeSections)
 	{
 		if(Composite.StartTime_DEPRECATED != 0.0f)
 		{
@@ -397,9 +408,9 @@ void UAnimMontage::PostLoad()
 		{
 			for (FAnimSegment& Segment : Slot.AnimTrack.AnimSegments)
 			{
-				if (Segment.AnimReference)
+				if (UAnimSequenceBase* AnimReference = Segment.GetAnimReference())
 				{
-					Segment.AnimReference->EnableRootMotionSettingFromMontage(true, RootMotionRootLock);
+					AnimReference->EnableRootMotionSettingFromMontage(true, RootMotionRootLock);
 				}
 			}
 		}
@@ -412,7 +423,7 @@ void UAnimMontage::PostLoad()
 		{
 			if ( SlotAnimTracks[I].AnimTrack.AnimSegments.Num() > 0 )
 			{
-				UAnimSequenceBase* SequenceBase = SlotAnimTracks[I].AnimTrack.AnimSegments[0].AnimReference;
+				UAnimSequenceBase* SequenceBase = SlotAnimTracks[I].AnimTrack.AnimSegments[0].GetAnimReference();
 				UAnimSequence* BaseAdditivePose = (SequenceBase) ? SequenceBase->GetAdditiveBasePose() : nullptr;
 				if (BaseAdditivePose)
 				{
@@ -431,10 +442,10 @@ void UAnimMontage::PostLoad()
 		{
 			if ( SlotAnimTracks[I].AnimTrack.AnimSegments.Num() > 0 )
 			{
-				UAnimSequenceBase* SequenceBase = SlotAnimTracks[I].AnimTrack.AnimSegments[0].AnimReference;
+				UAnimSequenceBase* SequenceBase = SlotAnimTracks[I].AnimTrack.AnimSegments[0].GetAnimReference();
 				if (SequenceBase && !MySkeleton->IsCompatible(SequenceBase->GetSkeleton()))
 				{
-					SlotAnimTracks[I].AnimTrack.AnimSegments[0].AnimReference = nullptr;
+					SlotAnimTracks[I].AnimTrack.AnimSegments[0].SetAnimReference(nullptr);
 					MarkPackageDirty();
 					break;
 				}
@@ -956,10 +967,10 @@ const TArray<class UAnimMetaData*> UAnimMontage::GetSectionMetaData(FName Sectio
 							// now add the animations within this section
 							for (auto& SegmentIter : SlotIter.AnimTrack.AnimSegments)
 							{
-								if (SegmentIter.AnimReference)
+								if (UAnimSequenceBase* AnimReference = SegmentIter.GetAnimReference())
 								{
 									// only add unique here
-									TArray<UAnimMetaData*> RefMetadata = SegmentIter.AnimReference->GetMetaData();
+									TArray<UAnimMetaData*> RefMetadata = AnimReference->GetMetaData();
 
 									for (auto& RefData : RefMetadata)
 									{
@@ -987,10 +998,10 @@ const TArray<class UAnimMetaData*> UAnimMontage::GetSectionMetaData(FName Sectio
 							{
 								if (SegmentIter.IsIncluded(SectionStartTime, SectionEndTime))
 								{
-									if (SegmentIter.AnimReference)
+									if (UAnimSequenceBase* AnimReference = SegmentIter.GetAnimReference())
 									{
 										// only add unique here
-										TArray<UAnimMetaData*> RefMetadata = SegmentIter.AnimReference->GetMetaData();
+										TArray<UAnimMetaData*> RefMetadata = AnimReference->GetMetaData();
 
 										for (auto& RefData : RefMetadata)
 										{
@@ -1069,8 +1080,6 @@ void UAnimMontage::UpdateLinkableElements()
 
 void UAnimMontage::UpdateLinkableElements(int32 SlotIdx, int32 SegmentIdx)
 {
-	FAnimSegment* UpdatedSegment = &SlotAnimTracks[SlotIdx].AnimTrack.AnimSegments[SegmentIdx];
-
 	for (FCompositeSection& Section : CompositeSections)
 	{
 		if (Section.GetSlotIndex() == SlotIdx && Section.GetSegmentIndex() == SegmentIdx)
@@ -1124,9 +1133,9 @@ void UAnimMontage::RefreshParentAssetData()
 		{
 			FAnimSegment& Segment = SlotTrack.AnimTrack.AnimSegments[SegmentIdx];
 			FAnimSegment& ParentSegment = ParentMontage->SlotAnimTracks[SlotIdx].AnimTrack.AnimSegments[SegmentIdx];
-			UAnimSequenceBase* SourceReference = Segment.AnimReference;
+			UAnimSequenceBase* SourceReference = Segment.GetAnimReference();
 			UAnimSequenceBase* TargetReference = Cast<UAnimSequenceBase>(AssetMappingTable->GetMappedAsset(SourceReference));
-			Segment.AnimReference = TargetReference;
+			Segment.SetAnimReference(TargetReference);
 
 			float LengthChange = FMath::IsNearlyZero(SourceReference->GetPlayLength()) ? 0.f : TargetReference->GetPlayLength() / SourceReference->GetPlayLength();
 			float RateChange = FMath::IsNearlyZero(SourceReference->RateScale) ? 0.f : FMath::Abs(TargetReference->RateScale / SourceReference->RateScale);
@@ -1227,7 +1236,7 @@ void UAnimMontage::CollectMarkers()
 		const FAnimTrack& AnimTrack = SlotAnimTracks[SyncSlotIndex].AnimTrack;
 		for (const auto& Seg : AnimTrack.AnimSegments)
 		{
-			const UAnimSequence* Sequence = Cast<UAnimSequence>(Seg.AnimReference);
+			const UAnimSequence* Sequence = Cast<UAnimSequence>(Seg.GetAnimReference());
 			if (Sequence && Sequence->AuthoredSyncMarkers.Num() > 0)
 			{
 				// @todo this won't work well if you have starttime < end time and it does have negative playrate
@@ -2848,11 +2857,7 @@ UAnimMontage* UAnimMontage::CreateSlotAnimationAsDynamicMontage_WithBlendSetting
 	FSlotAnimationTrack& NewTrack = NewMontage->SlotAnimTracks[0];
 	NewTrack.SlotName = SlotNodeName;
 	FAnimSegment NewSegment;
-	NewSegment.AnimReference = Asset;
-	NewSegment.AnimStartTime = 0.f;
-	NewSegment.AnimEndTime = Asset->GetPlayLength();
-	NewSegment.AnimPlayRate = 1.f;
-	NewSegment.StartPos = 0.f;
+	NewSegment.SetAnimReference(Asset, true);
 	NewSegment.LoopingCount = LoopCount;
 	PRAGMA_DISABLE_DEPRECATION_WARNINGS
     NewMontage->SequenceLength = NewSegment.GetLength();
