CIS NonUnity Error: AnimMontageSegmentDetails.h(38): error C2143: syntax error: missing ';' before '*'
#fix add missing include/predeclare

diff --git a/Engine/Source/Editor/DetailCustomizations/Private/AnimMontageSegmentDetails.cpp b/Engine/Source/Editor/DetailCustomizations/Private/AnimMontageSegmentDetails.cpp
--- a/Engine/Source/Editor/DetailCustomizations/Private/AnimMontageSegmentDetails.cpp
+++ b/Engine/Source/Editor/DetailCustomizations/Private/AnimMontageSegmentDetails.cpp
@@ -20,6 +20,7 @@
 #include "Settings/SkeletalMeshEditorSettings.h"
 #include "Widgets/Input/SNumericEntryBox.h"
 #include "Animation/EditorAnimSegment.h"
+#include "Animation/AnimCompositeBase.h"
 
 #define LOCTEXT_NAMESPACE "AnimMontageSegmentDetails"
 
diff --git a/Engine/Source/Editor/DetailCustomizations/Private/AnimMontageSegmentDetails.h b/Engine/Source/Editor/DetailCustomizations/Private/AnimMontageSegmentDetails.h
--- a/Engine/Source/Editor/DetailCustomizations/Private/AnimMontageSegmentDetails.h
+++ b/Engine/Source/Editor/DetailCustomizations/Private/AnimMontageSegmentDetails.h
@@ -22,6 +22,7 @@ class SViewport;
 class UAnimationAsset;
 class UAnimSequenceBase;
 class USkeleton;
+struct FAnimSegment;
 
 class FAnimMontageSegmentDetails : public IDetailCustomization
 {
