WorldPartition: Use UWorld::IsPartitionedWorld instead of HasSubsystem<UWorldPartitionSubsystem> (prepare for subsystem always existing)

#rnx

diff --git a/Engine/Source/Editor/DataLayerEditor/Private/DataLayer/SDataLayerBrowser.cpp b/Engine/Source/Editor/DataLayerEditor/Private/DataLayer/SDataLayerBrowser.cpp
--- a/Engine/Source/Editor/DataLayerEditor/Private/DataLayer/SDataLayerBrowser.cpp
+++ b/Engine/Source/Editor/DataLayerEditor/Private/DataLayer/SDataLayerBrowser.cpp
@@ -134,7 +134,7 @@ void SDataLayerBrowser::Construct(const FArguments& InArgs)
 	ChildSlot
 	[
 		SAssignNew(ContentAreaBox, SVerticalBox)
-		.IsEnabled_Lambda([]() { return GWorld ? UWorld::HasSubsystem<UWorldPartitionSubsystem>(GWorld) : false; })
+		.IsEnabled_Lambda([]() { return UWorld::IsPartitionedWorld(GWorld); })
 	];
 
 	InitializeDataLayerBrowser();
diff --git a/Engine/Source/Editor/DetailCustomizations/Private/ActorDetails.cpp b/Engine/Source/Editor/DetailCustomizations/Private/ActorDetails.cpp
--- a/Engine/Source/Editor/DetailCustomizations/Private/ActorDetails.cpp
+++ b/Engine/Source/Editor/DetailCustomizations/Private/ActorDetails.cpp
@@ -149,7 +149,7 @@ void FActorDetails::CustomizeDetails( IDetailLayoutBuilder& DetailLayout )
 		const bool bShouldDisplayWorldPartitionProperties = Algo::AnyOf(SelectedActors, [](const TWeakObjectPtr<AActor> Actor)
 		{
 			UWorld* World = Actor.IsValid() ? Actor->GetTypedOuter<UWorld>() : nullptr;
-			return World != nullptr && UWorld::HasSubsystem<UWorldPartitionSubsystem>(World);
+			return UWorld::IsPartitionedWorld(World);
 		});
 
 		if (!bShouldDisplayWorldPartitionProperties)
@@ -707,7 +707,7 @@ void FActorDetails::AddActorCategory( IDetailLayoutBuilder& DetailBuilder, const
 		if (SelectedActorInfo.SelectionClass != AWorldSettings::StaticClass())
 		{
 			// Actor Packaging Mode
-			const bool bIsPartitionedWorld = UWorld::HasSubsystem<UWorldPartitionSubsystem>(SelectedActorInfo.SharedWorld);
+			const bool bIsPartitionedWorld = UWorld::IsPartitionedWorld(SelectedActorInfo.SharedWorld);
 
 			auto OnGetMenuContent = [=]() -> TSharedRef<SWidget> {
 				FMenuBuilder MenuBuilder(true, nullptr);
diff --git a/Engine/Source/Editor/DetailCustomizations/Private/WorldSettingsDetails.cpp b/Engine/Source/Editor/DetailCustomizations/Private/WorldSettingsDetails.cpp
--- a/Engine/Source/Editor/DetailCustomizations/Private/WorldSettingsDetails.cpp
+++ b/Engine/Source/Editor/DetailCustomizations/Private/WorldSettingsDetails.cpp
@@ -106,7 +106,7 @@ void FWorldSettingsDetails::AddWorldCustomization(IDetailLayoutBuilder& DetailBu
 
 	if (CustomizedLevel)
 	{
-		const bool bIsPartitionedWorld = UWorld::HasSubsystem<UWorldPartitionSubsystem>(CustomizedLevel->GetWorld());
+		const bool bIsPartitionedWorld = UWorld::IsPartitionedWorld(CustomizedLevel->GetWorld());
 
 		IDetailCategoryBuilder& WorldCategory = DetailBuilder.EditCategory("World");
 		if (GetDefault<UEditorExperimentalSettings>()->bEnableOneFilePerActorSupport)
diff --git a/Engine/Source/Editor/FoliageEdit/Private/FoliageEdMode.cpp b/Engine/Source/Editor/FoliageEdit/Private/FoliageEdMode.cpp
--- a/Engine/Source/Editor/FoliageEdit/Private/FoliageEdMode.cpp
+++ b/Engine/Source/Editor/FoliageEdit/Private/FoliageEdMode.cpp
@@ -2069,7 +2069,7 @@ void FEdModeFoliage::RemoveSelectedInstances(UWorld* InWorld)
 void FEdModeFoliage::GetFoliageTypeFilters(TArray<const UClass*>& OutFilters) const
 {
 	OutFilters.Add(UFoliageType_InstancedStaticMesh::StaticClass());
-	if (GetWorld()->HasSubsystem<UWorldPartitionSubsystem>())
+	if (UWorld::IsPartitionedWorld(GetWorld()))
 	{
 		return;
 	}
diff --git a/Engine/Source/Editor/FoliageEdit/Private/SFoliageEdit.cpp b/Engine/Source/Editor/FoliageEdit/Private/SFoliageEdit.cpp
--- a/Engine/Source/Editor/FoliageEdit/Private/SFoliageEdit.cpp
+++ b/Engine/Source/Editor/FoliageEdit/Private/SFoliageEdit.cpp
@@ -908,7 +908,7 @@ FText SFoliageEdit::GetActiveToolMessage() const
 
 EVisibility SFoliageEdit::GetVisibility_DataLayer() const
 {
-	if (UWorld::HasSubsystem<UWorldPartitionSubsystem>(FoliageEditMode->GetWorld()) && (IsPaintTool() || IsPlaceTool() || IsReapplySettingsTool() || IsPaintFillTool()))
+	if (UWorld::IsPartitionedWorld(FoliageEditMode->GetWorld()) && (IsPaintTool() || IsPlaceTool() || IsReapplySettingsTool() || IsPaintFillTool()))
 	{
 		return EVisibility::Visible;
 	}
diff --git a/Engine/Source/Editor/Layers/Private/SLayerBrowser.cpp b/Engine/Source/Editor/Layers/Private/SLayerBrowser.cpp
--- a/Engine/Source/Editor/Layers/Private/SLayerBrowser.cpp
+++ b/Engine/Source/Editor/Layers/Private/SLayerBrowser.cpp
@@ -172,7 +172,7 @@ void SLayerBrowser::Construct(const FArguments& InArgs)
 	ChildSlot
 		[
 			SAssignNew(ContentAreaBox, SVerticalBox)
-			.IsEnabled_Lambda([]() { return GWorld ? !UWorld::HasSubsystem<UWorldPartitionSubsystem>(GWorld) : false; })
+			.IsEnabled_Lambda([]() { return !UWorld::IsPartitionedWorld(GWorld); })
 		];
 
 	SetupLayersMode();
diff --git a/Engine/Source/Editor/LevelEditor/Private/LevelEditorActions.cpp b/Engine/Source/Editor/LevelEditor/Private/LevelEditorActions.cpp
--- a/Engine/Source/Editor/LevelEditor/Private/LevelEditorActions.cpp
+++ b/Engine/Source/Editor/LevelEditor/Private/LevelEditorActions.cpp
@@ -422,7 +422,7 @@ bool FLevelEditorActionCallbacks::CanSaveUnpartitionedWorld()
 		return false;
 	}
 
-	return GetWorld() && !GetWorld()->HasSubsystem<UWorldPartitionSubsystem>();
+	return !UWorld::IsPartitionedWorld(GetWorld());
 }
 
 void FLevelEditorActionCallbacks::Save()
diff --git a/Engine/Source/Editor/LevelEditor/Private/SLevelViewport.cpp b/Engine/Source/Editor/LevelEditor/Private/SLevelViewport.cpp
--- a/Engine/Source/Editor/LevelEditor/Private/SLevelViewport.cpp
+++ b/Engine/Source/Editor/LevelEditor/Private/SLevelViewport.cpp
@@ -1618,7 +1618,7 @@ void SLevelViewport::BindShowCommands( FUICommandList& OutCommandList )
 
 	// Show Layers
 	{
-		auto CanToggleAllLayers = [this]() { return !UWorld::HasSubsystem<UWorldPartitionSubsystem>(GetWorld()); };
+		auto CanToggleAllLayers = [this]() { return !UWorld::IsPartitionedWorld(GetWorld()); };
 		// Map 'Show All' and 'Hide All' commands
 		OutCommandList.MapAction(
 			LevelViewportCommands.ShowAllLayers,
diff --git a/Engine/Source/Editor/LevelEditor/Private/SLevelViewportToolBar.cpp b/Engine/Source/Editor/LevelEditor/Private/SLevelViewportToolBar.cpp
--- a/Engine/Source/Editor/LevelEditor/Private/SLevelViewportToolBar.cpp
+++ b/Engine/Source/Editor/LevelEditor/Private/SLevelViewportToolBar.cpp
@@ -1207,7 +1207,7 @@ void SLevelViewportToolBar::FillShowMenu(UToolMenu* Menu) const
 			}
 
 			// Show Layers sub-menu is dynamically generated when the user enters 'show' menu
-			if (!UWorld::HasSubsystem<UWorldPartitionSubsystem>(World))
+			if (!UWorld::IsPartitionedWorld(World))
 			{
 				Section.AddSubMenu("ShowLayers", LOCTEXT("ShowLayersMenu", "Layers"), LOCTEXT("ShowLayersMenu_ToolTip", "Show layers flags"),
 					FNewToolMenuDelegate::CreateStatic(&SLevelViewportToolBar::FillShowLayersMenu, Viewport), false, FSlateIcon(FAppStyle::Get().GetStyleSetName(), "ShowFlagsMenu.SubMenu.Layers"));
diff --git a/Engine/Source/Editor/UnrealEd/Private/Commandlets/DerivedDataCacheCommandlet.cpp b/Engine/Source/Editor/UnrealEd/Private/Commandlets/DerivedDataCacheCommandlet.cpp
--- a/Engine/Source/Editor/UnrealEd/Private/Commandlets/DerivedDataCacheCommandlet.cpp
+++ b/Engine/Source/Editor/UnrealEd/Private/Commandlets/DerivedDataCacheCommandlet.cpp
@@ -358,7 +358,7 @@ void UDerivedDataCacheCommandlet::CacheWorldPackages(UWorld* World, uint8 Packag
 
 	// If the world is partitioned
 	bool bResult = true;
-	if (World->HasSubsystem<UWorldPartitionSubsystem>())
+	if (UWorld::IsPartitionedWorld(World))
 	{
 		// Ensure the world has a valid world partition.
 		UWorldPartition* WorldPartition = World->GetWorldPartition();
diff --git a/Engine/Source/Editor/UnrealEd/Private/WorldPartition/WorldPartitionBuilder.cpp b/Engine/Source/Editor/UnrealEd/Private/WorldPartition/WorldPartitionBuilder.cpp
--- a/Engine/Source/Editor/UnrealEd/Private/WorldPartition/WorldPartitionBuilder.cpp
+++ b/Engine/Source/Editor/UnrealEd/Private/WorldPartition/WorldPartitionBuilder.cpp
@@ -89,7 +89,7 @@ bool UWorldPartitionBuilder::RunBuilder(UWorld* World)
 		FScopedEditorWorld EditorWorld(World, IVS);
 
 		// Make sure the world is partitioned
-		if (World->HasSubsystem<UWorldPartitionSubsystem>())
+		if (UWorld::IsPartitionedWorld(World))
 		{
 			// Ensure the world has a valid world partition.
 			UWorldPartition* WorldPartition = World->GetWorldPartition();
diff --git a/Engine/Source/Editor/WorldBrowser/Private/SWorldHierarchy.cpp b/Engine/Source/Editor/WorldBrowser/Private/SWorldHierarchy.cpp
--- a/Engine/Source/Editor/WorldBrowser/Private/SWorldHierarchy.cpp
+++ b/Engine/Source/Editor/WorldBrowser/Private/SWorldHierarchy.cpp
@@ -62,7 +62,7 @@ void SWorldHierarchy::OnBrowseWorld(UWorld* InWorld)
 	// Bind to a new world
 	if (InWorld)
 	{
-		if (UWorld::HasSubsystem<UWorldPartitionSubsystem>(InWorld))
+		if (UWorld::IsPartitionedWorld(InWorld))
 		{
 			ChildSlot
 			[
diff --git a/Engine/Source/Runtime/Engine/Classes/Engine/World.h b/Engine/Source/Runtime/Engine/Classes/Engine/World.h
--- a/Engine/Source/Runtime/Engine/Classes/Engine/World.h
+++ b/Engine/Source/Runtime/Engine/Classes/Engine/World.h
@@ -2679,6 +2679,19 @@ public:
 	*/
 	bool IsPartitionedWorld() const { return GetWorldPartition() != nullptr; }
 
+	/**
+	* Returns true if world contains an associated UWorldPartition object.
+	*/
+	static bool IsPartitionedWorld(const UWorld* InWorld)
+	{
+		if (InWorld)
+		{
+			return InWorld->IsPartitionedWorld();
+		}
+
+		return false;
+	}
+		
 	/**
 	 * Returns the current levels BSP model.
 	 *
diff --git a/Engine/Source/Runtime/Engine/Private/ActorEditor.cpp b/Engine/Source/Runtime/Engine/Private/ActorEditor.cpp
--- a/Engine/Source/Runtime/Engine/Private/ActorEditor.cpp
+++ b/Engine/Source/Runtime/Engine/Private/ActorEditor.cpp
@@ -78,13 +78,9 @@ bool AActor::CanEditChange(const FProperty* PropertyThatWillChange) const
 	{
 		if (!IsTemplate())
 		{
-			if (UWorld* World = GetTypedOuter<UWorld>())
+			if (!UWorld::IsPartitionedWorld(GetTypedOuter<UWorld>()))
 			{
-				const bool bIsPartitionedWorld = UWorld::HasSubsystem<UWorldPartitionSubsystem>(World);
-				if (!bIsPartitionedWorld)
-				{
-					return false;
-				}
+				return false;
 			}
 		}
 	}
@@ -1516,7 +1512,7 @@ bool AActor::IsValidForDataLayer() const
 		return false;
 	}
 
-	const bool bIsPartitionedActor = UWorld::HasSubsystem<UWorldPartitionSubsystem>(World);
+	const bool bIsPartitionedActor = UWorld::IsPartitionedWorld(World);
 	const bool bIsInEditorWorld = World->WorldType == EWorldType::Editor;
 	const bool bIsBuilderBrush = FActorEditorUtils::IsABuilderBrush(this);
 	const bool bIsHidden = GetClass()->GetDefaultObject<AActor>()->bHiddenEd;
diff --git a/Engine/Source/Runtime/Engine/Private/ActorPartition/ActorPartitionSubsystem.cpp b/Engine/Source/Runtime/Engine/Private/ActorPartition/ActorPartitionSubsystem.cpp
--- a/Engine/Source/Runtime/Engine/Private/ActorPartition/ActorPartitionSubsystem.cpp
+++ b/Engine/Source/Runtime/Engine/Private/ActorPartition/ActorPartitionSubsystem.cpp
@@ -330,7 +330,7 @@ UActorPartitionSubsystem::UActorPartitionSubsystem()
 
 bool UActorPartitionSubsystem::IsLevelPartition() const
 {
-	return !GetWorld()->HasSubsystem<UWorldPartitionSubsystem>();
+	return !UWorld::IsPartitionedWorld(GetWorld());
 }
 
 #if WITH_EDITOR
diff --git a/Engine/Source/Runtime/Foliage/Private/ProceduralFoliageComponent.cpp b/Engine/Source/Runtime/Foliage/Private/ProceduralFoliageComponent.cpp
--- a/Engine/Source/Runtime/Foliage/Private/ProceduralFoliageComponent.cpp
+++ b/Engine/Source/Runtime/Foliage/Private/ProceduralFoliageComponent.cpp
@@ -360,8 +360,7 @@ bool UProceduralFoliageComponent::CanEditChange(const FProperty* InProperty) con
 	{
 		if (!IsTemplate())
 		{
-			const bool bIsPartitionedWorld = UWorld::HasSubsystem<UWorldPartitionSubsystem>(GetWorld());
-			if (!bIsPartitionedWorld)
+			if (!UWorld::IsPartitionedWorld(GetWorld()))
 			{
 				return false;
 			}
diff --git a/Engine/Source/Runtime/Landscape/Private/LandscapeSubsystem.cpp b/Engine/Source/Runtime/Landscape/Private/LandscapeSubsystem.cpp
--- a/Engine/Source/Runtime/Landscape/Private/LandscapeSubsystem.cpp
+++ b/Engine/Source/Runtime/Landscape/Private/LandscapeSubsystem.cpp
@@ -202,7 +202,7 @@ int32 ULandscapeSubsystem::GetOudatedPhysicalMaterialComponentsCount()
 
 bool ULandscapeSubsystem::IsGridBased() const
 {
-	return UWorld::HasSubsystem<UWorldPartitionSubsystem>(GetWorld());
+	return UWorld::IsPartitionedWorld(GetWorld());
 }
 
 void ULandscapeSubsystem::ChangeGridSize(ULandscapeInfo* LandscapeInfo, uint32 GridSizeInComponents)
