Fix to Localization macros usage.

diff --git a/Engine/Plugins/Developer/TraceDataFiltering/Source/UnrealInsightsInterface/Private/UnrealInsightsLauncher.cpp b/Engine/Plugins/Developer/TraceDataFiltering/Source/UnrealInsightsInterface/Private/UnrealInsightsLauncher.cpp
--- a/Engine/Plugins/Developer/TraceDataFiltering/Source/UnrealInsightsInterface/Private/UnrealInsightsLauncher.cpp
+++ b/Engine/Plugins/Developer/TraceDataFiltering/Source/UnrealInsightsInterface/Private/UnrealInsightsLauncher.cpp
@@ -124,3 +124,5 @@ void FUnrealInsightsLauncher::StartUnrealInsights(const FString& Path)
 		ReportMessageLog.Notify();
 	}
 }
+
+#undef LOCTEXT_NAMESPACE
\ No newline at end of file
diff --git a/Engine/Plugins/Experimental/Animation/PoseSearch/Source/Editor/Private/PoseSearchDebugger.cpp b/Engine/Plugins/Experimental/Animation/PoseSearch/Source/Editor/Private/PoseSearchDebugger.cpp
--- a/Engine/Plugins/Experimental/Animation/PoseSearch/Source/Editor/Private/PoseSearchDebugger.cpp
+++ b/Engine/Plugins/Experimental/Animation/PoseSearch/Source/Editor/Private/PoseSearchDebugger.cpp
@@ -2325,4 +2325,6 @@ FName FDebuggerViewCreator::GetName() const
 	return Name;
 }
 
-} // namespace UE::PoseSearch
\ No newline at end of file
+} // namespace UE::PoseSearch
+
+#undef LOCTEXT_NAMESPACE
\ No newline at end of file
diff --git a/Engine/Plugins/Experimental/GeometryCollectionPlugin/Source/GeometryCollectionEditor/Private/GeometryCollection/GeometryCollectionSelectionCommands.cpp b/Engine/Plugins/Experimental/GeometryCollectionPlugin/Source/GeometryCollectionEditor/Private/GeometryCollection/GeometryCollectionSelectionCommands.cpp
--- a/Engine/Plugins/Experimental/GeometryCollectionPlugin/Source/GeometryCollectionEditor/Private/GeometryCollection/GeometryCollectionSelectionCommands.cpp
+++ b/Engine/Plugins/Experimental/GeometryCollectionPlugin/Source/GeometryCollectionEditor/Private/GeometryCollection/GeometryCollectionSelectionCommands.cpp
@@ -10,3 +10,5 @@ void FGeometryCollectionSelectionCommands::RegisterCommands()
 	UI_COMMAND(SelectNone, "Deselect All Geometry In Hierarchy", "Deselect all geometry in hierarchy", EUserInterfaceActionType::Button, FInputChord());
 	UI_COMMAND(SelectInverseGeometry, "Select Inverse Geometry In Hierarchy", "Select inverse geometry in hierarchy", EUserInterfaceActionType::Button, FInputChord());
 }
+
+#undef LOCTEXT_NAMESPACE
\ No newline at end of file
diff --git a/Engine/Plugins/Experimental/MeshLODToolset/Source/MeshLODToolset/Private/Tools/LODManagerTool.cpp b/Engine/Plugins/Experimental/MeshLODToolset/Source/MeshLODToolset/Private/Tools/LODManagerTool.cpp
--- a/Engine/Plugins/Experimental/MeshLODToolset/Source/MeshLODToolset/Private/Tools/LODManagerTool.cpp
+++ b/Engine/Plugins/Experimental/MeshLODToolset/Source/MeshLODToolset/Private/Tools/LODManagerTool.cpp
@@ -675,3 +675,5 @@ void ULODManagerTool::RemoveUnreferencedMaterials()
 	GetToolManager()->EndUndoTransaction();
 	bLODInfoValid = false;
 }
+
+#undef LOCTEXT_NAMESPACE
diff --git a/Engine/Plugins/Media/BinkMedia/Source/BinkMediaPlayer/Private/BinkMediaPlayerModule.cpp b/Engine/Plugins/Media/BinkMedia/Source/BinkMediaPlayer/Private/BinkMediaPlayerModule.cpp
--- a/Engine/Plugins/Media/BinkMedia/Source/BinkMediaPlayer/Private/BinkMediaPlayerModule.cpp
+++ b/Engine/Plugins/Media/BinkMedia/Source/BinkMediaPlayer/Private/BinkMediaPlayerModule.cpp
@@ -6,6 +6,8 @@
 
 DEFINE_LOG_CATEGORY(LogBink);
 
+#define LOCTEXT_NAMESPACE "BinkMediaPlayerModule"
+
 TSharedPtr<FBinkMovieStreamer, ESPMode::ThreadSafe> MovieStreamer;
 
 #if BINKPLUGIN_UE4_EDITOR
@@ -288,5 +290,7 @@ struct FBinkMediaPlayerModule : IModuleInterface, FTickableGameObject
 	bool bPluginInitialized = false;
 };
 
+#undef LOCTEXT_NAMESPACE
+
 IMPLEMENT_MODULE( FBinkMediaPlayerModule, BinkMediaPlayer )
 
diff --git a/Engine/Plugins/Media/BinkMedia/Source/BinkMediaPlayer/Private/BinkMediaPlayerPCH.h b/Engine/Plugins/Media/BinkMedia/Source/BinkMediaPlayer/Private/BinkMediaPlayerPCH.h
--- a/Engine/Plugins/Media/BinkMedia/Source/BinkMediaPlayer/Private/BinkMediaPlayerPCH.h
+++ b/Engine/Plugins/Media/BinkMedia/Source/BinkMediaPlayer/Private/BinkMediaPlayerPCH.h
@@ -51,9 +51,3 @@ static int GetNumSpeakers()
 }
 
 DECLARE_LOG_CATEGORY_EXTERN(LogBink, Log, All);
-
-#ifndef LOCTEXT_NAMESPACE
-#define LOCTEXT_NAMESPACE "BinkMediaPlayerModule"
-#endif
-
-
diff --git a/Engine/Plugins/Runtime/Metasound/Source/MetasoundFrontend/Private/MetasoundSourceInterface.cpp b/Engine/Plugins/Runtime/Metasound/Source/MetasoundFrontend/Private/MetasoundSourceInterface.cpp
--- a/Engine/Plugins/Runtime/Metasound/Source/MetasoundFrontend/Private/MetasoundSourceInterface.cpp
+++ b/Engine/Plugins/Runtime/Metasound/Source/MetasoundFrontend/Private/MetasoundSourceInterface.cpp
@@ -274,4 +274,4 @@ namespace Metasound
 	} // namespace Frontend
 } // namespace Metasound
 
-#define LOCTEXT_NAMESPACE "Metasound"
+#undef LOCTEXT_NAMESPACE
diff --git a/Engine/Plugins/VirtualProduction/LevelSnapshots/Source/LevelSnapshots/Private/Data/WorldSnapshotData.cpp b/Engine/Plugins/VirtualProduction/LevelSnapshots/Source/LevelSnapshots/Private/Data/WorldSnapshotData.cpp
--- a/Engine/Plugins/VirtualProduction/LevelSnapshots/Source/LevelSnapshots/Private/Data/WorldSnapshotData.cpp
+++ b/Engine/Plugins/VirtualProduction/LevelSnapshots/Source/LevelSnapshots/Private/Data/WorldSnapshotData.cpp
@@ -100,3 +100,5 @@ void FWorldSnapshotData::PostSerialize(const FArchive& Ar)
 		SnapshotVersionInfo.Initialize(bWithoutSnapshotVersion);
 	}
 }
+
+#undef LOCTEXT_NAMESPACE
\ No newline at end of file
diff --git a/Engine/Plugins/VirtualProduction/RemoteControl/Source/RemoteControl/Private/RemoteControlModule.h b/Engine/Plugins/VirtualProduction/RemoteControl/Source/RemoteControl/Private/RemoteControlModule.h
--- a/Engine/Plugins/VirtualProduction/RemoteControl/Source/RemoteControl/Private/RemoteControlModule.h
+++ b/Engine/Plugins/VirtualProduction/RemoteControl/Source/RemoteControl/Private/RemoteControlModule.h
@@ -149,6 +149,4 @@ private:
 
 PRAGMA_ENABLE_DEPRECATION_WARNINGS
 
-#undef LOCTEXT_NAMESPACE
-
 IMPLEMENT_MODULE(FRemoteControlModule, RemoteControl);
diff --git a/Engine/Plugins/VirtualProduction/RemoteControl/Source/RemoteControlUI/Private/UI/SRCPanelTreeNode.h b/Engine/Plugins/VirtualProduction/RemoteControl/Source/RemoteControlUI/Private/UI/SRCPanelTreeNode.h
--- a/Engine/Plugins/VirtualProduction/RemoteControl/Source/RemoteControlUI/Private/UI/SRCPanelTreeNode.h
+++ b/Engine/Plugins/VirtualProduction/RemoteControl/Source/RemoteControlUI/Private/UI/SRCPanelTreeNode.h
@@ -74,6 +74,3 @@ private:
 	/** The splitter offset to align the group splitter with the other row's splitters. */
 	static constexpr float SplitterOffset = 0.008f;
 };
-
-
-#undef LOCTEXT_NAMESPACE /*RemoteControlPanelNode*/
diff --git a/Engine/Source/Editor/DerivedDataEditor/Private/DerivedDataEditorModule.cpp b/Engine/Source/Editor/DerivedDataEditor/Private/DerivedDataEditorModule.cpp
--- a/Engine/Source/Editor/DerivedDataEditor/Private/DerivedDataEditorModule.cpp
+++ b/Engine/Source/Editor/DerivedDataEditor/Private/DerivedDataEditorModule.cpp
@@ -201,4 +201,6 @@ void FDerivedDataEditorModule::OnSettingsDialogClosed(const TSharedRef<SWindow>&
 {
 	SettingsWindow = nullptr;
 	SettingsDialog = nullptr;
-}
\ No newline at end of file
+}
+
+#undef LOCTEXT_NAMESPACE
diff --git a/Engine/Source/Editor/DetailCustomizations/Private/MathStructProxyCustomizations.cpp b/Engine/Source/Editor/DetailCustomizations/Private/MathStructProxyCustomizations.cpp
--- a/Engine/Source/Editor/DetailCustomizations/Private/MathStructProxyCustomizations.cpp
+++ b/Engine/Source/Editor/DetailCustomizations/Private/MathStructProxyCustomizations.cpp
@@ -111,8 +111,6 @@ FText FMathStructProxyCustomization::OnGetValueToolTip(TWeakPtr<IPropertyHandle>
 	return FText::GetEmpty();
 }
 
-#define LOCTEXT_NAMESPACE "MatrixStructCustomization"
-
 TSharedRef<IPropertyTypeCustomization> FMatrixStructCustomization::MakeInstance()
 {
 	return MakeShareable( new FMatrixStructCustomization );
@@ -500,8 +498,6 @@ bool FMatrixStructCustomization::FlushValues( TWeakPtr<IPropertyHandle> Property
 	return true;
 }
 
-#undef LOCTEXT_NAMESPACE
-
 TSharedRef<IPropertyTypeCustomization> FTransformStructCustomization::MakeInstance() 
 {
 	return MakeShareable( new FTransformStructCustomization );
@@ -798,3 +794,4 @@ bool FQuatStructCustomization::FlushValues(TWeakPtr<IPropertyHandle> PropertyHan
 	return true;
 }
 
+#undef LOCTEXT_NAMESPACE
diff --git a/Engine/Source/Editor/LevelEditor/Private/LevelEditorContextMenu.cpp b/Engine/Source/Editor/LevelEditor/Private/LevelEditorContextMenu.cpp
--- a/Engine/Source/Editor/LevelEditor/Private/LevelEditorContextMenu.cpp
+++ b/Engine/Source/Editor/LevelEditor/Private/LevelEditorContextMenu.cpp
@@ -230,7 +230,7 @@ void FLevelEditorContextMenu::AddPlayFromHereSubMenu(FToolMenuSection& Section)
 {
 	if(FLevelEditorActionCallbacks::PlayFromHere_IsVisible())
 	{
-		Section.AddSubMenu("PlayFromHere", LOCTEXT("PlayFromHere", "Play From Here"), LOCTEXT("", ""), FNewToolMenuDelegate::CreateLambda([](UToolMenu* InMenu)
+		Section.AddSubMenu("PlayFromHere", LOCTEXT("PlayFromHere", "Play From Here"), FText(), FNewToolMenuDelegate::CreateLambda([](UToolMenu* InMenu)
 		{
 			FToolMenuSection& NewSection = InMenu->AddSection("Section");
 			FUIAction PlayFromHere(FExecuteAction::CreateStatic(&FLevelEditorActionCallbacks::PlayFromHere_Clicked, false));
