Device Profile Editor:
- Fixed an issue where the unsorted cvars list was missing any non-categorized cvar with a period in it
- Added sorting to the category lists and drop downs
- Added help text as a tooltip on cvar names (where possible: platform-specific cvars that aren't defined in a module loaded in the editor won't have tooltips)

diff --git a/Engine/Source/Editor/DetailCustomizations/Private/DeviceProfileDetails.cpp b/Engine/Source/Editor/DetailCustomizations/Private/DeviceProfileDetails.cpp
--- a/Engine/Source/Editor/DetailCustomizations/Private/DeviceProfileDetails.cpp
+++ b/Engine/Source/Editor/DetailCustomizations/Private/DeviceProfileDetails.cpp
@@ -27,7 +27,7 @@
 #include "DetailCategoryBuilder.h"
 #include "TextureLODSettingsDetails.h"
 #include "Widgets/Input/SSearchBox.h"
-
+#include "Misc/EnumRange.h"
 
 #define LOCTEXT_NAMESPACE "DeviceProfileDetails"
 
@@ -200,6 +200,7 @@ namespace DeviceProfileCVarFormatHelper
 	}
 };
 
+ENUM_RANGE_BY_COUNT(DeviceProfileCVarFormatHelper::ECVarGroup, DeviceProfileCVarFormatHelper::Max_CVarCategories);
 
 ////////////////////////////////////////////////
 // FConsoleVariablesAvailableVisitor
@@ -308,21 +309,56 @@ void SCVarSelectionPanel::Construct(const FArguments& InArgs, const FString& CVa
 
 	if (DeviceProfileCVarFormatHelper::CategoryEnumFromPrefix(CVarPrefix) == DeviceProfileCVarFormatHelper::CVG_Uncategorized)
 	{
-		for(TArray<TSharedPtr<FString>>::TIterator CVarIt(UnprocessedCVars); CVarIt; ++CVarIt)
+		// Make a list of existing prefixes
+		TSet<FString> CategoryPrefixes;
+		for (DeviceProfileCVarFormatHelper::ECVarGroup PrefixEnum : TEnumRange<DeviceProfileCVarFormatHelper::ECVarGroup>())
 		{
-			if((*CVarIt)->Contains(TEXT(".")) == false)
+			FString Prefix = DeviceProfileCVarFormatHelper::CategoryPrefixFromEnum(PrefixEnum);
+			if (Prefix.Len() > 0)
 			{
-				AllAvailableCVars.Add(*CVarIt);
-				CVarsToDisplay.Add(*CVarIt);
+				CategoryPrefixes.Add(Prefix + TEXT("."));
+			}
+		}
+		
+		// Add all cvars that *don't* match one of the other groups
+		for (const TSharedPtr<FString>& TestCVarPtr : UnprocessedCVars)
+		{
+			const FString& TestCVar = *TestCVarPtr.Get();
+
+			bool bBelongsInUncategorized = false;
+
+			// Figure out if the prefix on this variable belongs to any other category or if it's missing a prefix
+			int32 FirstPeriodIndex = INDEX_NONE;
+			TestCVar.FindChar(TEXT('.'), /*out*/ FirstPeriodIndex);
+
+			if (FirstPeriodIndex != INDEX_NONE)
+			{
+				const FString TestPrefix(TestCVar.Left(FirstPeriodIndex+1));
+				bBelongsInUncategorized = !CategoryPrefixes.Contains(TestPrefix);
+			}
+			else
+			{
+				// No period means no prefix, so it couldn't be in any other category
+				bBelongsInUncategorized = true;
+			}
+
+			if (bBelongsInUncategorized)			
+			{
+				AllAvailableCVars.Add(TestCVarPtr);
 			}
 		}
 	}
 	else
 	{
 		AllAvailableCVars = UnprocessedCVars;
-		CVarsToDisplay = UnprocessedCVars;
 	}
 
+	// Sort the list
+	AllAvailableCVars.Sort([](const TSharedPtr<FString>& A, const TSharedPtr<FString>& B) { return *A < *B; });
+
+	// Duplicate the list	
+	CVarsToDisplay = AllAvailableCVars;
+
 	ChildSlot
 	[
 		SNew(SBox)
@@ -359,6 +395,13 @@ FReply SCVarSelectionPanel::HandleCVarSelected(const TSharedPtr<FString> CVar)
 
 TSharedRef<ITableRow> SCVarSelectionPanel::GenerateCVarItemRow(TSharedPtr<FString> InItem, const TSharedRef<STableViewBase>& OwnerTable)
 {
+	FText ComposedTooltip = LOCTEXT("CVarSelectionMenuTooltip", "Select a Console Variable to add to the device profile");
+
+	if (IConsoleVariable* Var = IConsoleManager::Get().FindConsoleVariable(**InItem))
+	{
+		ComposedTooltip = FText::Format(LOCTEXT("CVarSelectionMenuTooltipWithHelp", "{0}\n\n{1}"), ComposedTooltip, FText::FromString(Var->GetHelp()));
+	}
+
 	return SNew(STableRow<TSharedPtr<FString>>, OwnerTable)
 		[
 			SNew(SButton)
@@ -366,7 +409,7 @@ TSharedRef<ITableRow> SCVarSelectionPanel::GenerateCVarItemRow(TSharedPtr<FStrin
 			.ButtonStyle(FEditorStyle::Get(), "HoverHintOnly")
 			.OnClicked(this, &SCVarSelectionPanel::HandleCVarSelected, InItem)
 			.ContentPadding(DeviceProfilePropertyConstants::CVarSelectionMenuPadding)
-			.ToolTipText(LOCTEXT("CVarSelectionMenuTooltip","Select a Console Variable to add to the device profile"))
+			.ToolTipText(ComposedTooltip)
 			[
 				SNew(STextBlock)
 				.Text(FText::FromString(*InItem))
@@ -485,6 +528,8 @@ void FDeviceProfileParentPropertyDetails::CreateParentPropertyView()
 		ActiveDeviceProfile->GatherParentCVarInformationRecursively(ParentCVarInformation);
 
 		IDetailGroup* ParentCVarsGroup = nullptr;
+
+		ParentCVarInformation.KeySort(TLess<>());
 		for(auto& ParentCVar : ParentCVarInformation)
 		{
 			FString ParentCVarName;
@@ -613,7 +658,16 @@ void FDeviceProfileConsoleVariablesPropertyDetails::CreateConsoleVariablesProper
 	// Put the property handles into the UI group for the details view.
 	for (auto& Current : CategoryPropertyMap)
 	{
-		const TArray<TSharedRef<IPropertyHandle>>& CurrentGroupsProperties = Current.Value;
+		TArray<TSharedRef<IPropertyHandle>>& CurrentGroupsProperties = Current.Value;
+		CurrentGroupsProperties.Sort([](const TSharedRef<IPropertyHandle>& A, const TSharedRef<IPropertyHandle>& B)
+		{ 
+			FString AVal;
+			A->GetValue(AVal);
+			FString BVal;
+			B->GetValue(BVal);
+			return AVal < BVal;
+		});
+
 		const FText GroupName = DeviceProfileCVarFormatHelper::CategoryTextFromEnum(Current.Key);
 
 		FString CVarPrefix = DeviceProfileCVarFormatHelper::CategoryPrefixFromEnum(Current.Key);
@@ -694,6 +748,12 @@ void FDeviceProfileConsoleVariablesPropertyDetails::CreateRowWidgetForCVarProper
 	const FString CVarName = UnformattedCVar.Left(CVarNameValueSplitIdx);
 	const FString CVarValueAsString = UnformattedCVar.Right(UnformattedCVar.Len() - (CVarNameValueSplitIdx + 1));
 
+	FText CVarHelp;
+	if (IConsoleVariable* Var = IConsoleManager::Get().FindConsoleVariable(*CVarName))
+	{
+		CVarHelp = FText::FromString(Var->GetHelp());
+	}
+
 	InGroup.AddWidgetRow()
 	.IsEnabled(true)
 	.Visibility(EVisibility::Visible)
@@ -701,6 +761,7 @@ void FDeviceProfileConsoleVariablesPropertyDetails::CreateRowWidgetForCVarProper
 	[
 		SNew(STextBlock)
 		.Text(FText::FromString(CVarName))
+		.ToolTipText(CVarHelp)
 		.Font(IDetailLayoutBuilder::GetDetailFont())
 	]
 	.ValueContent()
