DetailsView: Fix missing style for events button

diff --git a/Engine/Source/Editor/DetailCustomizations/Private/ActorDetails.cpp b/Engine/Source/Editor/DetailCustomizations/Private/ActorDetails.cpp
--- a/Engine/Source/Editor/DetailCustomizations/Private/ActorDetails.cpp
+++ b/Engine/Source/Editor/DetailCustomizations/Private/ActorDetails.cpp
@@ -548,7 +548,6 @@ void FActorDetails::AddEventsCategory(IDetailLayoutBuilder& DetailBuilder)
 					// A "Plus" button to add a binding. For dynamic delegates on the CDO, you can always
 					// make a new binding, so always display the "Plus"
 					SNew(SButton)
-					.ButtonStyle(FAppStyle::Get(), "DetailsView.CalloutButton")
 					.ContentPadding(FMargin(3.0, 2.0))
 					.HAlign(HAlign_Center)
 					.OnClicked(this, &FActorDetails::HandleAddOrViewEventForVariable, Blueprint, Property)
diff --git a/Engine/Source/Editor/Kismet/Private/BlueprintDetailsCustomization.cpp b/Engine/Source/Editor/Kismet/Private/BlueprintDetailsCustomization.cpp
--- a/Engine/Source/Editor/Kismet/Private/BlueprintDetailsCustomization.cpp
+++ b/Engine/Source/Editor/Kismet/Private/BlueprintDetailsCustomization.cpp
@@ -162,7 +162,6 @@ void FBlueprintDetails::AddEventsCategory(IDetailLayoutBuilder& DetailBuilder, F
 						[
 							SNew(SButton)
 							.ContentPadding(FMargin(3.0, 2.0))
-							.ButtonStyle(FAppStyle::Get(), "DetailsView.CalloutButton")
 							.OnClicked(this, &FBlueprintVarActionDetails::HandleAddOrViewEventForVariable, EventName, PropertyName, MakeWeakObjectPtr(PropertyClass))
 							[
 								SNew(SWidgetSwitcher)
diff --git a/Engine/Source/Editor/UMGEditor/Private/Customizations/UMGDetailCustomizations.cpp b/Engine/Source/Editor/UMGEditor/Private/Customizations/UMGDetailCustomizations.cpp
--- a/Engine/Source/Editor/UMGEditor/Private/Customizations/UMGDetailCustomizations.cpp
+++ b/Engine/Source/Editor/UMGEditor/Private/Customizations/UMGDetailCustomizations.cpp
@@ -723,7 +723,6 @@ void FBlueprintWidgetCustomization::CreateMulticastEventCustomization(IDetailLay
 			.Padding(0)
 			[
 				SNew(SButton)
-				.ButtonStyle(FAppStyle::Get(), "DetailsView.CalloutButton")
 				.ContentPadding(FMargin(3.0, 2.0))
 				.OnClicked(this, &FBlueprintWidgetCustomization::HandleAddOrViewEventForVariable, EventName, PropertyName, MakeWeakObjectPtr(PropertyClass))
 				[
