Make MathStructProxyCustomizations into a public part of the DetailCustomizations module.

diff --git a/Engine/Source/Editor/DetailCustomizations/Private/DetailCustomizations.cpp b/Engine/Source/Editor/DetailCustomizations/Private/DetailCustomizations.cpp
--- a/Engine/Source/Editor/DetailCustomizations/Private/DetailCustomizations.cpp
+++ b/Engine/Source/Editor/DetailCustomizations/Private/DetailCustomizations.cpp
@@ -40,6 +40,7 @@
 #include "Customizations/ColorStructCustomization.h"
 #include "Customizations/CurveTableCustomization.h"
 #include "Customizations/MathStructCustomizations.h"
+#include "Customizations/MathStructProxyCustomizations.h"
 #include "Customizations/SlateBrushCustomization.h"
 #include "Customizations/SlateFontInfoCustomization.h"
 #include "CustomPrimitiveDataCustomization.h"
@@ -85,7 +86,6 @@
 #include "MaterialInstanceDynamicDetails.h"
 #include "MaterialProxySettingsCustomizations.h"
 #include "MaterialShadingModelCustomization.h"
-#include "MathStructProxyCustomizations.h"
 #include "MeshComponentDetails.h"
 #include "MeshMergingSettingsCustomization.h"
 #include "MeshProxySettingsCustomizations.h"
diff --git a/Engine/Source/Editor/DetailCustomizations/Private/MathStructProxyCustomizations.cpp b/Engine/Source/Editor/DetailCustomizations/Private/MathStructProxyCustomizations.cpp
--- a/Engine/Source/Editor/DetailCustomizations/Private/MathStructProxyCustomizations.cpp
+++ b/Engine/Source/Editor/DetailCustomizations/Private/MathStructProxyCustomizations.cpp
@@ -1,6 +1,6 @@
 // Copyright Epic Games, Inc. All Rights Reserved.
 
-#include "MathStructProxyCustomizations.h"
+#include "Customizations/MathStructProxyCustomizations.h"
 #include "Framework/Commands/UIAction.h"
 #include "UObject/UnrealType.h"
 #include "Editor.h"
diff --git a/Engine/Source/Editor/DetailCustomizations/Private/MathStructProxyCustomizations.h b/Engine/Source/Editor/DetailCustomizations/Private/MathStructProxyCustomizations.h
deleted file mode 100644
diff --git a/Engine/Source/Editor/DetailCustomizations/Public/Customizations/MathStructProxyCustomizations.h b/Engine/Source/Editor/DetailCustomizations/Public/Customizations/MathStructProxyCustomizations.h
new file mode 100644
--- /dev/null
+++ b/Engine/Source/Editor/DetailCustomizations/Public/Customizations/MathStructProxyCustomizations.h
@@ -0,0 +1,331 @@
+// Copyright Epic Games, Inc. All Rights Reserved.
+
+#pragma once
+
+#include "CoreMinimal.h"
+#include "Styling/CoreStyle.h"
+#include "Widgets/SWidget.h"
+#include "IPropertyTypeCustomization.h"
+#include "PropertyHandle.h"
+#include "Customizations/MathStructCustomizations.h"
+
+class FDetailWidgetRow;
+class IPropertyUtilities;
+
+/** 
+ * Helper class used to track the dirty state of a proxy value
+ */
+template< typename ObjectType >
+class TProxyValue
+{
+public:
+	TProxyValue()
+		: Value()
+		, bIsSet(false)
+	{
+	}
+
+	TProxyValue(const ObjectType& InValue)
+		: Value(InValue)
+		, bIsSet(false)
+	{
+	}
+
+	/** 
+	 * Set the wrapped value
+	 * @param	InValue	The value to set.
+	 */
+	void Set(const ObjectType& InValue)
+	{
+		Value = InValue;
+		bIsSet = true;
+	}
+
+	/** 
+	 * Get the wrapped value
+	 * @return the wrapped value
+	 */
+	const ObjectType& Get() const
+	{
+		return Value;
+	}
+
+	/** 
+	 * Get the wrapped value
+	 * @return the wrapped value
+	 */
+	ObjectType& Get()
+	{
+		return Value;
+	}
+
+	/** 
+	 * Check to see if the value is set.
+	 * @return whether the value is set.
+	 */
+	bool IsSet() const
+	{
+		return bIsSet;
+	}
+
+	/** 
+	 * Mark the value as if it was set.
+	 */
+	void MarkAsSet()
+	{
+		bIsSet = true;
+	}
+
+private:
+	/** The value we are tracking */
+	ObjectType Value;
+
+	/** Whether the value is set */
+	bool bIsSet;
+};
+
+/**
+ * Helper class used to track the state of properties of proxy values.
+ */
+template< typename ObjectType, typename PropertyType >
+class TProxyProperty
+{
+public:
+	TProxyProperty(const TSharedRef< TProxyValue<ObjectType> >& InValue, PropertyType& InPropertyValue)
+		: Value(InValue)
+		, Property(InPropertyValue)
+		, bIsSet(false)
+	{
+	}
+
+	/**
+	 * Set the value of this property
+	 * @param	InPropertyValue		The value of the property to set
+	 */
+	void Set(const PropertyType& InPropertyValue)
+	{
+		Property = InPropertyValue;
+		Value->MarkAsSet();
+		bIsSet = true;
+	}
+
+	/**
+	 * Get the value of this property
+	 * @return The value of the property
+	 */
+	const PropertyType& Get() const
+	{
+		return Property;
+	}
+
+	/** 
+	 * Check to see if the value is set.
+	 * @return whether the value is set.
+	 */
+	bool IsSet() const
+	{
+		return bIsSet;
+	}
+
+private:
+	/** The proxy value we are tracking */
+	TSharedRef< TProxyValue<ObjectType> > Value;
+
+	/** The property of the value we are tracking */
+	PropertyType& Property;
+
+	/** Whether the value is set */
+	bool bIsSet;
+};
+
+/** 
+ * Helper class to aid representing math structs to the user in an editable form
+ * e.g. representing a quaternion as a set of euler angles
+ */
+class DETAILCUSTOMIZATIONS_API FMathStructProxyCustomization : public FMathStructCustomization
+{
+public:
+	/** IPropertyTypeCustomization interface */
+	virtual void CustomizeChildren( TSharedRef<class IPropertyHandle> StructPropertyHandle, class IDetailChildrenBuilder& StructBuilder, IPropertyTypeCustomizationUtils& StructCustomizationUtils ) override;
+
+	/** FMathStructCustomization interface */
+	virtual void MakeHeaderRow( TSharedRef<class IPropertyHandle>& InStructPropertyHandle, FDetailWidgetRow& Row ) override;
+
+protected:
+
+	/**
+	 * Cache the values from the property to the proxy.
+	 * @param	WeakHandlePtr	The property handle to get values from.
+	 * @return true if values(s) were successfully cached
+	 */
+	virtual bool CacheValues( TWeakPtr<IPropertyHandle> WeakHandlePtr ) const = 0;
+
+	/**
+	 * Flush the values from the proxy to the property.
+	 * @param	WeakHandlePtr	The property handle to set values to.
+	 * @return true if values(s) were successfully flushed
+	 */
+	virtual bool FlushValues( TWeakPtr<IPropertyHandle> WeakHandlePtr ) const = 0;
+
+	/** 
+	 * Helper function to make a numeric property widget to edit a proxy value.
+	 * @param	StructPropertyHandle	Property handle to the containing struct
+	 * @param	ProxyValue				The value we will be editing in the proxy data.
+	 * @param	Label					A label to use for this value.
+	 * @param	bRotationInDegrees		Whether this is to be used for a rotation value (configures sliders appropriately).
+	 * @return the newly created widget.
+	 */
+	template<typename ProxyType, typename NumericType>
+	TSharedRef<SWidget> MakeNumericProxyWidget(TSharedRef<IPropertyHandle>& StructPropertyHandle, TSharedRef< TProxyProperty<ProxyType, NumericType> >& ProxyValue, const FText& Label, bool bRotationInDegrees = false, const FLinearColor& LabelColor = FCoreStyle::Get().GetColor("DefaultForeground"));
+
+	template <typename ProxyType, typename NumericType>
+	FText OnGetValueToolTip(TWeakPtr<IPropertyHandle> WeakHandlePtr, TSharedRef<TProxyProperty<ProxyType, NumericType>> ProxyValue, FText Label) const;
+private:
+	/**
+	 * Gets the value as a float for the provided property handle
+	 *
+	 * @param WeakHandlePtr	Handle to the property to get the value from
+	 * @param ProxyValue	Proxy value to get value from.
+	 * @return The value or unset if it could not be accessed
+	 */
+	template<typename ProxyType, typename NumericType>
+	TOptional<NumericType> OnGetValue( TWeakPtr<IPropertyHandle> WeakHandlePtr, TSharedRef< TProxyProperty<ProxyType, NumericType> > ProxyValue ) const;
+
+	/**
+	 * Called when the value is committed from the property editor
+	 *
+	 * @param NewValue		The new value of the property as a float
+	 * @param CommitType	How the value was committed (unused)
+	 * @param WeakHandlePtr	Handle to the property that the new value is for
+	 */
+	template<typename ProxyType, typename NumericType>
+	void OnValueCommitted( NumericType NewValue, ETextCommit::Type CommitType, TWeakPtr<IPropertyHandle> WeakHandlePtr, TSharedRef< TProxyProperty<ProxyType, NumericType> > ProxyValue );
+	
+	/**
+	 * Called when the value is changed in the property editor
+	 *
+	 * @param NewValue		The new value of the property as a float
+	 * @param WeakHandlePtr	Handle to the property that the new value is for
+	 */
+	template<typename ProxyType, typename NumericType>
+	void OnValueChanged( NumericType NewValue, TWeakPtr<IPropertyHandle> WeakHandlePtr, TSharedRef< TProxyProperty<ProxyType, NumericType> > ProxyValue );
+
+	/** Called when a value starts to be changed by a slider */
+	void OnBeginSliderMovement();
+
+	/** Called when a value stops being changed by a slider */
+	template<typename ProxyType, typename NumericType>
+	void OnEndSliderMovement( NumericType NewValue, TWeakPtr<IPropertyHandle> WeakHandlePtr, TSharedRef< TProxyProperty<ProxyType, NumericType> > ProxyValue );
+
+protected:
+	/** Cached property utilities */
+	TSharedPtr<IPropertyUtilities> PropertyUtilities;
+};
+
+/** 
+ * Proxy struct customization that displays a matrix as a position, rotation & scale.
+ */
+class DETAILCUSTOMIZATIONS_API FMatrixStructCustomization : public FMathStructProxyCustomization
+{
+public:
+	static TSharedRef<IPropertyTypeCustomization> MakeInstance();
+
+public:
+	FMatrixStructCustomization()
+		: CachedRotation(MakeShareable( new TProxyValue<FRotator>(FRotator::ZeroRotator)))
+		, CachedRotationYaw(MakeShareable( new TProxyProperty<FRotator, FRotator::FReal>(CachedRotation, CachedRotation->Get().Yaw)))
+		, CachedRotationPitch(MakeShareable( new TProxyProperty<FRotator, FRotator::FReal>(CachedRotation, CachedRotation->Get().Pitch)))
+		, CachedRotationRoll(MakeShareable( new TProxyProperty<FRotator, FRotator::FReal>(CachedRotation, CachedRotation->Get().Roll)))
+		, CachedTranslation(MakeShareable( new TProxyValue<FVector>(FVector::ZeroVector)))
+		, CachedTranslationX(MakeShareable( new TProxyProperty<FVector, FVector::FReal>(CachedTranslation, CachedTranslation->Get().X)))
+		, CachedTranslationY(MakeShareable( new TProxyProperty<FVector, FVector::FReal>(CachedTranslation, CachedTranslation->Get().Y)))
+		, CachedTranslationZ(MakeShareable( new TProxyProperty<FVector, FVector::FReal>(CachedTranslation, CachedTranslation->Get().Z)))
+		, CachedScale(MakeShareable( new TProxyValue<FVector>(FVector::ZeroVector)))
+		, CachedScaleX(MakeShareable( new TProxyProperty<FVector, FVector::FReal>(CachedScale, CachedScale->Get().X)))
+		, CachedScaleY(MakeShareable( new TProxyProperty<FVector, FVector::FReal>(CachedScale, CachedScale->Get().Y)))
+		, CachedScaleZ(MakeShareable( new TProxyProperty<FVector, FVector::FReal>(CachedScale, CachedScale->Get().Z)))
+	{
+	}
+
+	/** IPropertyTypeCustomization interface */
+	virtual void CustomizeChildren(TSharedRef<class IPropertyHandle> StructPropertyHandle, class IDetailChildrenBuilder& StructBuilder, IPropertyTypeCustomizationUtils& StructCustomizationUtils) override;
+
+protected:
+	/** Customization utility functions */
+	void CustomizeLocation(TSharedRef<class IPropertyHandle> StructPropertyHandle, FDetailWidgetRow& Row);
+	void CustomizeRotation(TSharedRef<class IPropertyHandle> StructPropertyHandle, FDetailWidgetRow& Row);
+	void CustomizeScale(TSharedRef<class IPropertyHandle> StructPropertyHandle, FDetailWidgetRow& Row);
+
+	/** FMathStructCustomization interface */
+	virtual void MakeHeaderRow(TSharedRef<class IPropertyHandle>& InStructPropertyHandle, FDetailWidgetRow& Row) override;
+
+	/** FMathStructProxyCustomization interface */
+	virtual bool CacheValues(TWeakPtr<IPropertyHandle> PropertyHandlePtr) const override;
+	virtual bool FlushValues(TWeakPtr<IPropertyHandle> PropertyHandlePtr) const override;
+
+	struct FTransformField
+	{
+		enum Type
+		{
+			Location,
+			Rotation,
+			Scale
+		};
+	};
+
+	void OnCopy(FTransformField::Type Type, TWeakPtr<IPropertyHandle> PropertyHandlePtr);
+	void OnPaste(FTransformField::Type Type, TWeakPtr<IPropertyHandle> PropertyHandlePtr);
+
+protected:
+	/** Cached rotation values */
+	mutable TSharedRef< TProxyValue<FRotator> > CachedRotation;
+	mutable TSharedRef< TProxyProperty<FRotator, FRotator::FReal> > CachedRotationYaw;
+	mutable TSharedRef< TProxyProperty<FRotator, FRotator::FReal> > CachedRotationPitch;
+	mutable TSharedRef< TProxyProperty<FRotator, FRotator::FReal> > CachedRotationRoll;
+
+	/** Cached translation values */
+	mutable TSharedRef< TProxyValue<FVector> > CachedTranslation;
+	mutable TSharedRef< TProxyProperty<FVector, FVector::FReal> > CachedTranslationX;
+	mutable TSharedRef< TProxyProperty<FVector, FVector::FReal> > CachedTranslationY;
+	mutable TSharedRef< TProxyProperty<FVector, FVector::FReal> > CachedTranslationZ;
+
+	/** Cached scale values */
+	mutable TSharedRef< TProxyValue<FVector> > CachedScale;
+	mutable TSharedRef< TProxyProperty<FVector, FVector::FReal> > CachedScaleX;
+	mutable TSharedRef< TProxyProperty<FVector, FVector::FReal> > CachedScaleY;
+	mutable TSharedRef< TProxyProperty<FVector, FVector::FReal> > CachedScaleZ;
+};
+
+/** 
+ * Proxy struct customization that displays an FTransform as a position, euler rotation & scale.
+ */
+class FTransformStructCustomization : public FMatrixStructCustomization
+{
+public:
+	static TSharedRef<IPropertyTypeCustomization> MakeInstance();
+
+protected:
+	/** FMathStructProxyCustomization interface */
+	virtual bool CacheValues(TWeakPtr<IPropertyHandle> PropertyHandlePtr) const override;
+	virtual bool FlushValues(TWeakPtr<IPropertyHandle> PropertyHandlePtr) const override;
+};
+
+/**
+ * Proxy struct customization that displays an FQuat as an euler rotation
+ */
+class DETAILCUSTOMIZATIONS_API FQuatStructCustomization : public FMatrixStructCustomization
+{
+public:
+	static TSharedRef<IPropertyTypeCustomization> MakeInstance();
+
+	/** IPropertyTypeCustomization interface */
+	virtual void CustomizeChildren(TSharedRef<class IPropertyHandle> StructPropertyHandle, class IDetailChildrenBuilder& StructBuilder, IPropertyTypeCustomizationUtils& StructCustomizationUtils) override;
+
+	/** FMathStructCustomization interface */
+	virtual void MakeHeaderRow(TSharedRef<class IPropertyHandle>& InStructPropertyHandle, FDetailWidgetRow& Row) override;
+
+protected:
+	/** FMathStructProxyCustomization interface */
+	virtual bool CacheValues(TWeakPtr<IPropertyHandle> PropertyHandlePtr) const override;
+	virtual bool FlushValues(TWeakPtr<IPropertyHandle> PropertyHandlePtr) const override;
+};
