Fix float math struct proxy customizations (FTransform, FMatrix, FQuat)

diff --git a/Engine/Source/Editor/DetailCustomizations/Private/DetailCustomizations.cpp b/Engine/Source/Editor/DetailCustomizations/Private/DetailCustomizations.cpp
--- a/Engine/Source/Editor/DetailCustomizations/Private/DetailCustomizations.cpp
+++ b/Engine/Source/Editor/DetailCustomizations/Private/DetailCustomizations.cpp
@@ -222,15 +222,15 @@ void FDetailCustomizationsModule::RegisterPropertyTypeCustomizations()
 	REGISTER_UIMINMAX_CUSTOMIZATION(NAME_Rotator3d, &FRotatorStructCustomization::MakeInstance);
 	RegisterCustomPropertyTypeLayout(NAME_LinearColor, FOnGetPropertyTypeCustomizationInstance::CreateStatic(&FColorStructCustomization::MakeInstance));
 	RegisterCustomPropertyTypeLayout(NAME_Color, FOnGetPropertyTypeCustomizationInstance::CreateStatic(&FColorStructCustomization::MakeInstance));
-	RegisterCustomPropertyTypeLayout(NAME_Matrix, FOnGetPropertyTypeCustomizationInstance::CreateStatic(&FMatrixStructCustomization::MakeInstance));
-	RegisterCustomPropertyTypeLayout(NAME_Matrix44f, FOnGetPropertyTypeCustomizationInstance::CreateStatic(&FMatrixStructCustomization::MakeInstance));
-	RegisterCustomPropertyTypeLayout(NAME_Matrix44d, FOnGetPropertyTypeCustomizationInstance::CreateStatic(&FMatrixStructCustomization::MakeInstance));
-	RegisterCustomPropertyTypeLayout(NAME_Transform, FOnGetPropertyTypeCustomizationInstance::CreateStatic(&FTransformStructCustomization::MakeInstance));
-	RegisterCustomPropertyTypeLayout(NAME_Transform3f, FOnGetPropertyTypeCustomizationInstance::CreateStatic(&FTransformStructCustomization::MakeInstance));
-	RegisterCustomPropertyTypeLayout(NAME_Transform3d, FOnGetPropertyTypeCustomizationInstance::CreateStatic(&FTransformStructCustomization::MakeInstance));
-	RegisterCustomPropertyTypeLayout(NAME_Quat, FOnGetPropertyTypeCustomizationInstance::CreateStatic(&FQuatStructCustomization::MakeInstance));
-	RegisterCustomPropertyTypeLayout(NAME_Quat4f, FOnGetPropertyTypeCustomizationInstance::CreateStatic(&FQuatStructCustomization::MakeInstance));
-	RegisterCustomPropertyTypeLayout(NAME_Quat4d, FOnGetPropertyTypeCustomizationInstance::CreateStatic(&FQuatStructCustomization::MakeInstance));
+	RegisterCustomPropertyTypeLayout(NAME_Matrix, FOnGetPropertyTypeCustomizationInstance::CreateStatic(&FMatrixStructCustomization<FMatrix::FReal>::MakeInstance));
+	RegisterCustomPropertyTypeLayout(NAME_Matrix44f, FOnGetPropertyTypeCustomizationInstance::CreateStatic(&FMatrixStructCustomization<FMatrix44f::FReal>::MakeInstance));
+	RegisterCustomPropertyTypeLayout(NAME_Matrix44d, FOnGetPropertyTypeCustomizationInstance::CreateStatic(&FMatrixStructCustomization<FMatrix44d::FReal>::MakeInstance));
+	RegisterCustomPropertyTypeLayout(NAME_Transform, FOnGetPropertyTypeCustomizationInstance::CreateStatic(&FTransformStructCustomization<FTransform::FReal>::MakeInstance));
+	RegisterCustomPropertyTypeLayout(NAME_Transform3f, FOnGetPropertyTypeCustomizationInstance::CreateStatic(&FTransformStructCustomization<FTransform3f::FReal>::MakeInstance));
+	RegisterCustomPropertyTypeLayout(NAME_Transform3d, FOnGetPropertyTypeCustomizationInstance::CreateStatic(&FTransformStructCustomization<FTransform3d::FReal>::MakeInstance));
+	RegisterCustomPropertyTypeLayout(NAME_Quat, FOnGetPropertyTypeCustomizationInstance::CreateStatic(&FQuatStructCustomization<FQuat::FReal>::MakeInstance));
+	RegisterCustomPropertyTypeLayout(NAME_Quat4f, FOnGetPropertyTypeCustomizationInstance::CreateStatic(&FQuatStructCustomization<FQuat4f::FReal>::MakeInstance));
+	RegisterCustomPropertyTypeLayout(NAME_Quat4d, FOnGetPropertyTypeCustomizationInstance::CreateStatic(&FQuatStructCustomization<FQuat4d::FReal>::MakeInstance));
 	RegisterCustomPropertyTypeLayout("SlateColor", FOnGetPropertyTypeCustomizationInstance::CreateStatic(&FSlateColorCustomization::MakeInstance));
 	RegisterCustomPropertyTypeLayout("ForceFeedbackAttenuationSettings", FOnGetPropertyTypeCustomizationInstance::CreateStatic(&FForceFeedbackAttenuationSettingsCustomization::MakeInstance));
 	RegisterCustomPropertyTypeLayout("SoundAttenuationSettings", FOnGetPropertyTypeCustomizationInstance::CreateStatic(&FSoundAttenuationSettingsCustomization::MakeInstance));
diff --git a/Engine/Source/Editor/DetailCustomizations/Private/MathStructProxyCustomizations.cpp b/Engine/Source/Editor/DetailCustomizations/Private/MathStructProxyCustomizations.cpp
--- a/Engine/Source/Editor/DetailCustomizations/Private/MathStructProxyCustomizations.cpp
+++ b/Engine/Source/Editor/DetailCustomizations/Private/MathStructProxyCustomizations.cpp
@@ -111,12 +111,14 @@ FText FMathStructProxyCustomization::OnGetValueToolTip(TWeakPtr<IPropertyHandle>
 	return FText::GetEmpty();
 }
 
-TSharedRef<IPropertyTypeCustomization> FMatrixStructCustomization::MakeInstance()
+template<typename T>
+TSharedRef<IPropertyTypeCustomization> FMatrixStructCustomization<T>::MakeInstance()
 {
-	return MakeShareable( new FMatrixStructCustomization );
+	return MakeShareable( new FMatrixStructCustomization<T> );
 }
 
-void FMatrixStructCustomization::MakeHeaderRow(TSharedRef<class IPropertyHandle>& StructPropertyHandle, FDetailWidgetRow& Row)
+template<typename T>
+void FMatrixStructCustomization<T>::MakeHeaderRow(TSharedRef<class IPropertyHandle>& StructPropertyHandle, FDetailWidgetRow& Row)
 {
 	Row
 	.NameContent()
@@ -131,13 +133,14 @@ void FMatrixStructCustomization::MakeHeaderRow(TSharedRef<class IPropertyHandle>
 	];
 }
 
-void FMatrixStructCustomization::CustomizeLocation(TSharedRef<class IPropertyHandle> StructPropertyHandle, FDetailWidgetRow& Row)
+template<typename T>
+void FMatrixStructCustomization<T>::CustomizeLocation(TSharedRef<class IPropertyHandle> StructPropertyHandle, FDetailWidgetRow& Row)
 {
 	TWeakPtr<IPropertyHandle> WeakHandlePtr = StructPropertyHandle;
 
 	Row
-	.CopyAction(FUIAction(FExecuteAction::CreateSP(this, &FMatrixStructCustomization::OnCopy, FTransformField::Location, WeakHandlePtr)))
-	.PasteAction(FUIAction(FExecuteAction::CreateSP(this, &FMatrixStructCustomization::OnPaste, FTransformField::Location, WeakHandlePtr)))
+	.CopyAction(FUIAction(FExecuteAction::CreateSP(this, &FMatrixStructCustomization<T>::OnCopy, FTransformField::Location, WeakHandlePtr)))
+	.PasteAction(FUIAction(FExecuteAction::CreateSP(this, &FMatrixStructCustomization<T>::OnPaste, FTransformField::Location, WeakHandlePtr)))
 	.NameContent()
 	[
 		StructPropertyHandle->CreatePropertyNameWidget(LOCTEXT("LocationLabel", "Location"))
@@ -150,29 +153,29 @@ void FMatrixStructCustomization::CustomizeLocation(TSharedRef<class IPropertyHan
 		+ SHorizontalBox::Slot()
 		.Padding(FMargin(0.0f, 2.0f, 3.0f, 2.0f))
 		[
-			MakeNumericProxyWidget<FVector, FVector::FReal>(StructPropertyHandle, CachedTranslationX, LOCTEXT("TranslationX", "X"), false, SNumericEntryBox<FVector::FReal>::RedLabelBackgroundColor)
+			MakeNumericProxyWidget<UE::Math::TVector<T>, T>(StructPropertyHandle, CachedTranslationX, LOCTEXT("TranslationX", "X"), false, SNumericEntryBox<T>::RedLabelBackgroundColor)
 		]
 		+ SHorizontalBox::Slot()
 		.Padding(FMargin(0.0f, 2.0f, 3.0f, 2.0f))
 		[
-			MakeNumericProxyWidget<FVector, FVector::FReal>(StructPropertyHandle, CachedTranslationY, LOCTEXT("TranslationY", "Y"), false, SNumericEntryBox<FVector::FReal>::GreenLabelBackgroundColor)
+			MakeNumericProxyWidget<UE::Math::TVector<T>, T>(StructPropertyHandle, CachedTranslationY, LOCTEXT("TranslationY", "Y"), false, SNumericEntryBox<T>::GreenLabelBackgroundColor)
 		]
 		+ SHorizontalBox::Slot()
 		.Padding(FMargin(0.0f, 2.0f, 0.0f, 2.0f))
 		[
-			MakeNumericProxyWidget<FVector, FVector::FReal>(StructPropertyHandle, CachedTranslationZ, LOCTEXT("TranslationZ", "Z"), false, SNumericEntryBox<FVector::FReal>::BlueLabelBackgroundColor)
+			MakeNumericProxyWidget<UE::Math::TVector<T>, T>(StructPropertyHandle, CachedTranslationZ, LOCTEXT("TranslationZ", "Z"), false, SNumericEntryBox<T>::BlueLabelBackgroundColor)
 		]
 	];
 }
 
-
-void FMatrixStructCustomization::CustomizeRotation(TSharedRef<class IPropertyHandle> StructPropertyHandle, FDetailWidgetRow& Row)
+template<typename T>
+void FMatrixStructCustomization<T>::CustomizeRotation(TSharedRef<class IPropertyHandle> StructPropertyHandle, FDetailWidgetRow& Row)
 {
 	TWeakPtr<IPropertyHandle> WeakHandlePtr = StructPropertyHandle;
 
 	Row
-	.CopyAction(FUIAction(FExecuteAction::CreateSP(this, &FMatrixStructCustomization::OnCopy, FTransformField::Rotation, WeakHandlePtr)))
-	.PasteAction(FUIAction(FExecuteAction::CreateSP(this, &FMatrixStructCustomization::OnPaste, FTransformField::Rotation, WeakHandlePtr)))
+	.CopyAction(FUIAction(FExecuteAction::CreateSP(this, &FMatrixStructCustomization<T>::OnCopy, FTransformField::Rotation, WeakHandlePtr)))
+	.PasteAction(FUIAction(FExecuteAction::CreateSP(this, &FMatrixStructCustomization<T>::OnPaste, FTransformField::Rotation, WeakHandlePtr)))
 	.NameContent()
 	[
 		StructPropertyHandle->CreatePropertyNameWidget(LOCTEXT("RotationLabel", "Rotation"))
@@ -185,29 +188,29 @@ void FMatrixStructCustomization::CustomizeRotation(TSharedRef<class IPropertyHan
 		+ SHorizontalBox::Slot()
 		.Padding(FMargin(0.0f, 2.0f, 3.0f, 2.0f))
 		[
-			MakeNumericProxyWidget<FRotator, FRotator::FReal>(StructPropertyHandle, CachedRotationRoll, LOCTEXT("RotationRoll", "X"), true, SNumericEntryBox<FRotator::FReal>::RedLabelBackgroundColor)
+			MakeNumericProxyWidget<UE::Math::TRotator<T>, T>(StructPropertyHandle, CachedRotationRoll, LOCTEXT("RotationRoll", "X"), true, SNumericEntryBox<T>::RedLabelBackgroundColor)
 		]
 		+ SHorizontalBox::Slot()
 		.Padding(FMargin(0.0f, 2.0f, 3.0f, 2.0f))
 		[
-			MakeNumericProxyWidget<FRotator, FRotator::FReal>(StructPropertyHandle, CachedRotationPitch, LOCTEXT("RotationPitch", "Y"), true, SNumericEntryBox<FRotator::FReal>::GreenLabelBackgroundColor)
+			MakeNumericProxyWidget<UE::Math::TRotator<T>, T>(StructPropertyHandle, CachedRotationPitch, LOCTEXT("RotationPitch", "Y"), true, SNumericEntryBox<T>::GreenLabelBackgroundColor)
 		]
 		+ SHorizontalBox::Slot()
 		.Padding(FMargin(0.0f, 2.0f, 0.0f, 2.0f))
 		[
-			MakeNumericProxyWidget<FRotator, FRotator::FReal>(StructPropertyHandle, CachedRotationYaw, LOCTEXT("RotationYaw", "Z"), true, SNumericEntryBox<FRotator::FReal>::BlueLabelBackgroundColor)
+			MakeNumericProxyWidget<UE::Math::TRotator<T>, T>(StructPropertyHandle, CachedRotationYaw, LOCTEXT("RotationYaw", "Z"), true, SNumericEntryBox<T>::BlueLabelBackgroundColor)
 		]
 	];
 }
 
-
-void FMatrixStructCustomization::CustomizeScale(TSharedRef<class IPropertyHandle> StructPropertyHandle, FDetailWidgetRow& Row)
+template<typename T>
+void FMatrixStructCustomization<T>::CustomizeScale(TSharedRef<class IPropertyHandle> StructPropertyHandle, FDetailWidgetRow& Row)
 {
 	TWeakPtr<IPropertyHandle> WeakHandlePtr = StructPropertyHandle;
 
 	Row
-	.CopyAction(FUIAction(FExecuteAction::CreateSP(this, &FMatrixStructCustomization::OnCopy, FTransformField::Scale, WeakHandlePtr)))
-	.PasteAction(FUIAction(FExecuteAction::CreateSP(this, &FMatrixStructCustomization::OnPaste, FTransformField::Scale, WeakHandlePtr)))
+	.CopyAction(FUIAction(FExecuteAction::CreateSP(this, &FMatrixStructCustomization<T>::OnCopy, FTransformField::Scale, WeakHandlePtr)))
+	.PasteAction(FUIAction(FExecuteAction::CreateSP(this, &FMatrixStructCustomization<T>::OnPaste, FTransformField::Scale, WeakHandlePtr)))
 	.NameContent()
 	[
 		StructPropertyHandle->CreatePropertyNameWidget(LOCTEXT("ScaleLabel", "Scale"))
@@ -220,22 +223,23 @@ void FMatrixStructCustomization::CustomizeScale(TSharedRef<class IPropertyHandle
 		+ SHorizontalBox::Slot()
 		.Padding(FMargin(0.0f, 2.0f, 3.0f, 2.0f))
 		[
-			MakeNumericProxyWidget<FVector, FVector::FReal>(StructPropertyHandle, CachedScaleX, LOCTEXT("ScaleX", "X"), false, SNumericEntryBox<FVector::FReal>::RedLabelBackgroundColor)
+			MakeNumericProxyWidget<UE::Math::TVector<T>, T>(StructPropertyHandle, CachedScaleX, LOCTEXT("ScaleX", "X"), false, SNumericEntryBox<T>::RedLabelBackgroundColor)
 		]
 		+ SHorizontalBox::Slot()
 		.Padding(FMargin(0.0f, 2.0f, 3.0f, 2.0f))
 		[
-			MakeNumericProxyWidget<FVector, FVector::FReal>(StructPropertyHandle, CachedScaleY, LOCTEXT("ScaleY", "Y"), false, SNumericEntryBox<FVector::FReal>::GreenLabelBackgroundColor)
+			MakeNumericProxyWidget<UE::Math::TVector<T>, T>(StructPropertyHandle, CachedScaleY, LOCTEXT("ScaleY", "Y"), false, SNumericEntryBox<T>::GreenLabelBackgroundColor)
 		]
 		+ SHorizontalBox::Slot()
 		.Padding(FMargin(0.0f, 2.0f, 0.0f, 2.0f))
 		[
-			MakeNumericProxyWidget<FVector, FVector::FReal>(StructPropertyHandle, CachedScaleZ, LOCTEXT("ScaleZ", "Z"), false, SNumericEntryBox<FVector::FReal>::BlueLabelBackgroundColor)
+			MakeNumericProxyWidget<UE::Math::TVector<T>, T>(StructPropertyHandle, CachedScaleZ, LOCTEXT("ScaleZ", "Z"), false, SNumericEntryBox<T>::BlueLabelBackgroundColor)
 		]
 	];
 }
 
-void FMatrixStructCustomization::CustomizeChildren(TSharedRef<class IPropertyHandle> StructPropertyHandle, class IDetailChildrenBuilder& StructBuilder, IPropertyTypeCustomizationUtils& StructCustomizationUtils)
+template<typename T>
+void FMatrixStructCustomization<T>::CustomizeChildren(TSharedRef<class IPropertyHandle> StructPropertyHandle, class IDetailChildrenBuilder& StructBuilder, IPropertyTypeCustomizationUtils& StructCustomizationUtils)
 {
 	FMathStructProxyCustomization::CustomizeChildren(StructPropertyHandle, StructBuilder, StructCustomizationUtils);
 
@@ -246,7 +250,8 @@ void FMatrixStructCustomization::CustomizeChildren(TSharedRef<class IPropertyHan
 	CustomizeScale(StructPropertyHandle, StructBuilder.AddCustomRow(LOCTEXT("ScaleLabel", "Scale")));
 }
 
-void FMatrixStructCustomization::OnCopy(FTransformField::Type Type, TWeakPtr<IPropertyHandle> PropertyHandlePtr)
+template<typename T>
+void FMatrixStructCustomization<T>::OnCopy(FTransformField::Type Type, TWeakPtr<IPropertyHandle> PropertyHandlePtr)
 {
 	auto PropertyHandle = PropertyHandlePtr.Pin();
 
@@ -262,21 +267,21 @@ void FMatrixStructCustomization::OnCopy(FTransformField::Type Type, TWeakPtr<IPr
 	{
 		case FTransformField::Location:
 		{
-			FVector Location = CachedTranslation->Get();
+			UE::Math::TVector<T> Location = CachedTranslation->Get();
 			CopyStr = FString::Printf(TEXT("(X=%f,Y=%f,Z=%f)"), Location.X, Location.Y, Location.Z);
 			break;
 		}
 
 		case FTransformField::Rotation:
 		{
-			FRotator Rotation = CachedRotation->Get();
+			UE::Math::TRotator<T> Rotation = CachedRotation->Get();
 			CopyStr = FString::Printf(TEXT("(Pitch=%f,Yaw=%f,Roll=%f)"), Rotation.Pitch, Rotation.Yaw, Rotation.Roll);
 			break;
 		}
 
 		case FTransformField::Scale:
 		{
-			FVector Scale = CachedScale->Get();
+			UE::Math::TVector<T> Scale = CachedScale->Get();
 			CopyStr = FString::Printf(TEXT("(X=%f,Y=%f,Z=%f)"), Scale.X, Scale.Y, Scale.Z);
 			break;
 		}
@@ -288,7 +293,8 @@ void FMatrixStructCustomization::OnCopy(FTransformField::Type Type, TWeakPtr<IPr
 	}
 }
 
-void FMatrixStructCustomization::OnPaste(FTransformField::Type Type, TWeakPtr<IPropertyHandle> PropertyHandlePtr)
+template<typename T>
+void FMatrixStructCustomization<T>::OnPaste(FTransformField::Type Type, TWeakPtr<IPropertyHandle> PropertyHandlePtr)
 {
 	auto PropertyHandle = PropertyHandlePtr.Pin();
 
@@ -304,7 +310,7 @@ void FMatrixStructCustomization::OnPaste(FTransformField::Type Type, TWeakPtr<IP
 	{
 		case FTransformField::Location:
 		{
-			FVector Location;
+			UE::Math::TVector<T> Location;
 			if (Location.InitFromString(PastedText))
 			{
 				FScopedTransaction Transaction(LOCTEXT("PasteLocation", "Paste Location"));
@@ -318,7 +324,7 @@ void FMatrixStructCustomization::OnPaste(FTransformField::Type Type, TWeakPtr<IP
 
 		case FTransformField::Rotation:
 		{
-			FRotator Rotation;
+			UE::Math::TRotator<T> Rotation;
 			PastedText.ReplaceInline(TEXT("Pitch="), TEXT("P="));
 			PastedText.ReplaceInline(TEXT("Yaw="), TEXT("Y="));
 			PastedText.ReplaceInline(TEXT("Roll="), TEXT("R="));
@@ -335,7 +341,7 @@ void FMatrixStructCustomization::OnPaste(FTransformField::Type Type, TWeakPtr<IP
 
 		case FTransformField::Scale:
 		{
-			FVector Scale;
+			UE::Math::TVector<T> Scale;
 			if (Scale.InitFromString(PastedText))
 			{
 				FScopedTransaction Transaction(LOCTEXT("PasteScale", "Paste Scale"));
@@ -349,7 +355,8 @@ void FMatrixStructCustomization::OnPaste(FTransformField::Type Type, TWeakPtr<IP
 	}
 }
 
-bool FMatrixStructCustomization::CacheValues( TWeakPtr<IPropertyHandle> PropertyHandlePtr ) const
+template<typename T>
+bool FMatrixStructCustomization<T>::CacheValues( TWeakPtr<IPropertyHandle> PropertyHandlePtr ) const
 {
 	auto PropertyHandle = PropertyHandlePtr.Pin();
 
@@ -361,10 +368,10 @@ bool FMatrixStructCustomization::CacheValues( TWeakPtr<IPropertyHandle> Property
 	TArray<void*> RawData;
 	PropertyHandle->AccessRawData(RawData);
 
-	const FMatrix* FirstMatrixValue = nullptr;
+	const UE::Math::TMatrix<T>* FirstMatrixValue = nullptr;
 	for(void* RawDataPtr : RawData)
 	{
-		FMatrix* MatrixValue = reinterpret_cast<FMatrix*>(RawDataPtr);
+		UE::Math::TMatrix<T>* MatrixValue = reinterpret_cast<UE::Math::TMatrix<T>*>(RawDataPtr);
 		if (MatrixValue == nullptr)
 		{
 			return false;
@@ -394,7 +401,8 @@ bool FMatrixStructCustomization::CacheValues( TWeakPtr<IPropertyHandle> Property
 	return false;
 }
 
-bool FMatrixStructCustomization::FlushValues( TWeakPtr<IPropertyHandle> PropertyHandlePtr ) const
+template<typename T>
+bool FMatrixStructCustomization<T>::FlushValues( TWeakPtr<IPropertyHandle> PropertyHandlePtr ) const
 {
 	auto PropertyHandle = PropertyHandlePtr.Pin();
 	if (!PropertyHandle.IsValid())
@@ -417,31 +425,31 @@ bool FMatrixStructCustomization::FlushValues( TWeakPtr<IPropertyHandle> Property
 	bool bNotifiedPreChange = false;
 	for (int32 ValueIndex = 0; ValueIndex < RawData.Num(); ValueIndex++)
 	{
-		FMatrix* MatrixValue = reinterpret_cast<FMatrix*>(RawData[ValueIndex]);
+		UE::Math::TMatrix<T>* MatrixValue = reinterpret_cast<UE::Math::TMatrix<T>*>(RawData[ValueIndex]);
 		if (MatrixValue != NULL)
 		{
-			const FMatrix PreviousValue = *MatrixValue;
-			const FRotator CurrentRotation = MatrixValue->Rotator();
-			const FVector CurrentTranslation = MatrixValue->GetOrigin();
-			const FVector CurrentScale = MatrixValue->GetScaleVector();
+			const UE::Math::TMatrix<T> PreviousValue = *MatrixValue;
+			const UE::Math::TRotator<T> CurrentRotation = MatrixValue->Rotator();
+			const UE::Math::TVector<T> CurrentTranslation = MatrixValue->GetOrigin();
+			const UE::Math::TVector<T> CurrentScale = MatrixValue->GetScaleVector();
 
-			FRotator Rotation(
+			UE::Math::TRotator<T> Rotation(
 				CachedRotationPitch->IsSet() ? CachedRotationPitch->Get() : CurrentRotation.Pitch,
 				CachedRotationYaw->IsSet() ? CachedRotationYaw->Get() : CurrentRotation.Yaw,
 				CachedRotationRoll->IsSet() ? CachedRotationRoll->Get() : CurrentRotation.Roll
 				);
-			FVector Translation(
+			UE::Math::TVector<T> Translation(
 				CachedTranslationX->IsSet() ? CachedTranslationX->Get() : CurrentTranslation.X,
 				CachedTranslationY->IsSet() ? CachedTranslationY->Get() : CurrentTranslation.Y,
 				CachedTranslationZ->IsSet() ? CachedTranslationZ->Get() : CurrentTranslation.Z
 				);
-			FVector Scale(
+			UE::Math::TVector<T> Scale(
 				CachedScaleX->IsSet() ? CachedScaleX->Get() : CurrentScale.X,
 				CachedScaleY->IsSet() ? CachedScaleY->Get() : CurrentScale.Y,
 				CachedScaleZ->IsSet() ? CachedScaleZ->Get() : CurrentScale.Z
 				);
 
-			const FMatrix NewValue = FScaleRotationTranslationMatrix(Scale, Rotation, Translation);
+			const UE::Math::TMatrix<T> NewValue = UE::Math::TScaleRotationTranslationMatrix<T>(Scale, Rotation, Translation);
 
 			if (!bNotifiedPreChange && (!MatrixValue->Equals(NewValue, 0.0f) || (!bIsUsingSlider && bIsInteractiveChangeInProgress)))
 			{
@@ -461,14 +469,14 @@ bool FMatrixStructCustomization::FlushValues( TWeakPtr<IPropertyHandle> Property
 
 			// Propagate default value changes after updating, for archetypes. As per usual, we only propagate the change if the instance matches the archetype's value.
 			// Note: We cannot use the "normal" PropertyNode propagation logic here, because that is string-based and the decision to propagate relies on an exact value match.
-			// Here, we're dealing with conversions between FMatrix and FVector/FRotator values, so there is some precision loss that requires a tolerance when comparing values.
+			// Here, we're dealing with conversions between UE::Math::TMatrix<T> and UE::Math::TVector<T>/UE::Math::TRotator<T> values, so there is some precision loss that requires a tolerance when comparing values.
 			if (ValueIndex < OuterObjects.Num() && OuterObjects[ValueIndex]->IsTemplate())
 			{
 				TArray<UObject*> ArchetypeInstances;
 				OuterObjects[ValueIndex]->GetArchetypeInstances(ArchetypeInstances);
 				for (UObject* ArchetypeInstance : ArchetypeInstances)
 				{
-					FMatrix* CurrentValue = reinterpret_cast<FMatrix*>(PropertyHandle->GetValueBaseAddress(reinterpret_cast<uint8*>(ArchetypeInstance)));
+					UE::Math::TMatrix<T>* CurrentValue = reinterpret_cast<UE::Math::TMatrix<T>*>(PropertyHandle->GetValueBaseAddress(reinterpret_cast<uint8*>(ArchetypeInstance)));
 					if (CurrentValue && CurrentValue->Equals(PreviousValue))
 					{
 						*CurrentValue = NewValue;
@@ -498,12 +506,14 @@ bool FMatrixStructCustomization::FlushValues( TWeakPtr<IPropertyHandle> Property
 	return true;
 }
 
-TSharedRef<IPropertyTypeCustomization> FTransformStructCustomization::MakeInstance() 
+template<typename T>
+TSharedRef<IPropertyTypeCustomization> FTransformStructCustomization<T>::MakeInstance() 
 {
 	return MakeShareable( new FTransformStructCustomization );
 }
 
-bool FTransformStructCustomization::CacheValues( TWeakPtr<IPropertyHandle> PropertyHandlePtr ) const
+template<typename T>
+bool FTransformStructCustomization<T>::CacheValues( TWeakPtr<IPropertyHandle> PropertyHandlePtr ) const
 {
 	auto PropertyHandle = PropertyHandlePtr.Pin();
 
@@ -515,10 +525,10 @@ bool FTransformStructCustomization::CacheValues( TWeakPtr<IPropertyHandle> Prope
 	TArray<void*> RawData;
 	PropertyHandle->AccessRawData(RawData);
 
-	const FTransform* FirstTransformValue = nullptr;
+	const UE::Math::TTransform<T>* FirstTransformValue = nullptr;
 	for(void* RawDataPtr : RawData)
 	{
-		FTransform* TransformValue = reinterpret_cast<FTransform*>(RawDataPtr);
+		UE::Math::TTransform<T>* TransformValue = reinterpret_cast<UE::Math::TTransform<T>*>(RawDataPtr);
 		if (TransformValue == nullptr)
 		{
 			return false;
@@ -539,16 +549,17 @@ bool FTransformStructCustomization::CacheValues( TWeakPtr<IPropertyHandle> Prope
 
 	if(FirstTransformValue)
 	{
-		CachedTranslation->Set(FirstTransformValue->GetTranslation());
-		CachedRotation->Set(FirstTransformValue->GetRotation().Rotator());
-		CachedScale->Set(FirstTransformValue->GetScale3D());
+		this->CachedTranslation->Set(FirstTransformValue->GetTranslation());
+		this->CachedRotation->Set(FirstTransformValue->GetRotation().Rotator());
+		this->CachedScale->Set(FirstTransformValue->GetScale3D());
 		return true;
 	}
 
 	return false;
 }
 
-bool FTransformStructCustomization::FlushValues( TWeakPtr<IPropertyHandle> PropertyHandlePtr ) const
+template<typename T>
+bool FTransformStructCustomization<T>::FlushValues( TWeakPtr<IPropertyHandle> PropertyHandlePtr ) const
 {
 	auto PropertyHandle = PropertyHandlePtr.Pin();
 
@@ -572,33 +583,33 @@ bool FTransformStructCustomization::FlushValues( TWeakPtr<IPropertyHandle> Prope
 	bool bNotifiedPreChange = false;
 	for (int32 ValueIndex = 0; ValueIndex < RawData.Num(); ValueIndex++)
 	{
-		FTransform* TransformValue = reinterpret_cast<FTransform*>(RawData[0]);
+		UE::Math::TTransform<T>* TransformValue = reinterpret_cast<UE::Math::TTransform<T>*>(RawData[0]);
 		if (TransformValue != NULL)
 		{
-			const FTransform PreviousValue = *TransformValue;
-			const FRotator CurrentRotation = TransformValue->GetRotation().Rotator();
-			const FVector CurrentTranslation = TransformValue->GetTranslation();
-			const FVector CurrentScale = TransformValue->GetScale3D();
-
-			FRotator Rotation(
-				CachedRotationPitch->IsSet() ? CachedRotationPitch->Get() : CurrentRotation.Pitch,
-				CachedRotationYaw->IsSet() ? CachedRotationYaw->Get() : CurrentRotation.Yaw,
-				CachedRotationRoll->IsSet() ? CachedRotationRoll->Get() : CurrentRotation.Roll
+			const UE::Math::TTransform<T> PreviousValue = *TransformValue;
+			const UE::Math::TRotator<T> CurrentRotation = TransformValue->GetRotation().Rotator();
+			const UE::Math::TVector<T> CurrentTranslation = TransformValue->GetTranslation();
+			const UE::Math::TVector<T> CurrentScale = TransformValue->GetScale3D();
+
+			UE::Math::TRotator<T> Rotation(
+				this->CachedRotationPitch->IsSet() ? this->CachedRotationPitch->Get() : CurrentRotation.Pitch,
+				this->CachedRotationYaw->IsSet() ? this->CachedRotationYaw->Get() : CurrentRotation.Yaw,
+				this->CachedRotationRoll->IsSet() ? this->CachedRotationRoll->Get() : CurrentRotation.Roll
 				);
-			FVector Translation(
-				CachedTranslationX->IsSet() ? CachedTranslationX->Get() : CurrentTranslation.X,
-				CachedTranslationY->IsSet() ? CachedTranslationY->Get() : CurrentTranslation.Y,
-				CachedTranslationZ->IsSet() ? CachedTranslationZ->Get() : CurrentTranslation.Z
+			UE::Math::TVector<T> Translation(
+				this->CachedTranslationX->IsSet() ? this->CachedTranslationX->Get() : CurrentTranslation.X,
+				this->CachedTranslationY->IsSet() ? this->CachedTranslationY->Get() : CurrentTranslation.Y,
+				this->CachedTranslationZ->IsSet() ? this->CachedTranslationZ->Get() : CurrentTranslation.Z
 				);
-			FVector Scale(
-				CachedScaleX->IsSet() ? CachedScaleX->Get() : CurrentScale.X,
-				CachedScaleY->IsSet() ? CachedScaleY->Get() : CurrentScale.Y,
-				CachedScaleZ->IsSet() ? CachedScaleZ->Get() : CurrentScale.Z
+			UE::Math::TVector<T> Scale(
+				this->CachedScaleX->IsSet() ? this->CachedScaleX->Get() : CurrentScale.X,
+				this->CachedScaleY->IsSet() ? this->CachedScaleY->Get() : CurrentScale.Y,
+				this->CachedScaleZ->IsSet() ? this->CachedScaleZ->Get() : CurrentScale.Z
 				);
 
-			const FTransform NewValue = FTransform(Rotation, Translation, Scale);
+			const UE::Math::TTransform<T> NewValue = UE::Math::TTransform<T>(Rotation, Translation, Scale);
 
-			if (!bNotifiedPreChange && (!TransformValue->Equals(NewValue, 0.0f) || (!bIsUsingSlider && bIsInteractiveChangeInProgress)))
+			if (!bNotifiedPreChange && (!TransformValue->Equals(NewValue, 0.0f) || (!this->bIsUsingSlider && bIsInteractiveChangeInProgress)))
 			{
 				if (!bIsInteractiveChangeInProgress)
 				{
@@ -608,7 +619,7 @@ bool FTransformStructCustomization::FlushValues( TWeakPtr<IPropertyHandle> Prope
 				PropertyHandle->NotifyPreChange();
 				bNotifiedPreChange = true;
 
-				bIsInteractiveChangeInProgress = bIsUsingSlider;
+				bIsInteractiveChangeInProgress = this->bIsUsingSlider;
 			}
 
 			// Set the new value.
@@ -616,14 +627,14 @@ bool FTransformStructCustomization::FlushValues( TWeakPtr<IPropertyHandle> Prope
 
 			// Propagate default value changes after updating, for archetypes. As per usual, we only propagate the change if the instance matches the archetype's value.
 			// Note: We cannot use the "normal" PropertyNode propagation logic here, because that is string-based and the decision to propagate relies on an exact value match.
-			// Here, we're dealing with conversions between FTransform and FVector/FRotator values, so there is some precision loss that requires a tolerance when comparing values.
+			// Here, we're dealing with conversions between UE::Math::TTransform<T> and UE::Math::TVector<T>/UE::Math::TRotator<T> values, so there is some precision loss that requires a tolerance when comparing values.
 			if (ValueIndex < OuterObjects.Num() && OuterObjects[ValueIndex]->IsTemplate())
 			{
 				TArray<UObject*> ArchetypeInstances;
 				OuterObjects[ValueIndex]->GetArchetypeInstances(ArchetypeInstances);
 				for (UObject* ArchetypeInstance : ArchetypeInstances)
 				{
-					FTransform* CurrentValue = reinterpret_cast<FTransform*>(PropertyHandle->GetValueBaseAddress(reinterpret_cast<uint8*>(ArchetypeInstance)));
+					UE::Math::TTransform<T>* CurrentValue = reinterpret_cast<UE::Math::TTransform<T>*>(PropertyHandle->GetValueBaseAddress(reinterpret_cast<uint8*>(ArchetypeInstance)));
 					if (CurrentValue && CurrentValue->Equals(PreviousValue))
 					{
 						*CurrentValue = NewValue;
@@ -635,44 +646,45 @@ bool FTransformStructCustomization::FlushValues( TWeakPtr<IPropertyHandle> Prope
 	
 	if (bNotifiedPreChange)
 	{
-		PropertyHandle->NotifyPostChange(bIsUsingSlider ? EPropertyChangeType::Interactive : EPropertyChangeType::ValueSet);
+		PropertyHandle->NotifyPostChange(this->bIsUsingSlider ? EPropertyChangeType::Interactive : EPropertyChangeType::ValueSet);
 
-		if (!bIsUsingSlider)
+		if (!this->bIsUsingSlider)
 		{
 			GEditor->EndTransaction();
 			bIsInteractiveChangeInProgress = false;
 		}
 	}
 
-	if (PropertyUtilities.IsValid() && !bIsInteractiveChangeInProgress)
+	if (this->PropertyUtilities.IsValid() && !bIsInteractiveChangeInProgress)
 	{
 		FPropertyChangedEvent ChangeEvent(PropertyHandle->GetProperty(), EPropertyChangeType::ValueSet, OuterObjects);
-		PropertyUtilities->NotifyFinishedChangingProperties(ChangeEvent);
+		this->PropertyUtilities->NotifyFinishedChangingProperties(ChangeEvent);
 	}
 
 	return true;
 }
 
-
-TSharedRef<IPropertyTypeCustomization> FQuatStructCustomization::MakeInstance()
+template<typename T>
+TSharedRef<IPropertyTypeCustomization> FQuatStructCustomization<T>::MakeInstance()
 {
 	return MakeShareable(new FQuatStructCustomization);
 }
 
 
-void FQuatStructCustomization::MakeHeaderRow(TSharedRef<class IPropertyHandle>& InStructPropertyHandle, FDetailWidgetRow& Row)
+template<typename T>
+void FQuatStructCustomization<T>::MakeHeaderRow(TSharedRef<class IPropertyHandle>& InStructPropertyHandle, FDetailWidgetRow& Row)
 {
-	CustomizeRotation(InStructPropertyHandle, Row);
+	this->CustomizeRotation(InStructPropertyHandle, Row);
 }
 
-
-void FQuatStructCustomization::CustomizeChildren(TSharedRef<class IPropertyHandle> StructPropertyHandle, class IDetailChildrenBuilder& StructBuilder, IPropertyTypeCustomizationUtils& StructCustomizationUtils)
+template<typename T>
+void FQuatStructCustomization<T>::CustomizeChildren(TSharedRef<class IPropertyHandle> StructPropertyHandle, class IDetailChildrenBuilder& StructBuilder, IPropertyTypeCustomizationUtils& StructCustomizationUtils)
 {
 	FMathStructProxyCustomization::CustomizeChildren(StructPropertyHandle, StructBuilder, StructCustomizationUtils);
 }
 
-
-bool FQuatStructCustomization::CacheValues(TWeakPtr<IPropertyHandle> PropertyHandlePtr) const
+template<typename T>
+bool FQuatStructCustomization<T>::CacheValues(TWeakPtr<IPropertyHandle> PropertyHandlePtr) const
 {
 	auto PropertyHandle = PropertyHandlePtr.Pin();
 
@@ -686,10 +698,10 @@ bool FQuatStructCustomization::CacheValues(TWeakPtr<IPropertyHandle> PropertyHan
 
 	if (RawData.Num() == 1)
 	{
-		FQuat* QuatValue = reinterpret_cast<FQuat*>(RawData[0]);
+		UE::Math::TQuat<T>* QuatValue = reinterpret_cast<UE::Math::TQuat<T>*>(RawData[0]);
 		if (QuatValue != NULL)
 		{
-			CachedRotation->Set(QuatValue->Rotator());
+			this->CachedRotation->Set(QuatValue->Rotator());
 			return true;
 		}
 	}
@@ -697,7 +709,8 @@ bool FQuatStructCustomization::CacheValues(TWeakPtr<IPropertyHandle> PropertyHan
 	return false;
 }
 
-bool FQuatStructCustomization::FlushValues(TWeakPtr<IPropertyHandle> PropertyHandlePtr) const
+template<typename T>
+bool FQuatStructCustomization<T>::FlushValues(TWeakPtr<IPropertyHandle> PropertyHandlePtr) const
 {
 	auto PropertyHandle = PropertyHandlePtr.Pin();
 
@@ -721,25 +734,25 @@ bool FQuatStructCustomization::FlushValues(TWeakPtr<IPropertyHandle> PropertyHan
 	bool bNotifiedPreChange = false;
 	for (int32 ValueIndex = 0; ValueIndex < RawData.Num(); ValueIndex++)
 	{
-		FQuat* QuatValue = reinterpret_cast<FQuat*>(RawData[0]);
+		UE::Math::TQuat<T>* QuatValue = reinterpret_cast<UE::Math::TQuat<T>*>(RawData[0]);
 		if (QuatValue != NULL)
 		{
-			const FQuat PreviousValue = *QuatValue;
-			const FRotator CurrentRotation = QuatValue->Rotator();
+			const UE::Math::TQuat<T> PreviousValue = *QuatValue;
+			const UE::Math::TRotator<T> CurrentRotation = QuatValue->Rotator();
 
-			FRotator Rotation(
-				CachedRotationPitch->IsSet() ? CachedRotationPitch->Get() : CurrentRotation.Pitch,
-				CachedRotationYaw->IsSet() ? CachedRotationYaw->Get() : CurrentRotation.Yaw,
-				CachedRotationRoll->IsSet() ? CachedRotationRoll->Get() : CurrentRotation.Roll
+			UE::Math::TRotator<T> Rotation(
+				this->CachedRotationPitch->IsSet() ? this->CachedRotationPitch->Get() : CurrentRotation.Pitch,
+				this->CachedRotationYaw->IsSet() ? this->CachedRotationYaw->Get() : CurrentRotation.Yaw,
+				this->CachedRotationRoll->IsSet() ? this->CachedRotationRoll->Get() : CurrentRotation.Roll
 				);
 			
-			const FQuat NewValue = Rotation.Quaternion();
+			const UE::Math::TQuat<T> NewValue = Rotation.Quaternion();
 
-			// In some cases the FQuat pointed to in RawData is no longer aligned to 16 bytes.
-			// Make a local copy to guarantee the alignment criterions of the vector intrinsics inside FQuat::Equals
-			const FQuat AlignedQuatValue = *QuatValue; 
+			// In some cases the UE::Math::TQuat<T> pointed to in RawData is no longer aligned to 16 bytes.
+			// Make a local copy to guarantee the alignment criterions of the vector intrinsics inside UE::Math::TQuat<T>::Equals
+			const UE::Math::TQuat<T> AlignedQuatValue = *QuatValue; 
 
-			if (!bNotifiedPreChange && (!AlignedQuatValue.Equals(NewValue, 0.0f) || (!bIsUsingSlider && bIsInteractiveChangeInProgress)))
+			if (!bNotifiedPreChange && (!AlignedQuatValue.Equals(NewValue, 0.0f) || (!this->bIsUsingSlider && bIsInteractiveChangeInProgress)))
 			{
 				if (!bIsInteractiveChangeInProgress)
 				{
@@ -749,7 +762,7 @@ bool FQuatStructCustomization::FlushValues(TWeakPtr<IPropertyHandle> PropertyHan
 				PropertyHandle->NotifyPreChange();
 				bNotifiedPreChange = true;
 
-				bIsInteractiveChangeInProgress = bIsUsingSlider;
+				bIsInteractiveChangeInProgress = this->bIsUsingSlider;
 			}
 
 			// Set the new value.
@@ -757,14 +770,14 @@ bool FQuatStructCustomization::FlushValues(TWeakPtr<IPropertyHandle> PropertyHan
 
 			// Propagate default value changes after updating, for archetypes. As per usual, we only propagate the change if the instance matches the archetype's value.
 			// Note: We cannot use the "normal" PropertyNode propagation logic here, because that is string-based and the decision to propagate relies on an exact value match.
-			// Here, we're dealing with conversions between FQuat and FRotator values, so there is some precision loss that requires a tolerance when comparing values.
+			// Here, we're dealing with conversions between UE::Math::TQuat<T> and UE::Math::TRotator<T> values, so there is some precision loss that requires a tolerance when comparing values.
 			if (ValueIndex < OuterObjects.Num() && OuterObjects[ValueIndex]->IsTemplate())
 			{
 				TArray<UObject*> ArchetypeInstances;
 				OuterObjects[ValueIndex]->GetArchetypeInstances(ArchetypeInstances);
 				for (UObject* ArchetypeInstance : ArchetypeInstances)
 				{
-					FQuat* CurrentValue = reinterpret_cast<FQuat*>(PropertyHandle->GetValueBaseAddress(reinterpret_cast<uint8*>(ArchetypeInstance)));
+					UE::Math::TQuat<T>* CurrentValue = reinterpret_cast<UE::Math::TQuat<T>*>(PropertyHandle->GetValueBaseAddress(reinterpret_cast<uint8*>(ArchetypeInstance)));
 					if (CurrentValue && CurrentValue->Equals(PreviousValue))
 					{
 						*CurrentValue = NewValue;
@@ -776,22 +789,30 @@ bool FQuatStructCustomization::FlushValues(TWeakPtr<IPropertyHandle> PropertyHan
 
 	if (bNotifiedPreChange)
 	{
-		PropertyHandle->NotifyPostChange(bIsUsingSlider ? EPropertyChangeType::Interactive : EPropertyChangeType::ValueSet);
+		PropertyHandle->NotifyPostChange(this->bIsUsingSlider ? EPropertyChangeType::Interactive : EPropertyChangeType::ValueSet);
 
-		if (!bIsUsingSlider)
+		if (!this->bIsUsingSlider)
 		{
 			GEditor->EndTransaction();
 			bIsInteractiveChangeInProgress = false;
 		}
 	}
 
-	if (PropertyUtilities.IsValid() && !bIsInteractiveChangeInProgress)
+	if (this->PropertyUtilities.IsValid() && !bIsInteractiveChangeInProgress)
 	{
 		FPropertyChangedEvent ChangeEvent(PropertyHandle->GetProperty(), EPropertyChangeType::ValueSet, OuterObjects);
-		PropertyUtilities->NotifyFinishedChangingProperties(ChangeEvent);
+		this->PropertyUtilities->NotifyFinishedChangingProperties(ChangeEvent);
 	}
 
 	return true;
 }
 
+// Instantiate for linker
+template class FMatrixStructCustomization<float>;
+template class FMatrixStructCustomization<double>;
+template class FTransformStructCustomization<float>;
+template class FTransformStructCustomization<double>;
+template class FQuatStructCustomization<float>;
+template class FQuatStructCustomization<double>;
+
 #undef LOCTEXT_NAMESPACE
diff --git a/Engine/Source/Editor/DetailCustomizations/Public/Customizations/MathStructProxyCustomizations.h b/Engine/Source/Editor/DetailCustomizations/Public/Customizations/MathStructProxyCustomizations.h
--- a/Engine/Source/Editor/DetailCustomizations/Public/Customizations/MathStructProxyCustomizations.h
+++ b/Engine/Source/Editor/DetailCustomizations/Public/Customizations/MathStructProxyCustomizations.h
@@ -222,9 +222,21 @@ protected:
 	TSharedPtr<IPropertyUtilities> PropertyUtilities;
 };
 
+
+struct FTransformField
+{
+	enum Type
+	{
+		Location,
+		Rotation,
+		Scale
+	};
+};
+
 /** 
  * Proxy struct customization that displays a matrix as a position, rotation & scale.
  */
+template<typename T>
 class DETAILCUSTOMIZATIONS_API FMatrixStructCustomization : public FMathStructProxyCustomization
 {
 public:
@@ -232,18 +244,18 @@ public:
 
 public:
 	FMatrixStructCustomization()
-		: CachedRotation(MakeShareable( new TProxyValue<FRotator>(FRotator::ZeroRotator)))
-		, CachedRotationYaw(MakeShareable( new TProxyProperty<FRotator, FRotator::FReal>(CachedRotation, CachedRotation->Get().Yaw)))
-		, CachedRotationPitch(MakeShareable( new TProxyProperty<FRotator, FRotator::FReal>(CachedRotation, CachedRotation->Get().Pitch)))
-		, CachedRotationRoll(MakeShareable( new TProxyProperty<FRotator, FRotator::FReal>(CachedRotation, CachedRotation->Get().Roll)))
-		, CachedTranslation(MakeShareable( new TProxyValue<FVector>(FVector::ZeroVector)))
-		, CachedTranslationX(MakeShareable( new TProxyProperty<FVector, FVector::FReal>(CachedTranslation, CachedTranslation->Get().X)))
-		, CachedTranslationY(MakeShareable( new TProxyProperty<FVector, FVector::FReal>(CachedTranslation, CachedTranslation->Get().Y)))
-		, CachedTranslationZ(MakeShareable( new TProxyProperty<FVector, FVector::FReal>(CachedTranslation, CachedTranslation->Get().Z)))
-		, CachedScale(MakeShareable( new TProxyValue<FVector>(FVector::ZeroVector)))
-		, CachedScaleX(MakeShareable( new TProxyProperty<FVector, FVector::FReal>(CachedScale, CachedScale->Get().X)))
-		, CachedScaleY(MakeShareable( new TProxyProperty<FVector, FVector::FReal>(CachedScale, CachedScale->Get().Y)))
-		, CachedScaleZ(MakeShareable( new TProxyProperty<FVector, FVector::FReal>(CachedScale, CachedScale->Get().Z)))
+		: CachedRotation(MakeShareable( new TProxyValue<UE::Math::TRotator<T>>(UE::Math::TRotator<T>::ZeroRotator)))
+		, CachedRotationYaw(MakeShareable( new TProxyProperty<UE::Math::TRotator<T>, T>(CachedRotation, CachedRotation->Get().Yaw)))
+		, CachedRotationPitch(MakeShareable( new TProxyProperty<UE::Math::TRotator<T>, T>(CachedRotation, CachedRotation->Get().Pitch)))
+		, CachedRotationRoll(MakeShareable( new TProxyProperty<UE::Math::TRotator<T>, T>(CachedRotation, CachedRotation->Get().Roll)))
+		, CachedTranslation(MakeShareable( new TProxyValue<UE::Math::TVector<T>>(UE::Math::TVector<T>::ZeroVector)))
+		, CachedTranslationX(MakeShareable( new TProxyProperty<UE::Math::TVector<T>, T>(CachedTranslation, CachedTranslation->Get().X)))
+		, CachedTranslationY(MakeShareable( new TProxyProperty<UE::Math::TVector<T>, T>(CachedTranslation, CachedTranslation->Get().Y)))
+		, CachedTranslationZ(MakeShareable( new TProxyProperty<UE::Math::TVector<T>, T>(CachedTranslation, CachedTranslation->Get().Z)))
+		, CachedScale(MakeShareable( new TProxyValue<UE::Math::TVector<T>>(UE::Math::TVector<T>::ZeroVector)))
+		, CachedScaleX(MakeShareable( new TProxyProperty<UE::Math::TVector<T>, T>(CachedScale, CachedScale->Get().X)))
+		, CachedScaleY(MakeShareable( new TProxyProperty<UE::Math::TVector<T>, T>(CachedScale, CachedScale->Get().Y)))
+		, CachedScaleZ(MakeShareable( new TProxyProperty<UE::Math::TVector<T>, T>(CachedScale, CachedScale->Get().Z)))
 	{
 	}
 
@@ -263,43 +275,34 @@ protected:
 	virtual bool CacheValues(TWeakPtr<IPropertyHandle> PropertyHandlePtr) const override;
 	virtual bool FlushValues(TWeakPtr<IPropertyHandle> PropertyHandlePtr) const override;
 
-	struct FTransformField
-	{
-		enum Type
-		{
-			Location,
-			Rotation,
-			Scale
-		};
-	};
-
 	void OnCopy(FTransformField::Type Type, TWeakPtr<IPropertyHandle> PropertyHandlePtr);
 	void OnPaste(FTransformField::Type Type, TWeakPtr<IPropertyHandle> PropertyHandlePtr);
 
 protected:
 	/** Cached rotation values */
-	mutable TSharedRef< TProxyValue<FRotator> > CachedRotation;
-	mutable TSharedRef< TProxyProperty<FRotator, FRotator::FReal> > CachedRotationYaw;
-	mutable TSharedRef< TProxyProperty<FRotator, FRotator::FReal> > CachedRotationPitch;
-	mutable TSharedRef< TProxyProperty<FRotator, FRotator::FReal> > CachedRotationRoll;
+	mutable TSharedRef< TProxyValue<UE::Math::TRotator<T>> > CachedRotation;
+	mutable TSharedRef< TProxyProperty<UE::Math::TRotator<T>, T> > CachedRotationYaw;
+	mutable TSharedRef< TProxyProperty<UE::Math::TRotator<T>, T> > CachedRotationPitch;
+	mutable TSharedRef< TProxyProperty<UE::Math::TRotator<T>, T> > CachedRotationRoll;
 
 	/** Cached translation values */
-	mutable TSharedRef< TProxyValue<FVector> > CachedTranslation;
-	mutable TSharedRef< TProxyProperty<FVector, FVector::FReal> > CachedTranslationX;
-	mutable TSharedRef< TProxyProperty<FVector, FVector::FReal> > CachedTranslationY;
-	mutable TSharedRef< TProxyProperty<FVector, FVector::FReal> > CachedTranslationZ;
+	mutable TSharedRef< TProxyValue<UE::Math::TVector<T>> > CachedTranslation;
+	mutable TSharedRef< TProxyProperty<UE::Math::TVector<T>, T > > CachedTranslationX;
+	mutable TSharedRef< TProxyProperty<UE::Math::TVector<T>, T > > CachedTranslationY;
+	mutable TSharedRef< TProxyProperty<UE::Math::TVector<T>, T> > CachedTranslationZ;
 
 	/** Cached scale values */
-	mutable TSharedRef< TProxyValue<FVector> > CachedScale;
-	mutable TSharedRef< TProxyProperty<FVector, FVector::FReal> > CachedScaleX;
-	mutable TSharedRef< TProxyProperty<FVector, FVector::FReal> > CachedScaleY;
-	mutable TSharedRef< TProxyProperty<FVector, FVector::FReal> > CachedScaleZ;
+	mutable TSharedRef< TProxyValue<UE::Math::TVector<T>> > CachedScale;
+	mutable TSharedRef< TProxyProperty<UE::Math::TVector<T>, T> > CachedScaleX;
+	mutable TSharedRef< TProxyProperty<UE::Math::TVector<T>, T> > CachedScaleY;
+	mutable TSharedRef< TProxyProperty<UE::Math::TVector<T>, T> > CachedScaleZ;
 };
 
 /** 
  * Proxy struct customization that displays an FTransform as a position, euler rotation & scale.
  */
-class FTransformStructCustomization : public FMatrixStructCustomization
+template<typename T>
+class FTransformStructCustomization : public FMatrixStructCustomization<T>
 {
 public:
 	static TSharedRef<IPropertyTypeCustomization> MakeInstance();
@@ -313,7 +316,8 @@ protected:
 /**
  * Proxy struct customization that displays an FQuat as an euler rotation
  */
-class DETAILCUSTOMIZATIONS_API FQuatStructCustomization : public FMatrixStructCustomization
+template<typename T>
+class DETAILCUSTOMIZATIONS_API FQuatStructCustomization : public FMatrixStructCustomization<T>
 {
 public:
 	static TSharedRef<IPropertyTypeCustomization> MakeInstance();
