Add support for getting UScriptStruct on FIntPoint, FIntVector and FIntVector4.
Implement the FIntVector4 property struct and detail customisation.
Added FIntVector and FIntVector4 NAME_ FNames.

diff --git a/Engine/Config/BaseEditor.ini b/Engine/Config/BaseEditor.ini
--- a/Engine/Config/BaseEditor.ini
+++ b/Engine/Config/BaseEditor.ini
@@ -505,6 +505,7 @@ bDisableSearchResultTemplates=false
 +NoExportTypesWithDirectNativeFieldAccess=/Script/CoreUObject.PackedRGB10A2N
 +NoExportTypesWithDirectNativeFieldAccess=/Script/CoreUObject.IntPoint
 +NoExportTypesWithDirectNativeFieldAccess=/Script/CoreUObject.IntVector
++NoExportTypesWithDirectNativeFieldAccess=/Script/CoreUObject.IntVector4
 +NoExportTypesWithDirectNativeFieldAccess=/Script/CoreUObject.Color
 +NoExportTypesWithDirectNativeFieldAccess=/Script/CoreUObject.LinearColor
 +NoExportTypesWithDirectNativeFieldAccess=/Script/CoreUObject.Box
diff --git a/Engine/Source/Editor/DetailCustomizations/Private/DetailCustomizations.cpp b/Engine/Source/Editor/DetailCustomizations/Private/DetailCustomizations.cpp
--- a/Engine/Source/Editor/DetailCustomizations/Private/DetailCustomizations.cpp
+++ b/Engine/Source/Editor/DetailCustomizations/Private/DetailCustomizations.cpp
@@ -209,7 +209,7 @@ void FDetailCustomizationsModule::RegisterPropertyTypeCustomizations()
 	REGISTER_UIMINMAX_CUSTOMIZATION(NAME_Vector, &FVectorStructCustomization::MakeInstance);
 	REGISTER_UIMINMAX_CUSTOMIZATION(NAME_Vector3f, &FVectorStructCustomization::MakeInstance);
 	REGISTER_UIMINMAX_CUSTOMIZATION(NAME_Vector3d, &FVectorStructCustomization::MakeInstance);
-	REGISTER_UIMINMAX_CUSTOMIZATION("IntVector", &FVectorStructCustomization::MakeInstance);
+	REGISTER_UIMINMAX_CUSTOMIZATION(NAME_IntVector, &FVectorStructCustomization::MakeInstance);
 	REGISTER_UIMINMAX_CUSTOMIZATION(NAME_Vector4, &FVector4StructCustomization::MakeInstance);
 	REGISTER_UIMINMAX_CUSTOMIZATION(NAME_Vector4f, &FVector4StructCustomization::MakeInstance);
 	REGISTER_UIMINMAX_CUSTOMIZATION(NAME_Vector4d, &FVector4StructCustomization::MakeInstance);
@@ -217,6 +217,7 @@ void FDetailCustomizationsModule::RegisterPropertyTypeCustomizations()
 	REGISTER_UIMINMAX_CUSTOMIZATION(NAME_Vector2f, &FMathStructCustomization::MakeInstance);
 	REGISTER_UIMINMAX_CUSTOMIZATION(NAME_Vector2d, &FMathStructCustomization::MakeInstance);
 	REGISTER_UIMINMAX_CUSTOMIZATION(NAME_IntPoint, &FMathStructCustomization::MakeInstance);
+	REGISTER_UIMINMAX_CUSTOMIZATION(NAME_IntVector4, &FVector4StructCustomization::MakeInstance);
 	REGISTER_UIMINMAX_CUSTOMIZATION(NAME_Rotator, &FRotatorStructCustomization::MakeInstance);
 	REGISTER_UIMINMAX_CUSTOMIZATION(NAME_Rotator3f, &FRotatorStructCustomization::MakeInstance);
 	REGISTER_UIMINMAX_CUSTOMIZATION(NAME_Rotator3d, &FRotatorStructCustomization::MakeInstance);
diff --git a/Engine/Source/Runtime/Core/Public/UObject/UnrealNames.inl b/Engine/Source/Runtime/Core/Public/UObject/UnrealNames.inl
--- a/Engine/Source/Runtime/Core/Public/UObject/UnrealNames.inl
+++ b/Engine/Source/Runtime/Core/Public/UObject/UnrealNames.inl
@@ -98,6 +98,8 @@ REGISTER_NAME(91,Vector2d)
 REGISTER_NAME(92,Box2D)
 REGISTER_NAME(93,Box2f)
 REGISTER_NAME(94,Box2d)
+REGISTER_NAME(95,IntVector)
+REGISTER_NAME(96,IntVector4)
 
 
 // Object class names.
diff --git a/Engine/Source/Runtime/CoreUObject/Private/UObject/Class.cpp b/Engine/Source/Runtime/CoreUObject/Private/UObject/Class.cpp
--- a/Engine/Source/Runtime/CoreUObject/Private/UObject/Class.cpp
+++ b/Engine/Source/Runtime/CoreUObject/Private/UObject/Class.cpp
@@ -6205,6 +6205,24 @@ static UScriptStruct* StaticGetBaseStructureInternal(FName Name)
 	return Result;
 }
 
+UScriptStruct* TBaseStructure<FIntPoint>::Get()
+{
+	static UScriptStruct* ScriptStruct = StaticGetBaseStructureInternal(TEXT("IntPoint"));
+	return ScriptStruct;
+}
+
+UScriptStruct* TBaseStructure<FIntVector>::Get()
+{
+	static UScriptStruct* ScriptStruct = StaticGetBaseStructureInternal(TEXT("IntVector"));
+	return ScriptStruct;
+}
+
+UScriptStruct* TBaseStructure<FIntVector4>::Get()
+{
+	static UScriptStruct* ScriptStruct = StaticGetBaseStructureInternal(TEXT("IntVector4"));
+	return ScriptStruct;
+}
+
 UScriptStruct* TBaseStructure<FLinearColor>::Get()
 {
 	static UScriptStruct* ScriptStruct = StaticGetBaseStructureInternal(TEXT("LinearColor"));
diff --git a/Engine/Source/Runtime/CoreUObject/Private/UObject/Property.cpp b/Engine/Source/Runtime/CoreUObject/Private/UObject/Property.cpp
--- a/Engine/Source/Runtime/CoreUObject/Private/UObject/Property.cpp
+++ b/Engine/Source/Runtime/CoreUObject/Private/UObject/Property.cpp
@@ -73,6 +73,19 @@ struct TStructOpsTypeTraits<FIntVector> : public TStructOpsTypeTraitsBase2<FIntV
 };
 IMPLEMENT_STRUCT(IntVector);
 
+template<>
+struct TStructOpsTypeTraits<FIntVector4> : public TStructOpsTypeTraitsBase2<FIntVector4>
+{
+	enum
+	{
+		WithIdenticalViaEquality = true,
+		WithNoInitConstructor = true,
+		WithZeroConstructor = true,
+		WithSerializer = true,
+	};
+};
+IMPLEMENT_STRUCT(IntVector4);
+
 template<typename T>
 struct TVector2StructOpsTypeTraits : public TStructOpsTypeTraitsBase2<T>
 {
diff --git a/Engine/Source/Runtime/CoreUObject/Public/UObject/Class.h b/Engine/Source/Runtime/CoreUObject/Public/UObject/Class.h
--- a/Engine/Source/Runtime/CoreUObject/Public/UObject/Class.h
+++ b/Engine/Source/Runtime/CoreUObject/Public/UObject/Class.h
@@ -3534,6 +3534,21 @@ template< class T > struct TBaseStructure
 	}
 };
 
+template<> struct TBaseStructure<FIntPoint> 
+{
+	COREUOBJECT_API static UScriptStruct* Get(); 
+};
+
+template<> struct TBaseStructure<FIntVector> 
+{ 
+	COREUOBJECT_API static UScriptStruct* Get(); 
+};
+
+template<> struct TBaseStructure<FIntVector4>
+{
+	COREUOBJECT_API static UScriptStruct* Get();
+};
+
 template<> struct TBaseStructure<FLinearColor>
 {
 	COREUOBJECT_API static UScriptStruct* Get();
diff --git a/Engine/Source/Runtime/CoreUObject/Public/UObject/NoExportTypes.h b/Engine/Source/Runtime/CoreUObject/Public/UObject/NoExportTypes.h
--- a/Engine/Source/Runtime/CoreUObject/Public/UObject/NoExportTypes.h
+++ b/Engine/Source/Runtime/CoreUObject/Public/UObject/NoExportTypes.h
@@ -867,6 +867,26 @@ struct FIntVector
 	int32 Z;
 };
 
+/**
+ * An integer vector in 4D space.
+ * @note The full C++ class is located here: Engine\Source\Runtime\Core\Public\Math\IntVector.h
+ */
+USTRUCT(immutable, noexport, BlueprintType, IsAlwaysAccessible, HasDefaults, HasNoOpConstructor, IsCoreType)
+struct FIntVector4
+{
+	UPROPERTY(EditAnywhere, BlueprintReadWrite, Category = IntVector4, SaveGame)
+	int32 X;
+
+	UPROPERTY(EditAnywhere, BlueprintReadWrite, Category = IntVector4, SaveGame)
+	int32 Y;
+
+	UPROPERTY(EditAnywhere, BlueprintReadWrite, Category = IntVector4, SaveGame)
+	int32 Z;
+
+	UPROPERTY(EditAnywhere, BlueprintReadWrite, Category = IntVector4, SaveGame)
+	int32 W;
+};
+
 /**
  * Stores a color with 8 bits of precision per channel. (BGRA).
  * @note The full C++ class is located here: Engine\Source\Runtime\Core\Public\Math\Color.h
