PR #8872: Color picker for color properties on externally added FStructOnScope properties now resets/cancels correctly (Contributed by [REDACTED_PERSON])

diff --git a/Engine/Source/Editor/DetailCustomizations/Private/ColorStructCustomization.cpp b/Engine/Source/Editor/DetailCustomizations/Private/ColorStructCustomization.cpp
--- a/Engine/Source/Editor/DetailCustomizations/Private/ColorStructCustomization.cpp
+++ b/Engine/Source/Editor/DetailCustomizations/Private/ColorStructCustomization.cpp
@@ -176,13 +176,11 @@ void FColorStructCustomization::CreateColorPicker(bool bUseAlpha)
 {
 	GEditor->BeginTransaction(FText::Format(LOCTEXT("SetColorProperty", "Edit {0}"), StructPropertyHandle->GetPropertyDisplayName()));
 
-	int32 NumObjects = StructPropertyHandle->GetNumOuterObjects();
-
 	SavedPreColorPickerColors.Empty();
 	TArray<FString> PerObjectValues;
 	StructPropertyHandle->GetPerObjectValues(PerObjectValues);
 
-	for (int32 ObjectIndex = 0; ObjectIndex < NumObjects; ++ObjectIndex)
+	for (int32 ObjectIndex = 0; ObjectIndex < PerObjectValues.Num(); ++ObjectIndex)
 	{
 		if (bIsLinearColor)
 		{
@@ -234,13 +232,11 @@ TSharedRef<SColorPicker> FColorStructCustomization::CreateInlineColorPicker(TWea
 {
 	GEditor->BeginTransaction(FText::Format(LOCTEXT("SetColorProperty", "Edit {0}"), StructPropertyHandle->GetPropertyDisplayName()));
 
-	int32 NumObjects = StructPropertyHandle->GetNumOuterObjects();
-
 	SavedPreColorPickerColors.Empty();
 	TArray<FString> PerObjectValues;
 	StructPropertyHandle->GetPerObjectValues(PerObjectValues);
 
-	for (int32 ObjectIndex = 0; ObjectIndex < NumObjects; ++ObjectIndex)
+	for (int32 ObjectIndex = 0; ObjectIndex < PerObjectValues.Num(); ++ObjectIndex)
 	{
 		if (bIsLinearColor)
 		{
