Engine - Adding support for Double Range types, as well as reading doubles in metadata.

diff --git a/Engine/Source/Editor/DetailCustomizations/Private/DetailCustomizations.cpp b/Engine/Source/Editor/DetailCustomizations/Private/DetailCustomizations.cpp
--- a/Engine/Source/Editor/DetailCustomizations/Private/DetailCustomizations.cpp
+++ b/Engine/Source/Editor/DetailCustomizations/Private/DetailCustomizations.cpp
@@ -244,8 +244,10 @@ void FDetailCustomizationsModule::RegisterPropertyTypeCustomizations()
 	RegisterCustomPropertyTypeLayout("Guid", FOnGetPropertyTypeCustomizationInstance::CreateStatic(&FGuidStructCustomization::MakeInstance));
 	RegisterCustomPropertyTypeLayout("Key", FOnGetPropertyTypeCustomizationInstance::CreateStatic(&FKeyStructCustomization::MakeInstance));
 	RegisterCustomPropertyTypeLayout("FloatRange", FOnGetPropertyTypeCustomizationInstance::CreateStatic(&FRangeStructCustomization<float>::MakeInstance));
+	RegisterCustomPropertyTypeLayout("DoubleRange", FOnGetPropertyTypeCustomizationInstance::CreateStatic(&FRangeStructCustomization<double>::MakeInstance));
 	RegisterCustomPropertyTypeLayout("Int32Range", FOnGetPropertyTypeCustomizationInstance::CreateStatic(&FRangeStructCustomization<int32>::MakeInstance));
 	RegisterCustomPropertyTypeLayout("FloatInterval", FOnGetPropertyTypeCustomizationInstance::CreateStatic(&FIntervalStructCustomization<float>::MakeInstance));
+	RegisterCustomPropertyTypeLayout("DoubleInterval", FOnGetPropertyTypeCustomizationInstance::CreateStatic(&FIntervalStructCustomization<double>::MakeInstance));
 	RegisterCustomPropertyTypeLayout("Int32Interval", FOnGetPropertyTypeCustomizationInstance::CreateStatic(&FIntervalStructCustomization<int32>::MakeInstance));
 	RegisterCustomPropertyTypeLayout("DateTime", FOnGetPropertyTypeCustomizationInstance::CreateStatic(&FDateTimeStructCustomization::MakeInstance));
 	RegisterCustomPropertyTypeLayout("Timespan", FOnGetPropertyTypeCustomizationInstance::CreateStatic(&FTimespanStructCustomization::MakeInstance));
diff --git a/Engine/Source/Editor/DetailCustomizations/Private/IntervalStructCustomization.cpp b/Engine/Source/Editor/DetailCustomizations/Private/IntervalStructCustomization.cpp
--- a/Engine/Source/Editor/DetailCustomizations/Private/IntervalStructCustomization.cpp
+++ b/Engine/Source/Editor/DetailCustomizations/Private/IntervalStructCustomization.cpp
@@ -349,6 +349,7 @@ bool FIntervalStructCustomization<NumericType>::IsPropertyEnabled(EIntervalField
  *****************************************************************************/
 
 template class FIntervalStructCustomization<float>;
+template class FIntervalStructCustomization<double>;
 template class FIntervalStructCustomization<int32>;
 
 
diff --git a/Engine/Source/Editor/DetailCustomizations/Private/RangeStructCustomization.cpp b/Engine/Source/Editor/DetailCustomizations/Private/RangeStructCustomization.cpp
--- a/Engine/Source/Editor/DetailCustomizations/Private/RangeStructCustomization.cpp
+++ b/Engine/Source/Editor/DetailCustomizations/Private/RangeStructCustomization.cpp
@@ -36,6 +36,15 @@ namespace
 		}
 	};
 
+	template <>
+	struct FGetMetaDataHelper<double>
+	{
+		static double GetMetaData(const FProperty* Property, const TCHAR* Key)
+		{
+			return Property->GetDoubleMetaData(Key);
+		}
+	};
+
 	template <>
 	struct FGetMetaDataHelper<int32>
 	{
@@ -395,6 +404,7 @@ void FRangeStructCustomization<NumericType>::OnComboSelectionChanged(TSharedPtr<
  *****************************************************************************/
 
 template class FRangeStructCustomization<float>;
+template class FRangeStructCustomization<double>;
 template class FRangeStructCustomization<int32>;
 
 
diff --git a/Engine/Source/Runtime/Core/Public/Math/Interval.h b/Engine/Source/Runtime/Core/Public/Math/Interval.h
--- a/Engine/Source/Runtime/Core/Public/Math/Interval.h
+++ b/Engine/Source/Runtime/Core/Public/Math/Interval.h
@@ -266,4 +266,5 @@ public:
 	};
 
 DEFINE_INTERVAL_WRAPPER_STRUCT(FFloatInterval, float)
+DEFINE_INTERVAL_WRAPPER_STRUCT(FDoubleInterval, double)
 DEFINE_INTERVAL_WRAPPER_STRUCT(FInt32Interval, int32)
diff --git a/Engine/Source/Runtime/CoreUObject/Public/UObject/Class.h b/Engine/Source/Runtime/CoreUObject/Public/UObject/Class.h
--- a/Engine/Source/Runtime/CoreUObject/Public/UObject/Class.h
+++ b/Engine/Source/Runtime/CoreUObject/Public/UObject/Class.h
@@ -3625,6 +3625,16 @@ template<> struct TBaseStructure<FFloatRange>
 	COREUOBJECT_API static UScriptStruct* Get();
 };
 
+template<> struct TBaseStructure<FDoubleRangeBound>
+{
+	COREUOBJECT_API static UScriptStruct* Get();
+};
+
+template<> struct TBaseStructure<FDoubleRange>
+{
+	COREUOBJECT_API static UScriptStruct* Get();
+};
+
 template<> struct TBaseStructure<FInt32RangeBound>
 {
 	COREUOBJECT_API static UScriptStruct* Get();
@@ -3640,6 +3650,11 @@ template<> struct TBaseStructure<FFloatInterval>
 	COREUOBJECT_API static UScriptStruct* Get();
 };
 
+template<> struct TBaseStructure<FDoubleInterval>
+{
+	COREUOBJECT_API static UScriptStruct* Get();
+};
+
 template<> struct TBaseStructure<FInt32Interval>
 {
 	COREUOBJECT_API static UScriptStruct* Get();
diff --git a/Engine/Source/Runtime/CoreUObject/Public/UObject/Field.h b/Engine/Source/Runtime/CoreUObject/Public/UObject/Field.h
--- a/Engine/Source/Runtime/CoreUObject/Public/UObject/Field.h
+++ b/Engine/Source/Runtime/CoreUObject/Public/UObject/Field.h
@@ -823,6 +823,27 @@ public:
 		return Value;
 	}
 
+	/**
+	* Find the metadata value associated with the key
+	* and return double
+	* @param Key The key to lookup in the metadata
+	* @return the float value stored in the metadata.
+	*/
+	double GetDoubleMetaData(const TCHAR* Key) const
+	{
+		const FString& DOUBLEString = GetMetaData(Key);
+		// FString == operator does case insensitive comparison
+		double Value = FCString::Atod(*DOUBLEString);
+		return Value;
+	}
+	double GetDoubleMetaData(const FName& Key) const
+	{
+		const FString& DOUBLEString = GetMetaData(Key);
+		// FString == operator does case insensitive comparison
+		double Value = FCString::Atod(*DOUBLEString);
+		return Value;
+	}
+
 	/**
 	* Find the metadata value associated with the key
 	* and return Class
diff --git a/Engine/Source/Runtime/CoreUObject/Public/UObject/NoExportTypes.h b/Engine/Source/Runtime/CoreUObject/Public/UObject/NoExportTypes.h
--- a/Engine/Source/Runtime/CoreUObject/Public/UObject/NoExportTypes.h
+++ b/Engine/Source/Runtime/CoreUObject/Public/UObject/NoExportTypes.h
@@ -1790,6 +1790,43 @@ private:
 	FFloatRangeBound UpperBound;
 };
 
+
+/**
+ * Defines a single bound for a range of values.
+ * @note This is a mirror of TRangeBound<double>, defined in RangeBound.h
+ * @note Fields are private to match the C++ declaration in the header above.
+ */
+USTRUCT(noexport, BlueprintType)
+struct FDoubleRangeBound
+{
+private:
+	/** Holds the type of the bound. */
+	UPROPERTY(EditAnywhere, BlueprintReadWrite, Category=Range, meta=(AllowPrivateAccess="true"))
+	TEnumAsByte<ERangeBoundTypes::Type> Type;
+
+	/** Holds the bound's value. */
+	UPROPERTY(EditAnywhere, BlueprintReadWrite, Category=Range, meta=(AllowPrivateAccess="true"))
+	double Value;
+};
+
+/**
+ * A contiguous set of doubles described by lower and upper bound values.
+ * @note This is a mirror of TRange<double>, defined in Range.h
+ * @note Fields are private to match the C++ declaration in the header above.
+ */
+USTRUCT(noexport, BlueprintType)
+struct FDoubleRange
+{
+private:
+	/** Holds the range's lower bound. */
+	UPROPERTY(EditAnywhere, BlueprintReadWrite, Category=Range, meta=(AllowPrivateAccess="true"))
+	FDoubleRangeBound LowerBound;
+
+	/** Holds the range's upper bound. */
+	UPROPERTY(EditAnywhere, BlueprintReadWrite, Category=Range, meta=(AllowPrivateAccess="true"))
+	FDoubleRangeBound UpperBound;
+};
+
 /**
  * Defines a single bound for a range of values.
  * @note This is a mirror of TRangeBound<int32>, defined in RangeBound.h
