Speculative fix for a crash on an invalid pointer. My suspicion is that the UBodyInstance* is garbage because it's accesssing unitialized memory, changed an AddUnitialized to AddZeroed and added a null check to make sure.

Removed a few unnecessary forward declarations.

diff --git a/Engine/Source/Editor/DetailCustomizations/Private/BodyInstanceCustomization.cpp b/Engine/Source/Editor/DetailCustomizations/Private/BodyInstanceCustomization.cpp
--- a/Engine/Source/Editor/DetailCustomizations/Private/BodyInstanceCustomization.cpp
+++ b/Engine/Source/Editor/DetailCustomizations/Private/BodyInstanceCustomization.cpp
@@ -41,8 +41,13 @@ FBodyInstanceCustomization::FBodyInstanceCustomization()
 
 UStaticMeshComponent* FBodyInstanceCustomization::GetDefaultCollisionProvider(const FBodyInstance* BI) const
 {
+	if (!BI)
+	{
+		return nullptr;
+	}
+
 	UPrimitiveComponent* OwnerComp = BI->OwnerComponent.Get();
-	if(!OwnerComp)
+	if (!OwnerComp)
 	{
 		TWeakObjectPtr<UPrimitiveComponent> FoundComp = BodyInstanceToPrimComponent.FindRef(BI);
 		OwnerComp = FoundComp.Get();
@@ -55,7 +60,7 @@ UStaticMeshComponent* FBodyInstanceCustomization::GetDefaultCollisionProvider(co
 bool FBodyInstanceCustomization::CanUseDefaultCollision() const
 {
 	bool bResult = BodyInstances.Num() > 0;
-	for(const FBodyInstance* BI : BodyInstances)
+	for (const FBodyInstance* BI : BodyInstances)
 	{
 		bResult &= GetDefaultCollisionProvider(BI) != nullptr;
 	}
@@ -166,7 +171,7 @@ void FBodyInstanceCustomization::AddCollisionCategory(TSharedRef<class IProperty
 }
 
 BEGIN_SLATE_FUNCTION_BUILD_OPTIMIZATION
-void FBodyInstanceCustomization::CustomizeChildren( TSharedRef<class IPropertyHandle> StructPropertyHandle, class IDetailChildrenBuilder& StructBuilder, IPropertyTypeCustomizationUtils& StructCustomizationUtils )
+void FBodyInstanceCustomization::CustomizeChildren( TSharedRef<IPropertyHandle> StructPropertyHandle, IDetailChildrenBuilder& StructBuilder, IPropertyTypeCustomizationUtils& StructCustomizationUtils )
 {
 	BodyInstanceHandle = StructPropertyHandle;
 
@@ -175,7 +180,7 @@ void FBodyInstanceCustomization::CustomizeChildren( TSharedRef<class IPropertyHa
 	StructPropertyHandle->AccessRawData(StructPtrs);
 	check(StructPtrs.Num() != 0);
 
-	BodyInstances.AddUninitialized(StructPtrs.Num());
+	BodyInstances.AddZeroed(StructPtrs.Num());
 	for (auto Iter = StructPtrs.CreateIterator(); Iter; ++Iter)
 	{
 		check(*Iter);
@@ -186,7 +191,7 @@ void FBodyInstanceCustomization::CustomizeChildren( TSharedRef<class IPropertyHa
 	StructPropertyHandle->GetOuterObjects(OwningObjects);
 
 	PrimComponents.Empty(OwningObjects.Num());
-	for(UObject* Obj : OwningObjects)
+	for (UObject* Obj : OwningObjects)
 	{
 		if(UPrimitiveComponent* PrimComponent = Cast<UPrimitiveComponent>(Obj))
 		{
@@ -207,7 +212,6 @@ void FBodyInstanceCustomization::CustomizeChildren( TSharedRef<class IPropertyHa
 
 	if(CollisionCategoryHandle.IsValid())
 	{
-		
 		UseDefaultCollisionHandle = CollisionCategoryHandle->GetChildHandle(GET_MEMBER_NAME_CHECKED(UStaticMeshComponent, bUseDefaultCollision));
 	}
 
diff --git a/Engine/Source/Editor/DetailCustomizations/Private/BodyInstanceCustomization.h b/Engine/Source/Editor/DetailCustomizations/Private/BodyInstanceCustomization.h
--- a/Engine/Source/Editor/DetailCustomizations/Private/BodyInstanceCustomization.h
+++ b/Engine/Source/Editor/DetailCustomizations/Private/BodyInstanceCustomization.h
@@ -39,8 +39,8 @@ public:
 	}
 
 	/** IPropertyTypeCustomization interface */
-	virtual void CustomizeHeader( TSharedRef<IPropertyHandle> StructPropertyHandle, class FDetailWidgetRow& HeaderRow, IPropertyTypeCustomizationUtils& StructCustomizationUtils ) override {};
-	virtual void CustomizeChildren( TSharedRef<class IPropertyHandle> StructPropertyHandle, class IDetailChildrenBuilder& StructBuilder, IPropertyTypeCustomizationUtils& StructCustomizationUtils ) override;
+	virtual void CustomizeHeader( TSharedRef<IPropertyHandle> StructPropertyHandle, FDetailWidgetRow& HeaderRow, IPropertyTypeCustomizationUtils& StructCustomizationUtils ) override {};
+	virtual void CustomizeChildren( TSharedRef<IPropertyHandle> StructPropertyHandle, IDetailChildrenBuilder& StructBuilder, IPropertyTypeCustomizationUtils& StructCustomizationUtils ) override;
 
 private:
 
@@ -84,7 +84,7 @@ private:
 	bool AreAllCollisionUsingDefault() const;
 
 	// utility functions between property and struct
-	void CreateCustomCollisionSetup( TSharedRef<class IPropertyHandle> StructPropertyHandle, class IDetailGroup& CollisionGroup );
+	void CreateCustomCollisionSetup( TSharedRef<IPropertyHandle> StructPropertyHandle, class IDetailGroup& CollisionGroup );
 	void SetCollisionResponseContainer(const FCollisionResponseContainer& ResponseContainer);
 	void SetResponse(int32 ValidIndex, ECollisionResponse InCollisionResponse);
 	void UpdateCollisionProfile();
@@ -92,8 +92,8 @@ private:
 
 	void UpdateValidCollisionChannels();
 
-	void AddPhysicsCategory(TSharedRef<class IPropertyHandle> StructPropertyHandle, class IDetailChildrenBuilder& StructBuilder, IPropertyTypeCustomizationUtils& StructCustomizationUtils);
-	void AddCollisionCategory(TSharedRef<class IPropertyHandle> StructPropertyHandle, class IDetailChildrenBuilder& StructBuilder, IPropertyTypeCustomizationUtils& StructCustomizationUtils);
+	void AddPhysicsCategory(TSharedRef<IPropertyHandle> StructPropertyHandle, class IDetailChildrenBuilder& StructBuilder, IPropertyTypeCustomizationUtils& StructCustomizationUtils);
+	void AddCollisionCategory(TSharedRef<IPropertyHandle> StructPropertyHandle, class IDetailChildrenBuilder& StructBuilder, IPropertyTypeCustomizationUtils& StructCustomizationUtils);
 
 private:
 	// property handles
