Changing WindowsTargetSettingsDetails to use shared functions to describe shader platforms.
Removing the Experimental label from PCD3D_SM6.

diff --git a/Engine/Source/Editor/DetailCustomizations/Private/WindowsTargetSettingsDetails.cpp b/Engine/Source/Editor/DetailCustomizations/Private/WindowsTargetSettingsDetails.cpp
--- a/Engine/Source/Editor/DetailCustomizations/Private/WindowsTargetSettingsDetails.cpp
+++ b/Engine/Source/Editor/DetailCustomizations/Private/WindowsTargetSettingsDetails.cpp
@@ -29,6 +29,7 @@
 #include "Interfaces/ITargetPlatformModule.h"
 #include "SExternalImageReference.h"
 #include "UnrealEngine.h"
+#include "RHIShaderFormatDefinitions.inl"
 
 #if WITH_ENGINE
 #include "AudioDevice.h"
@@ -49,46 +50,44 @@ namespace WindowsTargetSettingsDetailsConstants
 	const FText DisabledTip = LOCTEXT("GitHubSourceRequiredToolTip", "This requires GitHub source.");
 }
 
-static FText GetFriendlyNameFromWindowsRHIName(const FString& InRHIName)
+static FText GetFriendlyNameFromShaderPlatform(const FString& InShaderPlatformName)
 {
+	const EShaderPlatform ShaderPlatform = ShaderFormatNameToShaderPlatform(FName(*InShaderPlatformName));
+
 	FText FriendlyRHIName;
-	if (InRHIName == TEXT("PCD3D_SM6"))
-	{
-		FriendlyRHIName = LOCTEXT("DirectX12", "DirectX 12 (SM6, Experimental)");
-	}
-	else if (InRHIName == TEXT("PCD3D_SM5"))
+	switch (ShaderPlatform)
 	{
+	case SP_PCD3D_SM6:
+		FriendlyRHIName = LOCTEXT("DirectX12", "DirectX 12 (SM6)");
+		break;
+	case SP_PCD3D_SM5:
 		FriendlyRHIName = LOCTEXT("DirectX11", "DirectX 11 & 12 (SM5)");
-	}
-	else if (InRHIName == TEXT("PCD3D_ES31"))
-	{
+		break;
+	case SP_PCD3D_ES3_1:
 		FriendlyRHIName = LOCTEXT("DirectXES31", "DirectX Mobile Emulation (ES3.1)");
-	}
-	else if (InRHIName == TEXT("D3D_ES3_1_HOLOLENS"))
-	{
+		break;
+	case SP_D3D_ES3_1_HOLOLENS:
 		FriendlyRHIName = LOCTEXT("DirectXES31HL", "DirectX Hololens (ES3.1)");
-	}
-	else if (InRHIName == TEXT("SF_VULKAN_SM5"))
-	{
+		break;
+	case SP_VULKAN_SM5:
 		FriendlyRHIName = LOCTEXT("VulkanSM5", "Vulkan (SM5)");
-	}
-	else if (InRHIName == TEXT("GLSL_150_ES31")
-		|| InRHIName == TEXT("SF_VULKAN_ES31_ANDROID") || InRHIName == TEXT("SF_VULKAN_ES31")
-		|| InRHIName == TEXT("GLSL_430"))
-	{
+		break;
+
+	case SP_OPENGL_PCES3_1:
+	case SP_VULKAN_PCES3_1:
 		// Explicitly remove these formats as they are obsolete/not quite supported; users can still target them by adding them as +TargetedRHIs in the TargetPlatform ini.
 		FriendlyRHIName = FText::GetEmpty();
-	}
-	else
-	{
-		UE_LOG(LogEngine, Warning, TEXT("Unknown Windows target RHI %s"), *InRHIName);
+		break;
+
+	default:
+		UE_LOG(LogEngine, Warning, TEXT("Unknown Windows target RHI %s"), *InShaderPlatformName);
 		FriendlyRHIName = LOCTEXT("UnknownRHI", "UnknownRHI");
+		break;
 	}
 
 	return FriendlyRHIName;
 }
 
-
 TSharedRef<IDetailCustomization> FWindowsTargetSettingsDetails::MakeInstance()
 {
 	return MakeShareable(new FWindowsTargetSettingsDetails);
@@ -161,7 +160,7 @@ void FWindowsTargetSettingsDetails::CustomizeDetails( IDetailLayoutBuilder& Deta
 	// Setup the supported/targeted RHI property view
 	ITargetPlatform* TargetPlatform = FModuleManager::GetModuleChecked<ITargetPlatformModule>("WindowsTargetPlatform").GetTargetPlatforms()[0];
 	TargetShaderFormatsDetails = MakeShareable(new FShaderFormatsPropertyDetails(&DetailBuilder));
-	TargetShaderFormatsDetails->CreateTargetShaderFormatsPropertyView(TargetPlatform, GetFriendlyNameFromWindowsRHIName);
+	TargetShaderFormatsDetails->CreateTargetShaderFormatsPropertyView(TargetPlatform, GetFriendlyNameFromShaderPlatform);
 
 	// Next add the splash image customization
 	const FText EditorSplashDesc(LOCTEXT("EditorSplashLabel", "Editor Splash"));
