Fix DirectoryPath slate UI not disabling path selection button when the property is disabled.

diff --git a/Engine/Source/Editor/DetailCustomizations/Private/DirectoryPathStructCustomization.cpp b/Engine/Source/Editor/DetailCustomizations/Private/DirectoryPathStructCustomization.cpp
--- a/Engine/Source/Editor/DetailCustomizations/Private/DirectoryPathStructCustomization.cpp
+++ b/Engine/Source/Editor/DetailCustomizations/Private/DirectoryPathStructCustomization.cpp
@@ -47,6 +47,7 @@ void FDirectoryPathStructCustomization::CustomizeHeader( TSharedRef<IPropertyHan
 			.ContentPadding(2.0f)
 			.ForegroundColor( FSlateColor::UseForeground() )
 			.IsFocusable(false)
+			.IsEnabled(StructPropertyHandle->IsEditable())
 			[
 				SNew(SImage)
 				.Image(FAppStyle::GetBrush("PropertyWindow.Button_Ellipsis"))
@@ -63,6 +64,7 @@ void FDirectoryPathStructCustomization::CustomizeHeader( TSharedRef<IPropertyHan
 			.ContentPadding( 2.0f )
 			.ForegroundColor( FSlateColor::UseForeground() )
 			.IsFocusable( false )
+			.IsEnabled(StructPropertyHandle->IsEditable())
 			[
 				SNew( SImage )
 				.Image( FAppStyle::GetBrush("PropertyWindow.Button_Ellipsis") )
