Fixed level not being marked as modified by editing FTimecode property in the details panel.

diff --git a/Engine/Source/Editor/DetailCustomizations/Private/TimecodeDetailsCustomization.cpp b/Engine/Source/Editor/DetailCustomizations/Private/TimecodeDetailsCustomization.cpp
--- a/Engine/Source/Editor/DetailCustomizations/Private/TimecodeDetailsCustomization.cpp
+++ b/Engine/Source/Editor/DetailCustomizations/Private/TimecodeDetailsCustomization.cpp
@@ -60,10 +60,13 @@ void FTimecodeDetailsCustomization::OnTimecodeTextCommitted(const FText& InText,
 
 		if (Splits.Num() == 4)
 		{
+			TimecodeProperty->NotifyPreChange();
 			((FTimecode*)RawData[0])->Hours = FCString::Atoi(*Splits[0]);
 			((FTimecode*)RawData[0])->Minutes = FCString::Atoi(*Splits[1]);
 			((FTimecode*)RawData[0])->Seconds = FCString::Atoi(*Splits[2]);
 			((FTimecode*)RawData[0])->Frames = FCString::Atoi(*Splits[3]);
+			TimecodeProperty->NotifyPostChange(EPropertyChangeType::ValueSet);
+			TimecodeProperty->NotifyFinishedChangingProperties();
 		}
 		else
 		{
