Fixing NonUnity issue.

diff --git a/Engine/Source/Editor/DetailCustomizations/Private/MacTargetSettingsDetails.cpp b/Engine/Source/Editor/DetailCustomizations/Private/MacTargetSettingsDetails.cpp
--- a/Engine/Source/Editor/DetailCustomizations/Private/MacTargetSettingsDetails.cpp
+++ b/Engine/Source/Editor/DetailCustomizations/Private/MacTargetSettingsDetails.cpp
@@ -21,6 +21,7 @@
 #include "Dialogs/Dialogs.h"
 #include "Widgets/Notifications/SErrorText.h"
 #include "IDetailPropertyRow.h"
+#include "RHI.h"
 
 namespace MacTargetSettingsDetailsConstants
 {
