UI : add for dedicated section geometry collection component specifics

diff --git a/Engine/Source/Editor/DetailCustomizations/Private/DetailCustomizations.cpp b/Engine/Source/Editor/DetailCustomizations/Private/DetailCustomizations.cpp
--- a/Engine/Source/Editor/DetailCustomizations/Private/DetailCustomizations.cpp
+++ b/Engine/Source/Editor/DetailCustomizations/Private/DetailCustomizations.cpp
@@ -753,6 +753,14 @@ void FDetailCustomizationsModule::RegisterSectionMappings()
 			Section->AddCategory("Foliage");
 		}
 	}
+
+	// Geometry Collections
+	{
+		{
+			TSharedRef<FPropertySection> Section = PropertyModule.FindOrCreateSection("GeometryCollectionComponent", "GC", LOCTEXT("GC", "GC"));
+			Section->AddCategory("ChaosPhysics");
+		}
+	}
 }
 
 #undef LOCTEXT_NAMESPACE
