Remove duplicate HandleMapChanged callback which was causing refresh of WorldSettingsView on invalid WorldSettings object
Event was already handled by SLevelEditor::HandleEditorMapChange

diff --git a/Engine/Source/Editor/DetailCustomizations/Private/WorldSettingsDetails.cpp b/Engine/Source/Editor/DetailCustomizations/Private/WorldSettingsDetails.cpp
--- a/Engine/Source/Editor/DetailCustomizations/Private/WorldSettingsDetails.cpp
+++ b/Engine/Source/Editor/DetailCustomizations/Private/WorldSettingsDetails.cpp
@@ -263,7 +263,6 @@ FLightmapCustomNodeBuilder::FLightmapCustomNodeBuilder(const TSharedPtr<FAssetTh
 FLightmapCustomNodeBuilder::~FLightmapCustomNodeBuilder()
 {
 	FEditorDelegates::OnLightingBuildKept.RemoveAll(this);
-	FEditorDelegates::MapChange.RemoveAll(this);
 	FEditorDelegates::NewCurrentLevel.RemoveAll(this);
 }
 
@@ -273,7 +272,6 @@ void FLightmapCustomNodeBuilder::SetOnRebuildChildren(FSimpleDelegate InOnRegene
 	OnRegenerateChildren = InOnRegenerateChildren;
 
 	FEditorDelegates::OnLightingBuildKept.AddSP(this, &FLightmapCustomNodeBuilder::HandleLightingBuildKept);
-	FEditorDelegates::MapChange.AddSP(this, &FLightmapCustomNodeBuilder::HandleMapChanged);
 	FEditorDelegates::NewCurrentLevel.AddSP(this, &FLightmapCustomNodeBuilder::HandleNewCurrentLevel);
 }
 
@@ -323,13 +321,6 @@ void FLightmapCustomNodeBuilder::HandleLightingBuildKept()
 	OnRegenerateChildren.ExecuteIfBound();
 }
 
-
-void FLightmapCustomNodeBuilder::HandleMapChanged(uint32 MapChangeFlags)
-{
-	OnRegenerateChildren.ExecuteIfBound();
-}
-
-
 void FLightmapCustomNodeBuilder::HandleNewCurrentLevel()
 {
 	OnRegenerateChildren.ExecuteIfBound();
diff --git a/Engine/Source/Editor/DetailCustomizations/Private/WorldSettingsDetails.h b/Engine/Source/Editor/DetailCustomizations/Private/WorldSettingsDetails.h
--- a/Engine/Source/Editor/DetailCustomizations/Private/WorldSettingsDetails.h
+++ b/Engine/Source/Editor/DetailCustomizations/Private/WorldSettingsDetails.h
@@ -130,9 +130,6 @@ private:
 	/** Handler for when lighting has been rebuilt and kept */
 	void HandleLightingBuildKept();
 
-	/** Handler for when the map changed or was rebuilt */
-	void HandleMapChanged(uint32 MapChangeFlags);
-
 	/** Handler for when the current level changes */
 	void HandleNewCurrentLevel();
 
