Timecode: Mark Timecode details customization read-only for non-editable properties.

diff --git a/Engine/Source/Editor/DetailCustomizations/Private/TimecodeDetailsCustomization.cpp b/Engine/Source/Editor/DetailCustomizations/Private/TimecodeDetailsCustomization.cpp
--- a/Engine/Source/Editor/DetailCustomizations/Private/TimecodeDetailsCustomization.cpp
+++ b/Engine/Source/Editor/DetailCustomizations/Private/TimecodeDetailsCustomization.cpp
@@ -31,6 +31,7 @@ void FTimecodeDetailsCustomization::CustomizeChildren(TSharedRef<IPropertyHandle
 			.SelectAllTextWhenFocused(true)
 			.RevertTextOnEscape(true)
 			.Font(IDetailLayoutBuilder::GetDetailFont())
+			.IsReadOnly(!PropertyHandle->IsEditable())
 		];
 }
 
