Crash when changing details of an animation in an Animation Composite
#fix pass weak-object rather than capturing this in lambda, and changed others to shared object bindings where possible

diff --git a/Engine/Source/Editor/DetailCustomizations/Private/AnimMontageSegmentDetails.cpp b/Engine/Source/Editor/DetailCustomizations/Private/AnimMontageSegmentDetails.cpp
--- a/Engine/Source/Editor/DetailCustomizations/Private/AnimMontageSegmentDetails.cpp
+++ b/Engine/Source/Editor/DetailCustomizations/Private/AnimMontageSegmentDetails.cpp
@@ -113,6 +113,7 @@ void FAnimMontageSegmentDetails::CustomizeDetails( IDetailLayoutBuilder& DetailB
 			ValueWidget.ToSharedRef()
 		];
 
+	TWeakPtr<FAnimMontageSegmentDetails> WeakDetails = TWeakPtr<FAnimMontageSegmentDetails>(StaticCastSharedRef<FAnimMontageSegmentDetails>(AsShared()));
 	if (AnimSegmentHandle.IsValid())
 	{
 		AnimStartTimeProperty = AnimSegmentHandle->GetChildHandle(GET_MEMBER_NAME_CHECKED(FAnimSegment, AnimStartTime));
@@ -229,11 +230,14 @@ void FAnimMontageSegmentDetails::CustomizeDetails( IDetailLayoutBuilder& DetailB
 			.MinValue(1)
 			.MaxSliderValue(4)
 			.MaxValue(32)
-			.Value_Lambda([this]
+			.Value_Lambda([WeakDetails]
 			{
-				if(const FAnimSegment* AnimSegmentPtr = GetAnimationSegment())
+				if (const TSharedPtr<FAnimMontageSegmentDetails> SegmentDetails = WeakDetails.Pin())
 				{
-					return AnimSegmentPtr->LoopingCount;
+					if(const FAnimSegment* AnimSegmentPtr = SegmentDetails->GetAnimationSegment())
+					{
+						return AnimSegmentPtr->LoopingCount;
+					}
 				}
 				return 0;
 			})
@@ -262,12 +266,24 @@ void FAnimMontageSegmentDetails::CustomizeDetails( IDetailLayoutBuilder& DetailB
 	SegmentCategory.AddCustomRow(FText::GetEmpty(), false)
 	[
 		SNew(SAnimationSegmentViewport)
-		.AnimRef_Lambda([this]() { return GetAnimationAsset(); })
-		.StartTime_Raw(this, &FAnimMontageSegmentDetails::GetStartTime)
-		.EndTime_Raw(this, &FAnimMontageSegmentDetails::GetEndTime)
-		.PlayRate_Raw(this, &FAnimMontageSegmentDetails::GetPlayRate)
-		.OnStartTimeChanged_Lambda([this](float InValue, bool bInteractive) { OnStartTimeChanged(InValue, ETextCommit::Default, bInteractive); })
-		.OnEndTimeChanged_Lambda([this](float InValue, bool bInteractive) { OnEndTimeChanged(InValue, ETextCommit::Default, bInteractive); })
+		.AnimRef(this, &FAnimMontageSegmentDetails::GetAnimationAsset)
+		.StartTime(this, &FAnimMontageSegmentDetails::GetStartTime)
+		.EndTime(this, &FAnimMontageSegmentDetails::GetEndTime)
+		.PlayRate(this, &FAnimMontageSegmentDetails::GetPlayRate)
+		.OnStartTimeChanged_Lambda([WeakDetails](float InValue, bool bInteractive)
+		{
+			if (const TSharedPtr<FAnimMontageSegmentDetails> SegmentDetails = WeakDetails.Pin())
+			{
+				SegmentDetails->OnStartTimeChanged(InValue, ETextCommit::Default, bInteractive);
+			}
+		})
+		.OnEndTimeChanged_Lambda([WeakDetails](float InValue, bool bInteractive)
+		{
+			if (const TSharedPtr<FAnimMontageSegmentDetails> SegmentDetails = WeakDetails.Pin())
+			{
+				SegmentDetails->OnEndTimeChanged(InValue, ETextCommit::Default, bInteractive);
+			}
+		})
 	];
 }
 
