Add a warning on the old Input Action and Axis mappings settings page that you should use Enhanced Input instead, if the plugin is enabled.

#rnx

diff --git a/Engine/Source/Editor/DetailCustomizations/Private/InputSettingsDetails.cpp b/Engine/Source/Editor/DetailCustomizations/Private/InputSettingsDetails.cpp
--- a/Engine/Source/Editor/DetailCustomizations/Private/InputSettingsDetails.cpp
+++ b/Engine/Source/Editor/DetailCustomizations/Private/InputSettingsDetails.cpp
@@ -618,16 +618,55 @@ void FInputSettingsDetails::CustomizeDetails(class IDetailLayoutBuilder& DetailB
 
 	IDetailCategoryBuilder& MappingsDetailCategoryBuilder = DetailBuilder.EditCategory(BindingsCategory);
 
+	// If the new Enhanced Input module is loaded, then add a warning telling people to use
+	// that instead of these legacy Action/Axis bindings
+	static const FName ForegroundColorStyle("Colors.Foreground");
+	static const FName WarningColorStyle("Colors.AccentYellow");
+	static const FSlateBrush* WarningBrush = FAppStyle::Get().GetBrush("Icons.AlertCircle");
+	
 	MappingsDetailCategoryBuilder.AddCustomRow(LOCTEXT("Mappings_Title", "Action Axis Mappings"))
 	[
 		SNew(SHorizontalBox)
 		+SHorizontalBox::Slot()
-		.FillWidth(1)
+		.FillWidth(1.0f)
 		[
-			SNew(STextBlock)
-			.Font(IDetailLayoutBuilder::GetDetailFont())
-			.AutoWrapText(true)
-			.Text(LOCTEXT("Mappings_Description", "Action and Axis Mappings provide a mechanism to conveniently map keys and axes to input behaviors by inserting a layer of indirection between the input behavior and the keys that invoke it. Action Mappings are for key presses and releases, while Axis Mappings allow for inputs that have a continuous range."))
+			SNew(SVerticalBox)
+			+ SVerticalBox::Slot()
+			.AutoHeight()
+			.HAlign(HAlign_Center)
+			.Padding(0.0f, 10.0f)
+			[
+				SNew(SHorizontalBox)
+				+ SHorizontalBox::Slot()
+				.AutoWidth()
+				.Padding(10.0f, 10.0f)
+				[
+					SNew(SImage)
+					.Image(WarningBrush)
+					.Visibility(this, &FInputSettingsDetails::GetLegacyWarningVisibility)
+					.ColorAndOpacity(FAppStyle::Get().GetSlateColor(ForegroundColorStyle))		
+				]
+				+ SHorizontalBox::Slot()
+				.AutoWidth()
+				.VAlign(VAlign_Bottom)
+				[
+					SNew(STextBlock)
+					.Font(IDetailLayoutBuilder::GetDetailFont())
+					.AutoWrapText(true)					
+					.Visibility(this, &FInputSettingsDetails::GetLegacyWarningVisibility)
+					.Text(LOCTEXT("Mappings_DeprecationWarning", "Axis and Action mappings are now deprecated, please use Enhanced Input Actions and Input Mapping Contexts instead."))
+					.ColorAndOpacity(FAppStyle::Get().GetSlateColor(WarningColorStyle))
+				]	
+			]
+			+ SVerticalBox::Slot()
+			.AutoHeight()			
+			.HAlign(HAlign_Left)
+			[
+				SNew(STextBlock)
+				.Font(IDetailLayoutBuilder::GetDetailFont())
+				.AutoWrapText(true)
+				.Text(LOCTEXT("Mappings_Description", "Action and Axis Mappings provide a mechanism to conveniently map keys and axes to input behaviors by inserting a layer of indirection between the input behavior and the keys that invoke it. Action Mappings are for key presses and releases, while Axis Mappings allow for inputs that have a continuous range."))
+			]		
 		]
 		+SHorizontalBox::Slot()
 		.AutoWidth()
@@ -651,4 +690,9 @@ void FInputSettingsDetails::CustomizeDetails(class IDetailLayoutBuilder& DetailB
 	MappingsDetailCategoryBuilder.AddCustomBuilder(AxisMappingsBuilder);
 }
 
+EVisibility FInputSettingsDetails::GetLegacyWarningVisibility() const
+{
+	return FModuleManager::Get().IsModuleLoaded("EnhancedInput") ? EVisibility::Visible : EVisibility::Hidden;
+}
+
 #undef LOCTEXT_NAMESPACE
diff --git a/Engine/Source/Editor/DetailCustomizations/Private/InputSettingsDetails.h b/Engine/Source/Editor/DetailCustomizations/Private/InputSettingsDetails.h
--- a/Engine/Source/Editor/DetailCustomizations/Private/InputSettingsDetails.h
+++ b/Engine/Source/Editor/DetailCustomizations/Private/InputSettingsDetails.h
@@ -7,6 +7,7 @@
 #include "PropertyHandle.h"
 #include "IDetailCustomNodeBuilder.h"
 #include "IDetailCustomization.h"
+#include "Layout/Visibility.h"
 
 class FDetailWidgetRow;
 class IDetailChildrenBuilder;
@@ -125,4 +126,12 @@ public:
 
 	/** ILayoutDetails interface */
 	virtual void CustomizeDetails( class IDetailLayoutBuilder& DetailBuilder ) override;
+
+private:
+
+	/**
+	 * If true, then we should display some warning text about the Axis/Action mappings being legacy
+	 * in favor of Enhanced Input
+	 */
+	EVisibility GetLegacyWarningVisibility() const; 
 };
