Timecode: Add support for partial timecode entry in details panel.

diff --git a/Engine/Source/Editor/DetailCustomizations/Private/TimecodeDetailsCustomization.cpp b/Engine/Source/Editor/DetailCustomizations/Private/TimecodeDetailsCustomization.cpp
--- a/Engine/Source/Editor/DetailCustomizations/Private/TimecodeDetailsCustomization.cpp
+++ b/Engine/Source/Editor/DetailCustomizations/Private/TimecodeDetailsCustomization.cpp
@@ -59,13 +59,14 @@ void FTimecodeDetailsCustomization::OnTimecodeTextCommitted(const FText& InText,
 		TArray<FString> Splits;
 		InText.ToString().ParseIntoArray(Splits, TEXT(":"));
 
-		if (Splits.Num() == 4)
+		const int NumSplits = Splits.Num();
+		if (NumSplits > 0 && NumSplits <= 4)
 		{
 			TimecodeProperty->NotifyPreChange();
 			((FTimecode*)RawData[0])->Hours = FCString::Atoi(*Splits[0]);
-			((FTimecode*)RawData[0])->Minutes = FCString::Atoi(*Splits[1]);
-			((FTimecode*)RawData[0])->Seconds = FCString::Atoi(*Splits[2]);
-			((FTimecode*)RawData[0])->Frames = FCString::Atoi(*Splits[3]);
+			((FTimecode*)RawData[0])->Minutes = NumSplits > 1 ? FCString::Atoi(*Splits[1]) : 0;
+			((FTimecode*)RawData[0])->Seconds = NumSplits > 2 ? FCString::Atoi(*Splits[2]) : 0;
+			((FTimecode*)RawData[0])->Frames = NumSplits > 3 ? FCString::Atoi(*Splits[3]) : 0;
 			TimecodeProperty->NotifyPostChange(EPropertyChangeType::ValueSet);
 			TimecodeProperty->NotifyFinishedChangingProperties();
 		}
