Timecode: Add undo transaction support for Timecode entry.

diff --git a/Engine/Source/Editor/DetailCustomizations/Private/TimecodeDetailsCustomization.cpp b/Engine/Source/Editor/DetailCustomizations/Private/TimecodeDetailsCustomization.cpp
--- a/Engine/Source/Editor/DetailCustomizations/Private/TimecodeDetailsCustomization.cpp
+++ b/Engine/Source/Editor/DetailCustomizations/Private/TimecodeDetailsCustomization.cpp
@@ -62,6 +62,8 @@ void FTimecodeDetailsCustomization::OnTimecodeTextCommitted(const FText& InText,
 		const int NumSplits = Splits.Num();
 		if (NumSplits > 0 && NumSplits <= 4)
 		{
+			GEditor->BeginTransaction(FText::Format(LOCTEXT("SetTimecodeProperty", "Edit {0}"), TimecodeProperty->GetPropertyDisplayName()));
+			
 			TimecodeProperty->NotifyPreChange();
 			((FTimecode*)RawData[0])->Hours = FCString::Atoi(*Splits[0]);
 			((FTimecode*)RawData[0])->Minutes = NumSplits > 1 ? FCString::Atoi(*Splits[1]) : 0;
@@ -69,6 +71,8 @@ void FTimecodeDetailsCustomization::OnTimecodeTextCommitted(const FText& InText,
 			((FTimecode*)RawData[0])->Frames = NumSplits > 3 ? FCString::Atoi(*Splits[3]) : 0;
 			TimecodeProperty->NotifyPostChange(EPropertyChangeType::ValueSet);
 			TimecodeProperty->NotifyFinishedChangingProperties();
+
+			GEditor->EndTransaction();
 		}
 		else
 		{
