Add missing include

diff --git a/Engine/Source/Editor/DetailCustomizations/Private/TimecodeDetailsCustomization.cpp b/Engine/Source/Editor/DetailCustomizations/Private/TimecodeDetailsCustomization.cpp
--- a/Engine/Source/Editor/DetailCustomizations/Private/TimecodeDetailsCustomization.cpp
+++ b/Engine/Source/Editor/DetailCustomizations/Private/TimecodeDetailsCustomization.cpp
@@ -1,6 +1,7 @@
 // Copyright Epic Games, Inc. All Rights Reserved.
 
 #include "TimecodeDetailsCustomization.h"
+
 #include "Misc/Timecode.h"
 #include "IDetailPropertyRow.h"
 #include "IDetailChildrenBuilder.h"
@@ -8,6 +9,8 @@
 #include "Widgets/Text/STextBlock.h"
 #include "Widgets/Input/SEditableTextBox.h"
 #include "PropertyCustomizationHelpers.h"
+#include "Editor.h"
+
 
 #define LOCTEXT_NAMESPACE "Timecode"
 
