Put the "Update Source" button of the Pose Asset editor below the asset field, as it was often hidden, depending on the width of the details panel.

diff --git a/Engine/Source/Editor/DetailCustomizations/Private/PoseAssetDetails.cpp b/Engine/Source/Editor/DetailCustomizations/Private/PoseAssetDetails.cpp
--- a/Engine/Source/Editor/DetailCustomizations/Private/PoseAssetDetails.cpp
+++ b/Engine/Source/Editor/DetailCustomizations/Private/PoseAssetDetails.cpp
@@ -264,20 +264,21 @@ void FPoseAssetDetails::CustomizeDetails(IDetailLayoutBuilder& DetailBuilder)
 	.ValueContent()
 	.MinDesiredWidth(200)
 	[
-		SNew(SHorizontalBox)
-		+ SHorizontalBox::Slot()
-		.AutoWidth()
+		SNew(SVerticalBox)
+		+ SVerticalBox::Slot()
+		.HAlign(EHorizontalAlignment::HAlign_Left)
+		.AutoHeight()
 		[
 			SNew(SObjectPropertyEntryBox)
 			.AllowedClass(UAnimSequence::StaticClass())
 			.OnShouldFilterAsset(this, &FPoseAssetDetails::ShouldFilterAsset)
 			.PropertyHandle(SourceAnimationPropertyHandle)
 		]
-		+ SHorizontalBox::Slot()
-		.Padding(FMargin(0.f, 5.f))
-		.AutoWidth()
-		[			
-			SNew(SButton)			
+		+ SVerticalBox::Slot()
+		.HAlign(EHorizontalAlignment::HAlign_Left)
+		.AutoHeight()
+		[
+			SNew(SButton)
 			.Text(LOCTEXT("UpdateSource_Lable", "Update Source"))
 			.ToolTipText(LOCTEXT("UpdateSource_Tooltip", "Update Poses From Source Animation"))
 			.OnClicked(this, &FPoseAssetDetails::OnUpdatePoseSourceAnimation)
