Fix crash opening project settings in monolithic targets

diff --git a/Engine/Plugins/EnhancedInput/Source/InputEditor/Private/ActionMappingDetails.cpp b/Engine/Plugins/EnhancedInput/Source/InputEditor/Private/ActionMappingDetails.cpp
--- a/Engine/Plugins/EnhancedInput/Source/InputEditor/Private/ActionMappingDetails.cpp
+++ b/Engine/Plugins/EnhancedInput/Source/InputEditor/Private/ActionMappingDetails.cpp
@@ -377,7 +377,7 @@ void FActionMappingsNodeBuilderEx::GenerateChildContent(IDetailChildrenBuilder&
 
 	for (int32 Index = 0; Index < GroupedMappings.Num(); ++Index)
 	{
-		FMappingSet& MappingSet = GroupedMappings[Index];
+		ActionMappingDetails::FMappingSet& MappingSet = GroupedMappings[Index];
 
 		FString GroupNameString(TEXT("ActionMappings."));
 		GroupNameString += MappingSet.SharedAction->GetPathName();
@@ -399,7 +399,7 @@ void FActionMappingsNodeBuilderEx::GenerateChildContent(IDetailChildrenBuilder&
 			.AutoWidth()
 			[
 				SNew(SBox)
-				.WidthOverride(InputConstants::TextBoxWidth)
+				.WidthOverride(ActionMappingDetails::InputConstants::TextBoxWidth)
 				[
 					SNew(SObjectPropertyEntryBox)
 					.AllowedClass(UInputAction::StaticClass())
@@ -409,7 +409,7 @@ void FActionMappingsNodeBuilderEx::GenerateChildContent(IDetailChildrenBuilder&
 				]
 			]
 			+ SHorizontalBox::Slot()
-			.Padding(InputConstants::PropertyPadding)
+			.Padding(ActionMappingDetails::InputConstants::PropertyPadding)
 			.HAlign(HAlign_Center)
 			.VAlign(VAlign_Center)
 			.AutoWidth()
@@ -417,7 +417,7 @@ void FActionMappingsNodeBuilderEx::GenerateChildContent(IDetailChildrenBuilder&
 				AddButton
 			]
 			+ SHorizontalBox::Slot()
-			.Padding(InputConstants::PropertyPadding)
+			.Padding(ActionMappingDetails::InputConstants::PropertyPadding)
 			.HAlign(HAlign_Center)
 			.VAlign(VAlign_Center)
 			.AutoWidth()
@@ -461,7 +461,7 @@ void FActionMappingsNodeBuilderEx::ClearActionMappingButton_OnClick()
 	ActionMappingsPropertyHandle->AsArray()->EmptyArray();
 }
 
-void FActionMappingsNodeBuilderEx::OnActionMappingActionChanged(const FAssetData& AssetData, const FMappingSet MappingSet)
+void FActionMappingsNodeBuilderEx::OnActionMappingActionChanged(const FAssetData& AssetData, const ActionMappingDetails::FMappingSet MappingSet)
 {
 	const FScopedTransaction Transaction(LOCTEXT("SwitchActionMapping_Transaction", "Switch Action Mapping"));
 
@@ -490,7 +490,7 @@ void FActionMappingsNodeBuilderEx::OnActionMappingActionChanged(const FAssetData
 	}
 }
 
-void FActionMappingsNodeBuilderEx::AddActionMappingToGroupButton_OnClick(const FMappingSet MappingSet)
+void FActionMappingsNodeBuilderEx::AddActionMappingToGroupButton_OnClick(const ActionMappingDetails::FMappingSet MappingSet)
 {
 	const FScopedTransaction Transaction(LOCTEXT("AddActionMappingToGroup_Transaction", "Add a control binding to the Action Mapping"));
 
@@ -510,7 +510,7 @@ void FActionMappingsNodeBuilderEx::AddActionMappingToGroupButton_OnClick(const F
 	}
 }
 
-void FActionMappingsNodeBuilderEx::RemoveActionMappingGroupButton_OnClick(const FMappingSet MappingSet)
+void FActionMappingsNodeBuilderEx::RemoveActionMappingGroupButton_OnClick(const ActionMappingDetails::FMappingSet MappingSet)
 {
 	const FScopedTransaction Transaction(LOCTEXT("RemoveActionMappingGroup_Transaction", "Remove Action Mapping and all control bindings"));
 
@@ -539,7 +539,7 @@ bool FActionMappingsNodeBuilderEx::GroupsRequireRebuild() const
 	
 	for (int32 GroupIndex = 0; GroupIndex < GroupedMappings.Num(); ++GroupIndex)
 	{
-		const FMappingSet& MappingSet = GroupedMappings[GroupIndex];
+		const ActionMappingDetails::FMappingSet& MappingSet = GroupedMappings[GroupIndex];
 		for (int32 MappingIndex = 0; MappingIndex < MappingSet.Mappings.Num(); ++MappingIndex)
 		{
 			const UObject* Action;
diff --git a/Engine/Plugins/EnhancedInput/Source/InputEditor/Private/ActionMappingDetails.h b/Engine/Plugins/EnhancedInput/Source/InputEditor/Private/ActionMappingDetails.h
--- a/Engine/Plugins/EnhancedInput/Source/InputEditor/Private/ActionMappingDetails.h
+++ b/Engine/Plugins/EnhancedInput/Source/InputEditor/Private/ActionMappingDetails.h
@@ -14,20 +14,24 @@ class FDetailWidgetRow;
 class IDetailChildrenBuilder;
 class IDetailGroup;
 class IDetailLayoutBuilder;
+class UInputAction;
 
-namespace InputConstants
+namespace ActionMappingDetails
 {
-	const FMargin PropertyPadding(2.0f, 0.0f, 2.0f, 0.0f);
-	const float TextBoxWidth = 250.0f;
-	const float ScaleBoxWidth = 50.0f;
-}
+	namespace InputConstants
+	{
+		const FMargin PropertyPadding(2.0f, 0.0f, 2.0f, 0.0f);
+		const float TextBoxWidth = 250.0f;
+		const float ScaleBoxWidth = 50.0f;
+	}
 
-struct FMappingSet
-{
-	const class UInputAction* SharedAction;
-	IDetailGroup* DetailGroup;
-	TArray<TSharedRef<IPropertyHandle>> Mappings;
-};
+	struct FMappingSet
+	{
+		const class UInputAction* SharedAction;
+		IDetailGroup* DetailGroup;
+		TArray<TSharedRef<IPropertyHandle>> Mappings;
+	};
+}
 
 class FActionMappingsNodeBuilderEx : public IDetailCustomNodeBuilder, public TSharedFromThis<FActionMappingsNodeBuilderEx>
 {
@@ -49,9 +53,9 @@ public:
 private:
 	void AddActionMappingButton_OnClick();
 	void ClearActionMappingButton_OnClick();
-	void OnActionMappingActionChanged(const FAssetData& AssetData, const FMappingSet MappingSet);
-	void AddActionMappingToGroupButton_OnClick(const FMappingSet MappingSet);
-	void RemoveActionMappingGroupButton_OnClick(const FMappingSet MappingSet);
+	void OnActionMappingActionChanged(const FAssetData& AssetData, const ActionMappingDetails::FMappingSet MappingSet);
+	void AddActionMappingToGroupButton_OnClick(const ActionMappingDetails::FMappingSet MappingSet);
+	void RemoveActionMappingGroupButton_OnClick(const ActionMappingDetails::FMappingSet MappingSet);
 
 	bool GroupsRequireRebuild() const;
 	void RebuildGroupedMappings();
@@ -84,7 +88,7 @@ private:
 	/** Property handle to associated Action Mappings */
 	TSharedPtr<IPropertyHandle> ActionMappingsPropertyHandle;
 
-	TArray<FMappingSet> GroupedMappings;
+	TArray<ActionMappingDetails::FMappingSet> GroupedMappings;
 
 	TArray<TPair<const UInputAction*, bool>> DelayedGroupExpansionStates;
 
diff --git a/Engine/Source/Editor/DetailCustomizations/Private/InputSettingsDetails.cpp b/Engine/Source/Editor/DetailCustomizations/Private/InputSettingsDetails.cpp
--- a/Engine/Source/Editor/DetailCustomizations/Private/InputSettingsDetails.cpp
+++ b/Engine/Source/Editor/DetailCustomizations/Private/InputSettingsDetails.cpp
@@ -84,7 +84,7 @@ void FActionMappingsNodeBuilder::GenerateChildContent( IDetailChildrenBuilder& C
 
 	for (int32 Index = 0; Index < GroupedMappings.Num(); ++Index)
 	{
-		FMappingSet& MappingSet = GroupedMappings[Index];
+		InputSettingsDetails::FMappingSet& MappingSet = GroupedMappings[Index];
 
 		FString GroupNameString(TEXT("ActionMappings."));
 		MappingSet.SharedName.AppendString(GroupNameString);
@@ -105,7 +105,7 @@ void FActionMappingsNodeBuilder::GenerateChildContent( IDetailChildrenBuilder& C
 			.AutoWidth()
 			[
 				SNew(SBox)
-				.WidthOverride( InputConstants::TextBoxWidth )
+				.WidthOverride(InputSettingsDetails::InputConstants::TextBoxWidth)
 				[
 					SNew(SEditableTextBox)
 					.Padding(2.0f)
@@ -115,7 +115,7 @@ void FActionMappingsNodeBuilder::GenerateChildContent( IDetailChildrenBuilder& C
 				]
 			]
 			+SHorizontalBox::Slot()
-			.Padding(InputConstants::PropertyPadding)
+			.Padding(InputSettingsDetails::InputConstants::PropertyPadding)
 			.HAlign(HAlign_Center)
 			.VAlign(VAlign_Center)
 			.AutoWidth()
@@ -123,7 +123,7 @@ void FActionMappingsNodeBuilder::GenerateChildContent( IDetailChildrenBuilder& C
 				AddButton
 			]
 			+SHorizontalBox::Slot()
-			.Padding(InputConstants::PropertyPadding)
+			.Padding(InputSettingsDetails::InputConstants::PropertyPadding)
 			.HAlign(HAlign_Center)
 			.VAlign(VAlign_Center)
 			.AutoWidth()
@@ -168,7 +168,7 @@ void FActionMappingsNodeBuilder::ClearActionMappingButton_OnClick()
 	ActionMappingsPropertyHandle->AsArray()->EmptyArray();
 }
 
-void FActionMappingsNodeBuilder::OnActionMappingNameCommitted(const FText& InName, ETextCommit::Type CommitInfo, const FMappingSet MappingSet)
+void FActionMappingsNodeBuilder::OnActionMappingNameCommitted(const FText& InName, ETextCommit::Type CommitInfo, const InputSettingsDetails::FMappingSet MappingSet)
 {
 	const FScopedTransaction Transaction(LOCTEXT("RenameActionMapping_Transaction", "Rename Action Mapping"));
 
@@ -197,7 +197,7 @@ void FActionMappingsNodeBuilder::OnActionMappingNameCommitted(const FText& InNam
 	}
 }
 
-void FActionMappingsNodeBuilder::AddActionMappingToGroupButton_OnClick(const FMappingSet MappingSet)
+void FActionMappingsNodeBuilder::AddActionMappingToGroupButton_OnClick(const InputSettingsDetails::FMappingSet MappingSet)
 {
 	const FScopedTransaction Transaction(LOCTEXT("AddActionMappingToGroup_Transaction", "Add Action Mapping To Group"));
 
@@ -218,7 +218,7 @@ void FActionMappingsNodeBuilder::AddActionMappingToGroupButton_OnClick(const FMa
 	}
 }
 
-void FActionMappingsNodeBuilder::RemoveActionMappingGroupButton_OnClick(const FMappingSet MappingSet)
+void FActionMappingsNodeBuilder::RemoveActionMappingGroupButton_OnClick(const InputSettingsDetails::FMappingSet MappingSet)
 {
 	const FScopedTransaction Transaction(LOCTEXT("RemoveActionMappingGroup_Transaction", "Remove Action Mapping Group"));
 
@@ -241,7 +241,7 @@ bool FActionMappingsNodeBuilder::GroupsRequireRebuild() const
 {
 	for (int32 GroupIndex = 0; GroupIndex < GroupedMappings.Num(); ++GroupIndex)
 	{
-		const FMappingSet& MappingSet = GroupedMappings[GroupIndex];
+		const InputSettingsDetails::FMappingSet& MappingSet = GroupedMappings[GroupIndex];
 		for (int32 MappingIndex = 0; MappingIndex < MappingSet.Mappings.Num(); ++MappingIndex)
 		{
 			FName ActionName;
@@ -374,7 +374,7 @@ void FAxisMappingsNodeBuilder::GenerateChildContent( IDetailChildrenBuilder& Chi
 
 	for (int32 Index = 0; Index < GroupedMappings.Num(); ++Index)
 	{
-		FMappingSet& MappingSet = GroupedMappings[Index];
+		InputSettingsDetails::FMappingSet& MappingSet = GroupedMappings[Index];
 
 		FString GroupNameString(TEXT("AxisMappings."));
 		MappingSet.SharedName.AppendString(GroupNameString);
@@ -395,7 +395,7 @@ void FAxisMappingsNodeBuilder::GenerateChildContent( IDetailChildrenBuilder& Chi
 			.AutoWidth()
 			[
 				SNew( SBox )
-				.WidthOverride( InputConstants::TextBoxWidth )
+				.WidthOverride(InputSettingsDetails::InputConstants::TextBoxWidth)
 				[
 					SNew(SEditableTextBox)
 					.Padding(2.0f)
@@ -405,7 +405,7 @@ void FAxisMappingsNodeBuilder::GenerateChildContent( IDetailChildrenBuilder& Chi
 				]
 			]
 			+SHorizontalBox::Slot()
-			.Padding(InputConstants::PropertyPadding)
+			.Padding(InputSettingsDetails::InputConstants::PropertyPadding)
 			.HAlign(HAlign_Center)
 			.VAlign(VAlign_Center)
 			.AutoWidth()
@@ -413,7 +413,7 @@ void FAxisMappingsNodeBuilder::GenerateChildContent( IDetailChildrenBuilder& Chi
 				AddButton
 			]
 			+SHorizontalBox::Slot()
-			.Padding(InputConstants::PropertyPadding)
+			.Padding(InputSettingsDetails::InputConstants::PropertyPadding)
 			.HAlign(HAlign_Center)
 			.VAlign(VAlign_Center)
 			.AutoWidth()
@@ -459,7 +459,7 @@ void FAxisMappingsNodeBuilder::ClearAxisMappingButton_OnClick()
 	AxisMappingsPropertyHandle->AsArray()->EmptyArray();
 }
 
-void FAxisMappingsNodeBuilder::OnAxisMappingNameCommitted(const FText& InName, ETextCommit::Type CommitInfo, const FMappingSet MappingSet)
+void FAxisMappingsNodeBuilder::OnAxisMappingNameCommitted(const FText& InName, ETextCommit::Type CommitInfo, const InputSettingsDetails::FMappingSet MappingSet)
 {
 	const FScopedTransaction Transaction(LOCTEXT("RenameAxisMapping_Transaction", "Rename Axis Mapping"));
 
@@ -488,7 +488,7 @@ void FAxisMappingsNodeBuilder::OnAxisMappingNameCommitted(const FText& InName, E
 	}
 }
 
-void FAxisMappingsNodeBuilder::AddAxisMappingToGroupButton_OnClick(const FMappingSet MappingSet)
+void FAxisMappingsNodeBuilder::AddAxisMappingToGroupButton_OnClick(const InputSettingsDetails::FMappingSet MappingSet)
 {
 	const FScopedTransaction Transaction(LOCTEXT("AddAxisMappingToGroup_Transaction", "Add Axis Mapping To Group"));
 
@@ -509,7 +509,7 @@ void FAxisMappingsNodeBuilder::AddAxisMappingToGroupButton_OnClick(const FMappin
 	}
 }
 
-void FAxisMappingsNodeBuilder::RemoveAxisMappingGroupButton_OnClick(const FMappingSet MappingSet)
+void FAxisMappingsNodeBuilder::RemoveAxisMappingGroupButton_OnClick(const InputSettingsDetails::FMappingSet MappingSet)
 {
 	const FScopedTransaction Transaction(LOCTEXT("RemoveAxisMappingGroup_Transaction", "Remove Axis Mapping Group"));
 
@@ -532,7 +532,7 @@ bool FAxisMappingsNodeBuilder::GroupsRequireRebuild() const
 {
 	for (int32 GroupIndex = 0; GroupIndex < GroupedMappings.Num(); ++GroupIndex)
 	{
-		const FMappingSet& MappingSet = GroupedMappings[GroupIndex];
+		const InputSettingsDetails::FMappingSet& MappingSet = GroupedMappings[GroupIndex];
 		for (int32 MappingIndex = 0; MappingIndex < MappingSet.Mappings.Num(); ++MappingIndex)
 		{
 			FName AxisName;
diff --git a/Engine/Source/Editor/DetailCustomizations/Private/InputSettingsDetails.h b/Engine/Source/Editor/DetailCustomizations/Private/InputSettingsDetails.h
--- a/Engine/Source/Editor/DetailCustomizations/Private/InputSettingsDetails.h
+++ b/Engine/Source/Editor/DetailCustomizations/Private/InputSettingsDetails.h
@@ -14,19 +14,22 @@ class IDetailChildrenBuilder;
 class IDetailGroup;
 class IDetailLayoutBuilder;
 
-namespace InputConstants
+namespace InputSettingsDetails
 {
-	const FMargin PropertyPadding(2.0f, 0.0f, 2.0f, 0.0f);
-	const float TextBoxWidth = 250.0f;
-	const float ScaleBoxWidth = 50.0f;
-}
+	namespace InputConstants
+	{
+		const FMargin PropertyPadding(2.0f, 0.0f, 2.0f, 0.0f);
+		const float TextBoxWidth = 250.0f;
+		const float ScaleBoxWidth = 50.0f;
+	}
 
-struct FMappingSet
-{
-	FName SharedName;
-	IDetailGroup* DetailGroup;
-	TArray<TSharedRef<IPropertyHandle>> Mappings;
-};
+	struct FMappingSet
+	{
+		FName SharedName;
+		IDetailGroup* DetailGroup;
+		TArray<TSharedRef<IPropertyHandle>> Mappings;
+	};
+}
 
 class FActionMappingsNodeBuilder : public IDetailCustomNodeBuilder, public TSharedFromThis<FActionMappingsNodeBuilder>
 {
@@ -45,9 +48,9 @@ public:
 private:
 	void AddActionMappingButton_OnClick();
 	void ClearActionMappingButton_OnClick();
-	void OnActionMappingNameCommitted(const FText& InName, ETextCommit::Type CommitInfo, const FMappingSet MappingSet);
-	void AddActionMappingToGroupButton_OnClick(const FMappingSet MappingSet);
-	void RemoveActionMappingGroupButton_OnClick(const FMappingSet MappingSet);
+	void OnActionMappingNameCommitted(const FText& InName, ETextCommit::Type CommitInfo, const InputSettingsDetails::FMappingSet MappingSet);
+	void AddActionMappingToGroupButton_OnClick(const InputSettingsDetails::FMappingSet MappingSet);
+	void RemoveActionMappingGroupButton_OnClick(const InputSettingsDetails::FMappingSet MappingSet);
 
 	bool GroupsRequireRebuild() const;
 	void RebuildGroupedMappings();
@@ -68,7 +71,7 @@ private:
 	/** Property handle to associated action mappings */
 	TSharedPtr<IPropertyHandle> ActionMappingsPropertyHandle;
 
-	TArray<FMappingSet> GroupedMappings;
+	TArray<InputSettingsDetails::FMappingSet> GroupedMappings;
 
 	TArray<TPair<FName, bool>> DelayedGroupExpansionStates;
 };
@@ -90,9 +93,9 @@ public:
 private:
 	void AddAxisMappingButton_OnClick();
 	void ClearAxisMappingButton_OnClick();
-	void OnAxisMappingNameCommitted(const FText& InName, ETextCommit::Type CommitInfo, const FMappingSet MappingSet);
-	void AddAxisMappingToGroupButton_OnClick(const FMappingSet MappingSet);
-	void RemoveAxisMappingGroupButton_OnClick(const FMappingSet MappingSet);
+	void OnAxisMappingNameCommitted(const FText& InName, ETextCommit::Type CommitInfo, const InputSettingsDetails::FMappingSet MappingSet);
+	void AddAxisMappingToGroupButton_OnClick(const InputSettingsDetails::FMappingSet MappingSet);
+	void RemoveAxisMappingGroupButton_OnClick(const InputSettingsDetails::FMappingSet MappingSet);
 
 	bool GroupsRequireRebuild() const;
 	void RebuildGroupedMappings();
@@ -113,7 +116,7 @@ private:
 	/** Property handle to associated axis mappings */
 	TSharedPtr<IPropertyHandle> AxisMappingsPropertyHandle;
 
-	TArray<FMappingSet> GroupedMappings;
+	TArray<InputSettingsDetails::FMappingSet> GroupedMappings;
 
 	TArray<TPair<FName, bool>> DelayedGroupExpansionStates;
 };
diff --git a/Engine/Source/Editor/DetailCustomizations/Private/InputStructCustomization.cpp b/Engine/Source/Editor/DetailCustomizations/Private/InputStructCustomization.cpp
--- a/Engine/Source/Editor/DetailCustomizations/Private/InputStructCustomization.cpp
+++ b/Engine/Source/Editor/DetailCustomizations/Private/InputStructCustomization.cpp
@@ -73,17 +73,17 @@ void FInputActionMappingCustomization::CustomizeChildren( TSharedRef<class IProp
 	[
 		SNew(SHorizontalBox)
 		+ SHorizontalBox::Slot()
-		.Padding(InputConstants::PropertyPadding)
+		.Padding(InputSettingsDetails::InputConstants::PropertyPadding)
 		.AutoWidth()
 		[
 			SNew( SBox )
-			.WidthOverride( InputConstants::TextBoxWidth )
+			.WidthOverride(InputSettingsDetails::InputConstants::TextBoxWidth)
 			[
 				StructBuilder.GenerateStructValueWidget( KeyHandle.ToSharedRef() )
 			]
 		]
 		+ SHorizontalBox::Slot()
-		.Padding(InputConstants::PropertyPadding)
+		.Padding(InputSettingsDetails::InputConstants::PropertyPadding)
 		.HAlign(HAlign_Center)
 		.VAlign(VAlign_Center)
 		.AutoWidth()
@@ -91,7 +91,7 @@ void FInputActionMappingCustomization::CustomizeChildren( TSharedRef<class IProp
 			ShiftHandle->CreatePropertyNameWidget()
 		]
 		+ SHorizontalBox::Slot()
-		.Padding(InputConstants::PropertyPadding)
+		.Padding(InputSettingsDetails::InputConstants::PropertyPadding)
 		.HAlign(HAlign_Center)
 		.VAlign(VAlign_Center)
 		.AutoWidth()
@@ -99,7 +99,7 @@ void FInputActionMappingCustomization::CustomizeChildren( TSharedRef<class IProp
 			ShiftHandle->CreatePropertyValueWidget()
 		]
 		+ SHorizontalBox::Slot()
-		.Padding(InputConstants::PropertyPadding)
+		.Padding(InputSettingsDetails::InputConstants::PropertyPadding)
 		.HAlign(HAlign_Center)
 		.VAlign(VAlign_Center)
 		.AutoWidth()
@@ -107,7 +107,7 @@ void FInputActionMappingCustomization::CustomizeChildren( TSharedRef<class IProp
 			CtrlHandle->CreatePropertyNameWidget()
 		]
 		+ SHorizontalBox::Slot()
-		.Padding(InputConstants::PropertyPadding)
+		.Padding(InputSettingsDetails::InputConstants::PropertyPadding)
 		.HAlign(HAlign_Center)
 		.VAlign(VAlign_Center)
 		.AutoWidth()
@@ -115,7 +115,7 @@ void FInputActionMappingCustomization::CustomizeChildren( TSharedRef<class IProp
 			CtrlHandle->CreatePropertyValueWidget()
 		]
 		+ SHorizontalBox::Slot()
-		.Padding(InputConstants::PropertyPadding)
+		.Padding(InputSettingsDetails::InputConstants::PropertyPadding)
 		.HAlign(HAlign_Center)
 		.VAlign(VAlign_Center)
 		.AutoWidth()
@@ -123,7 +123,7 @@ void FInputActionMappingCustomization::CustomizeChildren( TSharedRef<class IProp
 			AltHandle->CreatePropertyNameWidget()
 		]
 		+ SHorizontalBox::Slot()
-		.Padding(InputConstants::PropertyPadding)
+		.Padding(InputSettingsDetails::InputConstants::PropertyPadding)
 		.HAlign(HAlign_Center)
 		.VAlign(VAlign_Center)
 		.AutoWidth()
@@ -131,7 +131,7 @@ void FInputActionMappingCustomization::CustomizeChildren( TSharedRef<class IProp
 			AltHandle->CreatePropertyValueWidget()
 		]
 		+ SHorizontalBox::Slot()
-		.Padding(InputConstants::PropertyPadding)
+		.Padding(InputSettingsDetails::InputConstants::PropertyPadding)
 		.HAlign(HAlign_Center)
 		.VAlign(VAlign_Center)
 		.AutoWidth()
@@ -139,7 +139,7 @@ void FInputActionMappingCustomization::CustomizeChildren( TSharedRef<class IProp
 			CmdHandle->CreatePropertyNameWidget()
 		]
 		+ SHorizontalBox::Slot()
-		.Padding(InputConstants::PropertyPadding)
+		.Padding(InputSettingsDetails::InputConstants::PropertyPadding)
 		.HAlign(HAlign_Center)
 		.VAlign(VAlign_Center)
 		.AutoWidth()
@@ -147,7 +147,7 @@ void FInputActionMappingCustomization::CustomizeChildren( TSharedRef<class IProp
 			CmdHandle->CreatePropertyValueWidget()
 		]
 		+ SHorizontalBox::Slot()
-		.Padding(InputConstants::PropertyPadding)
+		.Padding(InputSettingsDetails::InputConstants::PropertyPadding)
 		.HAlign(HAlign_Center)
 		.VAlign(VAlign_Center)
 		.AutoWidth()
@@ -196,17 +196,17 @@ void FInputAxisMappingCustomization::CustomizeChildren( TSharedRef<class IProper
 	[
 		SNew(SHorizontalBox)
 		+ SHorizontalBox::Slot()
-		.Padding(InputConstants::PropertyPadding)
+		.Padding(InputSettingsDetails::InputConstants::PropertyPadding)
 		.AutoWidth()
 		[
 			SNew( SBox )
-			.WidthOverride( InputConstants::TextBoxWidth )
+			.WidthOverride(InputSettingsDetails::InputConstants::TextBoxWidth)
 			[
 				StructBuilder.GenerateStructValueWidget( KeyHandle.ToSharedRef() )
 			]
 		]
 		+SHorizontalBox::Slot()
-		.Padding(InputConstants::PropertyPadding)
+		.Padding(InputSettingsDetails::InputConstants::PropertyPadding)
 		.HAlign(HAlign_Center)
 		.VAlign(VAlign_Center)
 		.AutoWidth()
@@ -214,19 +214,19 @@ void FInputAxisMappingCustomization::CustomizeChildren( TSharedRef<class IProper
 			ScaleHandle->CreatePropertyNameWidget()
 		]
 		+SHorizontalBox::Slot()
-		.Padding(InputConstants::PropertyPadding)
+		.Padding(InputSettingsDetails::InputConstants::PropertyPadding)
 		.HAlign(HAlign_Left)
 		.VAlign(VAlign_Center)
 		.AutoWidth()
 		[
 			SNew(SBox)
-			.WidthOverride(InputConstants::ScaleBoxWidth)
+			.WidthOverride(InputSettingsDetails::InputConstants::ScaleBoxWidth)
 			[
 				ScaleHandle->CreatePropertyValueWidget()
 			]
 		]
 		+SHorizontalBox::Slot()
-		.Padding(InputConstants::PropertyPadding)
+		.Padding(InputSettingsDetails::InputConstants::PropertyPadding)
 		.HAlign(HAlign_Center)
 		.VAlign(VAlign_Center)
 		.AutoWidth()
diff --git a/Engine/Source/Editor/DetailCustomizations/Private/KeyStructCustomization.cpp b/Engine/Source/Editor/DetailCustomizations/Private/KeyStructCustomization.cpp
--- a/Engine/Source/Editor/DetailCustomizations/Private/KeyStructCustomization.cpp
+++ b/Engine/Source/Editor/DetailCustomizations/Private/KeyStructCustomization.cpp
@@ -52,7 +52,7 @@ void FKeyStructCustomization::CustomizeHeaderOnlyWithButton(TSharedRef<class IPr
 	[
 		SNew(SHorizontalBox)
 		+ SHorizontalBox::Slot()
-		.Padding(InputConstants::PropertyPadding)
+		.Padding(InputSettingsDetails::InputConstants::PropertyPadding)
 		//.AutoWidth()
 		[
 			SNew(SKeySelector)
@@ -63,7 +63,7 @@ void FKeyStructCustomization::CustomizeHeaderOnlyWithButton(TSharedRef<class IPr
 		    .FilterBlueprintBindable(false)
 		]
 		+ SHorizontalBox::Slot()
-		.Padding(InputConstants::PropertyPadding)
+		.Padding(InputSettingsDetails::InputConstants::PropertyPadding)
 		.HAlign(HAlign_Center)
 		.VAlign(VAlign_Center)
 		.AutoWidth()
