Inclusive Terminology Push
Renamed terms at DetailCustomizations

diff --git a/Engine/Source/Editor/DetailCustomizations/Private/DetailCustomizations.cpp b/Engine/Source/Editor/DetailCustomizations/Private/DetailCustomizations.cpp
--- a/Engine/Source/Editor/DetailCustomizations/Private/DetailCustomizations.cpp
+++ b/Engine/Source/Editor/DetailCustomizations/Private/DetailCustomizations.cpp
@@ -682,7 +682,7 @@ void FDetailCustomizationsModule::RegisterSectionMappings()
 			TSharedRef<FPropertySection> Section = PropertyModule.FindOrCreateSection("SkeletalMeshComponent", "Animation", LOCTEXT("Animation", "Animation"));
 			Section->AddCategory("Animation");
 			Section->AddCategory("Animation Rig");
-			Section->AddCategory("Master Pose Component");
+			Section->AddCategory("Leader Pose Component");
 		}
 
 		{
