[UE][FEATURE] PhAT Constraint Transform Editing
- Enable users to view and edit the parent and child frame transforms of constraints in the PhAT details panel
- Constraint transforms are represented as a position and a rotator in the details panel (internally they are represented as a position and a pair of orthonormal vectors)
- Constraint transforms can optionally be displayed in the frame of the assocaited bone OR relative to the default (snapped) transforms
- Allow parent / child transform position / orientation componenets to be 'snapped' to defaults individually
- Support using modifier keys to manipulate only the child frame transform via a viewport widget

diff --git a/Engine/Source/Developer/PhysicsUtilities/Private/PhysicsAssetUtils.cpp b/Engine/Source/Developer/PhysicsUtilities/Private/PhysicsAssetUtils.cpp
--- a/Engine/Source/Developer/PhysicsUtilities/Private/PhysicsAssetUtils.cpp
+++ b/Engine/Source/Developer/PhysicsUtilities/Private/PhysicsAssetUtils.cpp
@@ -244,18 +244,12 @@ bool CreateFromSkeletalMeshInternal(UPhysicsAsset* PhysicsAsset, USkeletalMesh*
 				// create joint to parent body
 				if (Params.bCreateConstraints)
 				{
-					// Transform of child from parent is just child ref-pose entry.
-					FTransform RelTM = FTransform::Identity;
-
 					int32 ParentIndex = BoneIndex;
 					int32 ParentBodyIndex = INDEX_NONE;
 					FName ParentName;
 
 					do
 					{
-						// Transform of child from parent is just child ref-pose entry.
-						RelTM = RelTM * LocalPose[ParentIndex];
-
 						//Travel up the hierarchy to find a parent which has a valid body
 						ParentIndex = SkelMesh->GetRefSkeleton().GetParentIndex(ParentIndex);
 						if(ParentIndex != INDEX_NONE)
@@ -284,14 +278,8 @@ bool CreateFromSkeletalMeshInternal(UPhysicsAsset* PhysicsAsset, USkeletalMesh*
 
 						// Place joint at origin of child
 						CS->DefaultInstance.ConstraintBone1 = BoneName;
-						CS->DefaultInstance.Pos1 = FVector::ZeroVector;
-						CS->DefaultInstance.PriAxis1 = FVector(1, 0, 0);
-						CS->DefaultInstance.SecAxis1 = FVector(0, 1, 0);
-
 						CS->DefaultInstance.ConstraintBone2 = ParentName;
-						CS->DefaultInstance.Pos2 = RelTM.GetLocation();
-						CS->DefaultInstance.PriAxis2 = RelTM.GetUnitAxis(EAxis::X);
-						CS->DefaultInstance.SecAxis2 = RelTM.GetUnitAxis(EAxis::Y);
+						CS->DefaultInstance.SnapTransformsToDefault(EConstraintTransformComponentFlags::All, PhysicsAsset);
 
 						CS->SetDefaultProfile(CS->DefaultInstance);
 
diff --git a/Engine/Source/Editor/DetailCustomizations/Private/PhysicsConstraintComponentDetails.cpp b/Engine/Source/Editor/DetailCustomizations/Private/PhysicsConstraintComponentDetails.cpp
--- a/Engine/Source/Editor/DetailCustomizations/Private/PhysicsConstraintComponentDetails.cpp
+++ b/Engine/Source/Editor/DetailCustomizations/Private/PhysicsConstraintComponentDetails.cpp
@@ -18,19 +18,669 @@
 #include "PropertyHandle.h"
 #include "DetailLayoutBuilder.h"
 #include "DetailWidgetRow.h"
+#include "HAL/PlatformApplicationMisc.h"
 #include "IDetailGroup.h"
 #include "IDetailPropertyRow.h"
 #include "DetailCategoryBuilder.h"
+#include "Math/Rotator.h"
+#include "Math/UnrealMathUtility.h"
+#include "PhysicsAssetRenderUtils.h"
 #include "PhysicsEngine/ConstraintTypes.h"
 #include "PhysicsEngine/ConstraintDrives.h"
-#include "PhysicsEngine/ConstraintInstance.h"
+#include "PhysicsEngine/PhysicsAsset.h"
 #include "PhysicsEngine/PhysicsConstraintTemplate.h"
 #include "PhysicsEngine/PhysicsConstraintComponent.h"
+#include "VectorUtil.h"
 #include "Widgets/Input/SNumericEntryBox.h"
 #include "ScopedTransaction.h"
 #include "Widgets/Text/SRichTextBlock.h"
 
 #define LOCTEXT_NAMESPACE "PhysicsConstraintComponentDetails"
+
+namespace ConstraintCopyPasteStringTokens
+{
+	static const FString Position = "Position";
+	static const FString Rotation = "Rotation";
+};
+
+// File scope utility functions.
+
+bool IsNearlyEqual(const FVector A, const FVector B, float ErrorTolerance)
+{
+	return FMath::IsNearlyEqual(A.X, B.X, ErrorTolerance) && FMath::IsNearlyEqual(A.Y, B.Y, ErrorTolerance) && FMath::IsNearlyEqual(A.Z, B.Z, ErrorTolerance);
+}
+
+template< typename TPropertyValueType > bool IsChildPropertyNearlyEqual(TSharedPtr<IPropertyHandle> ParentPropertyHandle, const FName& LhsPropertyName, const TPropertyValueType& RhsValue)
+{
+	if (ParentPropertyHandle.IsValid())
+	{
+		TSharedPtr<IPropertyHandle> LhsPropertyHandle = ParentPropertyHandle->GetChildHandle(LhsPropertyName);
+
+		if (LhsPropertyHandle.IsValid())
+		{
+			TPropertyValueType LhsValue;
+			LhsPropertyHandle->GetValue(LhsValue);
+
+			return IsNearlyEqual(LhsValue, RhsValue, SMALL_NUMBER);
+		}
+	}
+
+	return false;
+}
+
+// Class FOrthonormalVectorPairStructCustomization.
+FConstraintTransformCustomization::FConstraintTransformCustomization()
+	: CachedRotation(MakeShareable(new TProxyValue<FRotator>(FRotator::ZeroRotator)))
+	, CachedRotationYaw(MakeShareable(new TProxyProperty<FRotator, FRotator::FReal>(CachedRotation, CachedRotation->Get().Yaw)))
+	, CachedRotationPitch(MakeShareable(new TProxyProperty<FRotator, FRotator::FReal>(CachedRotation, CachedRotation->Get().Pitch)))
+	, CachedRotationRoll(MakeShareable(new TProxyProperty<FRotator, FRotator::FReal>(CachedRotation, CachedRotation->Get().Roll)))
+	, CachedPosition(MakeShareable(new TProxyValue<FVector>(FVector::ZeroVector)))
+	, CachedPositionX(MakeShareable(new TProxyProperty<FVector, FRotator::FReal>(CachedPosition, CachedPosition->Get().X)))
+	, CachedPositionY(MakeShareable(new TProxyProperty<FVector, FRotator::FReal>(CachedPosition, CachedPosition->Get().Y)))
+	, CachedPositionZ(MakeShareable(new TProxyProperty<FVector, FRotator::FReal>(CachedPosition, CachedPosition->Get().Z)))
+	, DefaultTransform(FTransform::Identity)
+	, InverseDefaultTransform(FTransform::Identity)
+	, bPositionDisplayRelativeToDefault(false)
+	, bRotationDisplayRelativeToDefault(false)
+{}
+
+TSharedRef<IPropertyTypeCustomization> FConstraintTransformCustomization::MakeInstance()
+{
+	return MakeShareable(new FConstraintTransformCustomization);
+}
+
+void FConstraintTransformCustomization::MakeRotationRow(TSharedRef<class IPropertyHandle>& InPriAxisPropertyHandle, TSharedRef<class IPropertyHandle>& InSecAxisPropertyHandle, FDetailWidgetRow& Row, TSharedRef<SWidget> EditSpaceToggleButtonWidget)
+{
+	TWeakPtr<IPropertyHandle> WeakPriAxisHandlePtr = InPriAxisPropertyHandle;
+	TWeakPtr<IPropertyHandle> WeakSecAxisHandlePtr = InSecAxisPropertyHandle;
+
+	const FText ComponentToolTipFormatText = LOCTEXT("ConstraintTransformRotationToolTip", "{0} component of the constraint rotation relative to the {1} bone (in degrees).");
+
+	Row
+	.CopyAction(FUIAction(FExecuteAction::CreateSP(this, &FConstraintTransformCustomization::OnCopyRotation, WeakPriAxisHandlePtr, WeakSecAxisHandlePtr)))
+	.PasteAction(FUIAction(FExecuteAction::CreateSP(this, &FConstraintTransformCustomization::OnPasteRotation, WeakPriAxisHandlePtr, WeakSecAxisHandlePtr)))
+	.NameContent()
+	.HAlign(HAlign_Fill)
+	[
+		SNew(SHorizontalBox)
+		+ SHorizontalBox::Slot()
+		.VAlign(VAlign_Center)
+		.HAlign(HAlign_Left)
+		.Padding(FMargin(0.0f, 0.0f, 4.0f, 0.0f))
+		[
+			InPriAxisPropertyHandle->CreatePropertyNameWidget(LOCTEXT("RotationLabel", "Rotation"), FText::Format(LOCTEXT("RotationToolTip", "Constraint rotation relative to the {0} bone (in degrees)."), FrameLabelText))
+		]
+		+ SHorizontalBox::Slot()
+		.VAlign(VAlign_Center)
+		.HAlign(HAlign_Right)
+		.Padding(FMargin(0.0f, 0.0f, 0.0f, 0.0f))
+		[
+			EditSpaceToggleButtonWidget
+		]
+	]
+	.ValueContent()
+	.MinDesiredWidth(375.0f)
+	.MaxDesiredWidth(375.0f)
+	[
+		SNew(SHorizontalBox)
+		+ SHorizontalBox::Slot()
+		.Padding(FMargin(0.0f, 2.0f, 3.0f, 2.0f))
+		[
+			MakeNumericProxyWidget<FRotator, FRotator::FReal>(InPriAxisPropertyHandle, InSecAxisPropertyHandle, CachedRotationRoll, FText::Format(ComponentToolTipFormatText, LOCTEXT("RotationRoll", "Roll (X)"), FrameLabelText), true, SNumericEntryBox<FRotator::FReal>::RedLabelBackgroundColor)
+		]
+		+ SHorizontalBox::Slot()
+		.Padding(FMargin(0.0f, 2.0f, 3.0f, 2.0f))
+		[
+			MakeNumericProxyWidget<FRotator, FRotator::FReal>(InPriAxisPropertyHandle, InSecAxisPropertyHandle, CachedRotationPitch, FText::Format(ComponentToolTipFormatText, LOCTEXT("RotationPitch", "Pitch (Y)"), FrameLabelText), true, SNumericEntryBox<FRotator::FReal>::GreenLabelBackgroundColor)
+		]
+		+ SHorizontalBox::Slot()
+		.Padding(FMargin(0.0f, 2.0f, 0.0f, 2.0f))
+		[
+			MakeNumericProxyWidget<FRotator, FRotator::FReal>(InPriAxisPropertyHandle, InSecAxisPropertyHandle, CachedRotationYaw, FText::Format(ComponentToolTipFormatText, LOCTEXT("RotationYaw", "Yaw (Z)"), FrameLabelText), true, SNumericEntryBox<FRotator::FReal>::BlueLabelBackgroundColor)
+		]
+	];
+}
+
+void FConstraintTransformCustomization::MakePositionRow(TSharedRef<class IPropertyHandle>& InPositionPropertyHandle, FDetailWidgetRow& Row, TSharedRef<SWidget> EditSpaceToggleButtonWidget)
+{
+	TWeakPtr<IPropertyHandle> WeakPositionHandlePtr = InPositionPropertyHandle;
+
+	const FText ComponentToolTipFormatText = LOCTEXT("ConstraintTransformPositionToolTip", "{0} component of the constraint position relative to the {1} bone.");
+
+	Row
+	.CopyAction(FUIAction(FExecuteAction::CreateSP(this, &FConstraintTransformCustomization::OnCopyPosition, WeakPositionHandlePtr)))
+	.PasteAction(FUIAction(FExecuteAction::CreateSP(this, &FConstraintTransformCustomization::OnPastePosition, WeakPositionHandlePtr)))
+	.NameContent()
+	.HAlign(HAlign_Fill)
+	[
+		SNew(SHorizontalBox)
+		+ SHorizontalBox::Slot()
+		.VAlign(VAlign_Center)
+		.HAlign(HAlign_Left)
+		.Padding(FMargin(0.0f, 0.0f, 4.0f, 0.0f))
+		[
+			InPositionPropertyHandle->CreatePropertyNameWidget(LOCTEXT("PositionLabel", "Position"), FText::Format(LOCTEXT("PositionToolTip", "Constraint position relative to the {0} bone."), FrameLabelText))
+		]
+		+ SHorizontalBox::Slot()
+		.VAlign(VAlign_Center)
+		.HAlign(HAlign_Right)
+		.Padding(FMargin(0.0f, 0.0f, 0.0f, 0.0f))
+		[
+			EditSpaceToggleButtonWidget
+		]
+	]
+	.ValueContent()
+	.MinDesiredWidth(375.0f)
+	.MaxDesiredWidth(375.0f)
+	[
+		SNew(SHorizontalBox)
+		+ SHorizontalBox::Slot()
+		.Padding(FMargin(0.0f, 2.0f, 3.0f, 2.0f))
+		[
+			MakeNumericProxyWidget<FVector, FVector::FReal>(InPositionPropertyHandle, CachedPositionX, FText::Format(ComponentToolTipFormatText, LOCTEXT("PositionX", "X"), FrameLabelText), SNumericEntryBox<FRotator::FReal>::RedLabelBackgroundColor)
+		]
+		+ SHorizontalBox::Slot()
+		.Padding(FMargin(0.0f, 2.0f, 3.0f, 2.0f))
+		[
+			MakeNumericProxyWidget<FVector, FVector::FReal>(InPositionPropertyHandle, CachedPositionY, FText::Format(ComponentToolTipFormatText, LOCTEXT("PositionY", "Y"), FrameLabelText), SNumericEntryBox<FRotator::FReal>::GreenLabelBackgroundColor)
+		]
+		+ SHorizontalBox::Slot()
+		.Padding(FMargin(0.0f, 2.0f, 0.0f, 2.0f))
+		[
+			MakeNumericProxyWidget<FVector, FVector::FReal>(InPositionPropertyHandle, CachedPositionZ, FText::Format(ComponentToolTipFormatText, LOCTEXT("PositionY", "Z"), FrameLabelText), SNumericEntryBox<FRotator::FReal>::BlueLabelBackgroundColor)
+		]
+	];
+}
+
+void FConstraintTransformCustomization::OrthonormalVectorPairToDisplayedRotator(const FVector& PriAxis, const FVector& SecAxis, FRotator& OutRotator) const
+{
+	FMatrix Rotation = FRotationMatrix::MakeFromXY(PriAxis, SecAxis);
+
+	// If rotation should be displayed relative to default then transform it from bone space.
+	if (bRotationDisplayRelativeToDefault)
+	{
+		Rotation = Rotation * InverseDefaultTransform.ToMatrixNoScale();
+	}
+
+	OutRotator = Rotation.Rotator();
+}
+
+void FConstraintTransformCustomization::DisplayedRotatorToOrthonormalVectorPair(const FRotator& InRotator, FVector& OutPriAxis, FVector& OutSecAxis) const
+{
+	FMatrix Rotation = FRotationMatrix::Make(InRotator);
+
+	// If rotation is displayed relative to default then transform it back to bone space.
+	if (bRotationDisplayRelativeToDefault)
+	{
+		Rotation = Rotation * DefaultTransform.ToMatrixNoScale();
+	}
+
+	OutPriAxis = Rotation.GetUnitAxis(EAxis::X);
+	OutSecAxis = Rotation.GetUnitAxis(EAxis::Y);
+}
+
+bool FConstraintTransformCustomization::CacheValues(TWeakPtr<IPropertyHandle> PositionPropertyHandlePtr) const
+{
+	TSharedPtr<IPropertyHandle> PositionPropertyHandle = PositionPropertyHandlePtr.Pin();
+
+	if (!PositionPropertyHandle.IsValid())
+	{
+		return false;
+	}
+
+	FVector Position;
+
+	if (PositionPropertyHandle->GetValue(Position) == FPropertyAccess::Success)
+	{
+		if (bPositionDisplayRelativeToDefault)
+		{
+			Position = InverseDefaultTransform.TransformPosition(Position);
+		}
+
+		CachedPosition->Set(Position);
+		CachedPositionX->Set(Position.X);
+		CachedPositionY->Set(Position.Y);
+		CachedPositionZ->Set(Position.Z);
+
+		return true;
+	}
+
+	return false;
+}
+
+bool FConstraintTransformCustomization::CacheValues(TWeakPtr<IPropertyHandle> PriAxisPropertyHandlePtr, TWeakPtr<IPropertyHandle> SecAxisPropertyHandlePtr) const
+{
+	TSharedPtr<IPropertyHandle> PriAxisPropertyHandle = PriAxisPropertyHandlePtr.Pin();
+	TSharedPtr<IPropertyHandle> SecAxisPropertyHandle = SecAxisPropertyHandlePtr.Pin();
+
+	if (!PriAxisPropertyHandle.IsValid() || !SecAxisPropertyHandle.IsValid())
+	{
+		return false;
+	}
+
+	FVector PriAxis;
+	FVector SecAxis;
+
+	if ((PriAxisPropertyHandle->GetValue(PriAxis) == FPropertyAccess::Success) && (SecAxisPropertyHandle->GetValue(SecAxis) == FPropertyAccess::Success))
+	{
+		FRotator CurrentRotation;
+		OrthonormalVectorPairToDisplayedRotator(PriAxis, SecAxis, CurrentRotation);
+
+		CachedRotation->Set(CurrentRotation);
+		CachedRotationPitch->Set(CurrentRotation.Pitch);
+		CachedRotationYaw->Set(CurrentRotation.Yaw);
+		CachedRotationRoll->Set(CurrentRotation.Roll);
+
+		return true;
+	}
+
+	return false;
+}
+
+bool FConstraintTransformCustomization::FlushValues(TWeakPtr<IPropertyHandle> PositionPropertyHandlePtr) const
+{
+	// Update the constraint position from the one displayed in the details panel.
+
+	TSharedPtr<IPropertyHandle> PositionPropertyHandle = PositionPropertyHandlePtr.Pin();
+
+	if (!PositionPropertyHandle.IsValid())
+	{
+		return false;
+	}
+
+	// Read position from the property
+	FVector Position;
+	PositionPropertyHandle->GetValue(Position);
+
+	// Update position with any changes from the UI.
+	FVector ModifiedPosition(
+		CachedPositionX->IsSet() ? CachedPositionX->Get() : Position.X,
+		CachedPositionY->IsSet() ? CachedPositionY->Get() : Position.Y,
+		CachedPositionZ->IsSet() ? CachedPositionZ->Get() : Position.Z
+	);
+
+	// If position is displayed relative to default then transform it back to bone space.
+	if (bPositionDisplayRelativeToDefault)
+	{
+		ModifiedPosition = DefaultTransform.TransformPosition(ModifiedPosition);
+	}
+
+	// Write back to the property.
+	return PositionPropertyHandle->SetValue(ModifiedPosition) == FPropertyAccess::Success;
+}
+
+bool FConstraintTransformCustomization::FlushValues(TWeakPtr<IPropertyHandle> PriAxisPropertyHandlePtr, TWeakPtr<IPropertyHandle> SecAxisPropertyHandlePtr) const
+{
+	// Update the constraint rotation from the one displayed in the details panel.
+
+	TSharedPtr<IPropertyHandle> PriAxisPropertyHandle = PriAxisPropertyHandlePtr.Pin();
+	TSharedPtr<IPropertyHandle> SecAxisPropertyHandle = SecAxisPropertyHandlePtr.Pin();
+
+	if (!PriAxisPropertyHandle.IsValid() || !SecAxisPropertyHandle.IsValid())
+	{
+		return false;
+	}
+
+	FVector PriAxis;
+	FVector SecAxis;
+	PriAxisPropertyHandle->GetValue(PriAxis);
+	SecAxisPropertyHandle->GetValue(SecAxis);
+	
+	FRotator CurrentRotation;
+	OrthonormalVectorPairToDisplayedRotator(PriAxis, SecAxis, CurrentRotation);
+
+	FRotator Rotation(
+		CachedRotationPitch->IsSet() ? CachedRotationPitch->Get() : CurrentRotation.Pitch,
+		CachedRotationYaw->IsSet() ? CachedRotationYaw->Get() : CurrentRotation.Yaw,
+		CachedRotationRoll->IsSet() ? CachedRotationRoll->Get() : CurrentRotation.Roll
+	);
+
+	DisplayedRotatorToOrthonormalVectorPair(Rotation, PriAxis, SecAxis);
+
+	bool Success = true;
+	Success &= (PriAxisPropertyHandle->SetValue(PriAxis) == FPropertyAccess::Result::Success);
+	Success &= (SecAxisPropertyHandle->SetValue(SecAxis) == FPropertyAccess::Result::Success);
+
+	return Success;
+}
+
+void FConstraintTransformCustomization::SetDefaultTransform(const FTransform& InTransform, TWeakPtr<IPropertyHandle> PositionPropertyHandlePtr, TWeakPtr<IPropertyHandle> PriAxisPropertyHandlePtr, TWeakPtr<IPropertyHandle> SecAxisPropertyHandlePtr)
+{ 
+	DefaultTransform = InTransform;
+	InverseDefaultTransform = InTransform.Inverse();
+
+	CacheValues(PositionPropertyHandlePtr);
+	CacheValues(PriAxisPropertyHandlePtr, SecAxisPropertyHandlePtr); // Update display values from actual properties (which as always stored as absolutes).
+}
+
+void FConstraintTransformCustomization::GetPositionAsFormattedString(FString& OutString)
+{
+	FVector Position = CachedPosition->Get();
+	OutString = FString::Printf(TEXT("(X=%f,Y=%f,Z=%f)"), Position.X, Position.Y, Position.Z);
+}
+
+void FConstraintTransformCustomization::GetRotationAsFormattedString(FString& OutString)
+{
+	FRotator Rotation = CachedRotation->Get();
+	OutString = FString::Printf(TEXT("(Pitch=%f,Yaw=%f,Roll=%f)"), Rotation.Pitch, Rotation.Yaw, Rotation.Roll);
+}
+
+void FConstraintTransformCustomization::GetValueAsFormattedString(FString& OutString)
+{
+	FString Buffer;
+	GetPositionAsFormattedString(Buffer);
+	OutString += ConstraintCopyPasteStringTokens::Position + Buffer;
+	Buffer.Reset();
+	GetRotationAsFormattedString(Buffer);
+	OutString += ConstraintCopyPasteStringTokens::Rotation + Buffer;
+}
+
+bool FConstraintTransformCustomization::SetPositionFromFormattedString(const FString& InString)
+{
+	FVector Position;
+
+	if (Position.InitFromString(InString))
+	{
+		CachedPosition->Set(Position);
+		CachedPositionX->Set(Position.X);
+		CachedPositionY->Set(Position.Y);
+		CachedPositionZ->Set(Position.Z);
+
+		return true;
+	}
+
+	return false;
+}
+
+bool FConstraintTransformCustomization::SetRotationFromFormattedString(const FString& InString)
+{
+	FString MutableString = InString;
+
+	FRotator Rotation;
+	MutableString.ReplaceInline(TEXT("Pitch="), TEXT("P="));
+	MutableString.ReplaceInline(TEXT("Yaw="), TEXT("Y="));
+	MutableString.ReplaceInline(TEXT("Roll="), TEXT("R="));
+
+	if (Rotation.InitFromString(MutableString))
+	{
+		CachedRotation->Set(Rotation);
+		CachedRotationPitch->Set(Rotation.Pitch);
+		CachedRotationYaw->Set(Rotation.Yaw);
+		CachedRotationRoll->Set(Rotation.Roll);
+
+		return true;
+	}
+
+	return false;
+}
+
+bool FConstraintTransformCustomization::SetValueFromFormattedString(const FString& InString)
+{
+	bool Success = true;
+
+	FString FormattedString = InString;
+	FString SubString;
+	
+	FormattedString.Split(ConstraintCopyPasteStringTokens::Rotation, &FormattedString, &SubString);
+	Success &= SetRotationFromFormattedString(SubString);
+
+	FormattedString.Split(ConstraintCopyPasteStringTokens::Position, &FormattedString, &SubString);
+	Success &= SetPositionFromFormattedString(SubString);
+
+	return Success;
+}
+
+void FConstraintTransformCustomization::OnCopy(TWeakPtr<IPropertyHandle> PositionPropertyHandlePtr, TWeakPtr<IPropertyHandle> PriAxisPropertyHandlePtr, TWeakPtr<IPropertyHandle> SecAxisPropertyHandlePtr)
+{
+	TSharedPtr<IPropertyHandle> PositionPropertyHandle = PositionPropertyHandlePtr.Pin();
+	TSharedPtr<IPropertyHandle> PriAxisPropertyHandle = PriAxisPropertyHandlePtr.Pin();
+	TSharedPtr<IPropertyHandle> SecAxisPropertyHandle = SecAxisPropertyHandlePtr.Pin();
+
+	if (!PositionPropertyHandle.IsValid() || !PriAxisPropertyHandle.IsValid() || !SecAxisPropertyHandle.IsValid())
+	{
+		return;
+	}
+
+	CacheValues(PositionPropertyHandle);
+	CacheValues(PriAxisPropertyHandle, SecAxisPropertyHandle);
+
+	FString CopyStr;
+	GetValueAsFormattedString(CopyStr);
+
+	if (!CopyStr.IsEmpty())
+	{
+		FPlatformApplicationMisc::ClipboardCopy(*CopyStr);
+	}
+}
+
+void FConstraintTransformCustomization::OnPaste(TWeakPtr<IPropertyHandle> PositionPropertyHandlePtr, TWeakPtr<IPropertyHandle> PriAxisPropertyHandlePtr, TWeakPtr<IPropertyHandle> SecAxisPropertyHandlePtr)
+{
+	TSharedPtr<IPropertyHandle> PositionPropertyHandle = PositionPropertyHandlePtr.Pin();
+	TSharedPtr<IPropertyHandle> PriAxisPropertyHandle = PriAxisPropertyHandlePtr.Pin();
+	TSharedPtr<IPropertyHandle> SecAxisPropertyHandle = SecAxisPropertyHandlePtr.Pin();
+
+	if (!PositionPropertyHandle.IsValid() || !PriAxisPropertyHandle.IsValid() || !SecAxisPropertyHandle.IsValid())
+	{
+		return;
+	}
+
+	FString PastedText;
+	FPlatformApplicationMisc::ClipboardPaste(PastedText);
+
+	{
+		FScopedTransaction Transaction(LOCTEXT("PastePosition", "Paste Position"));
+		SetValueFromFormattedString(PastedText);
+		FlushValues(PriAxisPropertyHandle, SecAxisPropertyHandle);
+		FlushValues(PositionPropertyHandle);
+	}
+}
+
+void FConstraintTransformCustomization::OnCopyPosition(TWeakPtr<IPropertyHandle> PositionPropertyHandlePtr)
+{
+	TSharedPtr<IPropertyHandle> PositionPropertyHandle = PositionPropertyHandlePtr.Pin();
+
+	if (!PositionPropertyHandle.IsValid())
+	{
+		return;
+	}
+
+	CacheValues(PositionPropertyHandle);
+
+	FString CopyStr;
+	GetPositionAsFormattedString(CopyStr);
+
+	if (!CopyStr.IsEmpty())
+	{
+		FPlatformApplicationMisc::ClipboardCopy(*CopyStr);
+	}
+}
+
+void FConstraintTransformCustomization::OnCopyRotation(TWeakPtr<IPropertyHandle> PriAxisPropertyHandlePtr, TWeakPtr<IPropertyHandle> SecAxisPropertyHandlePtr)
+{
+	TSharedPtr<IPropertyHandle> PriAxisPropertyHandle = PriAxisPropertyHandlePtr.Pin();
+	TSharedPtr<IPropertyHandle> SecAxisPropertyHandle = SecAxisPropertyHandlePtr.Pin();
+
+	if (!PriAxisPropertyHandle.IsValid() || !SecAxisPropertyHandle.IsValid())
+	{
+		return;
+	}
+
+	CacheValues(PriAxisPropertyHandle, SecAxisPropertyHandle);
+
+	FString CopyStr;
+	GetRotationAsFormattedString(CopyStr);
+
+	if (!CopyStr.IsEmpty())
+	{
+		FPlatformApplicationMisc::ClipboardCopy(*CopyStr);
+	}
+}
+
+void FConstraintTransformCustomization::OnPastePosition(TWeakPtr<IPropertyHandle> PositionPropertyHandlePtr)
+{
+	TSharedPtr<IPropertyHandle> PositionPropertyHandle = PositionPropertyHandlePtr.Pin();
+
+	if (!PositionPropertyHandle.IsValid())
+	{
+		return;
+	}
+
+	FString PastedText;
+	FPlatformApplicationMisc::ClipboardPaste(PastedText);
+
+	{
+		FScopedTransaction Transaction(LOCTEXT("PastePosition", "Paste Position"));
+		SetPositionFromFormattedString(PastedText);
+		FlushValues(PositionPropertyHandle);
+	}
+}
+
+void FConstraintTransformCustomization::OnPasteRotation(TWeakPtr<IPropertyHandle> PriAxisPropertyHandlePtr, TWeakPtr<IPropertyHandle> SecAxisPropertyHandlePtr)
+{
+	TSharedPtr<IPropertyHandle> PriAxisPropertyHandle = PriAxisPropertyHandlePtr.Pin();
+	TSharedPtr<IPropertyHandle> SecAxisPropertyHandle = SecAxisPropertyHandlePtr.Pin();
+
+	if (!PriAxisPropertyHandle.IsValid() || !SecAxisPropertyHandle.IsValid())
+	{
+		return;
+	}
+
+	FString PastedText;
+	FPlatformApplicationMisc::ClipboardPaste(PastedText);
+	
+	{
+		FScopedTransaction Transaction(LOCTEXT("PasteRotation", "Paste Rotation"));
+		SetRotationFromFormattedString(PastedText);
+		FlushValues(PriAxisPropertyHandle, SecAxisPropertyHandle);
+	}
+}
+
+template<typename ProxyType, typename NumericType> TSharedRef<SWidget> FConstraintTransformCustomization::MakeNumericProxyWidget(TSharedRef<IPropertyHandle>& PriAxisPropertyHandle, TSharedRef<IPropertyHandle>& SecAxisPropertyHandle, TSharedRef< TProxyProperty<ProxyType, NumericType> >& ProxyValue, const FText& ToolTipText, bool bRotationInDegrees, const FLinearColor& LabelBackgroundColor)
+{
+	TWeakPtr<IPropertyHandle> WeakPriAxisHandlePtr = PriAxisPropertyHandle;
+	TWeakPtr<IPropertyHandle> WeakSecAxisHandlePtr = SecAxisPropertyHandle;
+
+	return
+		SNew(SNumericEntryBox<NumericType>)
+		.IsEnabled(this, &FConstraintTransformCustomization::IsRotationValueEnabled, WeakPriAxisHandlePtr, WeakSecAxisHandlePtr)
+		.Value(this, &FConstraintTransformCustomization::OnGetRotationValue<ProxyType, NumericType>, WeakPriAxisHandlePtr, WeakSecAxisHandlePtr, ProxyValue)
+		.Font(IDetailLayoutBuilder::GetDetailFont())
+		.UndeterminedString(NSLOCTEXT("PropertyEditor", "MultipleValues", "Multiple Values"))
+		.OnValueCommitted(this, &FConstraintTransformCustomization::OnRotationValueCommitted<ProxyType, NumericType>, WeakPriAxisHandlePtr, WeakSecAxisHandlePtr, ProxyValue)
+		.OnValueChanged(this, &FConstraintTransformCustomization::OnRotationValueChanged<ProxyType, NumericType>, WeakPriAxisHandlePtr, WeakSecAxisHandlePtr, ProxyValue)
+		.AllowSpin(false)
+		.MinValue(TOptional<NumericType>())
+		.MaxValue(TOptional<NumericType>())
+		.MaxSliderValue(bRotationInDegrees ? 360.0f : TOptional<NumericType>())
+		.MinSliderValue(bRotationInDegrees ? 0.0f : TOptional<NumericType>())
+		.LabelPadding(FMargin(3))
+		.ToolTipText(ToolTipText)
+		.LabelLocation(SNumericEntryBox<NumericType>::ELabelLocation::Inside)
+		.Label()
+		[
+			SNumericEntryBox<NumericType>::BuildNarrowColorLabel(LabelBackgroundColor)
+		];
+}
+
+template<typename ProxyType, typename NumericType> TSharedRef<SWidget> FConstraintTransformCustomization::MakeNumericProxyWidget(TSharedRef<IPropertyHandle>& PositionPropertyHandle, TSharedRef< TProxyProperty<ProxyType, NumericType> >& ProxyValue, const FText& ToolTipText, const FLinearColor& LabelBackgroundColor)
+{
+	TWeakPtr<IPropertyHandle> WeakPositionHandlePtr = PositionPropertyHandle;
+
+	return
+		SNew(SNumericEntryBox<NumericType>)
+		.IsEnabled(this, &FConstraintTransformCustomization::IsPositionValueEnabled, WeakPositionHandlePtr)
+		.Value(this, &FConstraintTransformCustomization::OnGetPositionValue<ProxyType, NumericType>, WeakPositionHandlePtr, ProxyValue)
+		.Font(IDetailLayoutBuilder::GetDetailFont())
+		.UndeterminedString(NSLOCTEXT("PropertyEditor", "MultipleValues", "Multiple Values"))
+		.OnValueCommitted(this, &FConstraintTransformCustomization::OnPositionValueCommitted<ProxyType, NumericType>, WeakPositionHandlePtr, ProxyValue)
+		.OnValueChanged(this, &FConstraintTransformCustomization::OnPositionValueChanged<ProxyType, NumericType>, WeakPositionHandlePtr, ProxyValue)
+		.AllowSpin(false)
+		.LabelPadding(FMargin(3))
+		.ToolTipText(ToolTipText)
+		.LabelLocation(SNumericEntryBox<NumericType>::ELabelLocation::Inside)
+		.Label()
+		[
+			SNumericEntryBox<NumericType>::BuildNarrowColorLabel(LabelBackgroundColor)
+		];
+}
+
+
+bool FConstraintTransformCustomization::IsRotationValueEnabled(TWeakPtr<IPropertyHandle> PriAxisWeakHandlePtr, TWeakPtr<IPropertyHandle> SecAxisWeakHandlePtr) const
+{
+	return (PriAxisWeakHandlePtr.IsValid() && !PriAxisWeakHandlePtr.Pin()->IsEditConst()) && (SecAxisWeakHandlePtr.IsValid() && !SecAxisWeakHandlePtr.Pin()->IsEditConst());
+}
+
+template<typename ProxyType, typename NumericType>
+TOptional<NumericType> FConstraintTransformCustomization::OnGetRotationValue(TWeakPtr<IPropertyHandle> PriAxisWeakHandlePtr, TWeakPtr<IPropertyHandle> SecAxisWeakHandlePtr, TSharedRef< TProxyProperty<ProxyType, NumericType> > ProxyValue) const
+{
+	if (CacheValues(PriAxisWeakHandlePtr, SecAxisWeakHandlePtr))
+	{
+		return ProxyValue->Get();
+	}
+	return TOptional<NumericType>();
+}
+
+template<typename ProxyType, typename NumericType>
+void FConstraintTransformCustomization::OnRotationValueCommitted(NumericType NewValue, ETextCommit::Type CommitType, TWeakPtr<IPropertyHandle> PriAxisWeakHandlePtr, TWeakPtr<IPropertyHandle> SecAxisWeakHandlePtr, TSharedRef< TProxyProperty<ProxyType, NumericType> > ProxyValue)
+{
+	if (!bIsUsingSlider && !GIsTransacting)
+	{
+		ProxyValue->Set(NewValue);
+		FlushValues(PriAxisWeakHandlePtr, SecAxisWeakHandlePtr);
+	}
+}
+
+template<typename ProxyType, typename NumericType>
+void FConstraintTransformCustomization::OnRotationValueChanged(NumericType NewValue, TWeakPtr<IPropertyHandle> PriAxisWeakHandlePtr, TWeakPtr<IPropertyHandle> SecAxisWeakHandlePtr, TSharedRef< TProxyProperty<ProxyType, NumericType> > ProxyValue)
+{
+	if (bIsUsingSlider)
+	{
+		ProxyValue->Set(NewValue);
+		FlushValues(PriAxisWeakHandlePtr, SecAxisWeakHandlePtr);
+	}
+}
+
+bool FConstraintTransformCustomization::IsPositionValueEnabled(TWeakPtr<IPropertyHandle> PositionWeakHandlePtr) const
+{
+	return (PositionWeakHandlePtr.IsValid() && !PositionWeakHandlePtr.Pin()->IsEditConst());
+}
+
+template<typename ProxyType, typename NumericType> TOptional<NumericType> FConstraintTransformCustomization::OnGetPositionValue(TWeakPtr<IPropertyHandle> PositionWeakHandlePtr, TSharedRef< TProxyProperty<ProxyType, NumericType> > ProxyValue) const
+{
+	if (CacheValues(PositionWeakHandlePtr))
+	{
+		return ProxyValue->Get();
+	}
+	return TOptional<NumericType>();
+}
+
+template<typename ProxyType, typename NumericType>
+void FConstraintTransformCustomization::OnPositionValueCommitted(NumericType NewValue, ETextCommit::Type CommitType, TWeakPtr<IPropertyHandle> PositionWeakHandlePtr, TSharedRef< TProxyProperty<ProxyType, NumericType> > ProxyValue)
+{
+	if (!GIsTransacting)
+	{
+		ProxyValue->Set(NewValue);
+		FlushValues(PositionWeakHandlePtr);
+	}
+}
+
+template<typename ProxyType, typename NumericType>
+void FConstraintTransformCustomization::OnPositionValueChanged(NumericType NewValue, TWeakPtr<IPropertyHandle> PositionWeakHandlePtr, TSharedRef< TProxyProperty<ProxyType, NumericType> > ProxyValue)
+{
+	ProxyValue->Set(NewValue);
+	FlushValues(PositionWeakHandlePtr);
+}
+
+void FConstraintTransformCustomization::SetFrameLabelText(const FText InText)
+{
+	FrameLabelText = InText;
+}
+
+FText& FConstraintTransformCustomization::GetFrameLabelText()
+{
+	return FrameLabelText;
+}
+
 namespace ConstraintDetails
 {
 	bool GetBoolProperty(TSharedPtr<IPropertyHandle> Prop)
@@ -156,11 +806,352 @@ namespace ConstraintDetails
 
 }
 
+FPhysicsConstraintComponentDetails::FPhysicsConstraintComponentDetails()
+: ChildTransformProxy(StaticCastSharedRef<FConstraintTransformCustomization>(FConstraintTransformCustomization::MakeInstance()))
+, ParentTransformProxy(StaticCastSharedRef<FConstraintTransformCustomization>(FConstraintTransformCustomization::MakeInstance()))
+{}
+
 TSharedRef<IDetailCustomization> FPhysicsConstraintComponentDetails::MakeInstance()
 {
 	return MakeShareable(new FPhysicsConstraintComponentDetails());
 }
 
+void FPhysicsConstraintComponentDetails::AddConstraintProperties(IDetailLayoutBuilder& DetailBuilder, TSharedPtr<IPropertyHandle> ConstraintInstance, TArray<TWeakObjectPtr<UObject>>& Objects)
+{
+	ChildPositionPropertyHandle = ConstraintInstance->GetChildHandle(GET_MEMBER_NAME_CHECKED(FConstraintInstance, Pos1));
+	ChildPriAxisPropertyHandle = ConstraintInstance->GetChildHandle(GET_MEMBER_NAME_CHECKED(FConstraintInstance, PriAxis1));
+	ChildSecAxisPropertyHandle = ConstraintInstance->GetChildHandle(GET_MEMBER_NAME_CHECKED(FConstraintInstance, SecAxis1));
+
+	ParentPositionPropertyHandle = ConstraintInstance->GetChildHandle(GET_MEMBER_NAME_CHECKED(FConstraintInstance, Pos2));
+	ParentPriAxisPropertyHandle = ConstraintInstance->GetChildHandle(GET_MEMBER_NAME_CHECKED(FConstraintInstance, PriAxis2));
+	ParentSecAxisPropertyHandle = ConstraintInstance->GetChildHandle(GET_MEMBER_NAME_CHECKED(FConstraintInstance, SecAxis2));
+
+	IDetailCategoryBuilder& ConstraintCategory = DetailBuilder.EditCategory("Constraint");
+
+	// Add current profile name to the constraint category header.
+	if (bInPhat)
+	{
+		ConstraintCategory.HeaderContent(
+			SNew(SHorizontalBox)
+			+ SHorizontalBox::Slot()
+			.FillWidth(1.0f)
+			.HAlign(HAlign_Right)
+			.VAlign(VAlign_Center)
+			[
+				SNew(SRichTextBlock)
+				.DecoratorStyleSet(&FAppStyle::Get())
+				.Text_Lambda([Objects]()
+				{
+					if (Objects.Num() > 0)
+					{
+						if (UPhysicsConstraintTemplate* Constraint = Cast<UPhysicsConstraintTemplate>(Objects[0].Get()))
+						{
+							FName CurrentProfileName = Constraint->GetCurrentConstraintProfileName();
+							if (CurrentProfileName != NAME_None)
+							{
+								if (Constraint->ContainsConstraintProfile(CurrentProfileName))
+								{
+									return FText::Format(LOCTEXT("ProfileFormatAssigned", "Assigned to Profile: <RichTextBlock.Bold>{0}</>"), FText::FromName(CurrentProfileName));
+								}
+								else
+								{
+									return FText::Format(LOCTEXT("ProfileFormatNotAssigned", "Not Assigned to Profile: <RichTextBlock.Bold>{0}</>"), FText::FromName(CurrentProfileName));
+								}
+							}
+							else
+							{
+								return LOCTEXT("ProfileFormatNone", "Current Profile: <RichTextBlock.Bold>None</>");
+							}
+						}
+					}
+
+					return FText();
+				})
+			]);
+	}
+
+	ConstraintCategory.AddProperty(ConstraintInstance->GetChildHandle(GET_MEMBER_NAME_CHECKED(FConstraintInstance, ConstraintBone1))).DisplayName(LOCTEXT("ConstraintChildBoneName", "Child Bone Name"));
+	ConstraintCategory.AddProperty(ConstraintInstance->GetChildHandle(GET_MEMBER_NAME_CHECKED(FConstraintInstance, ConstraintBone2))).DisplayName(LOCTEXT("ConstraintParentBoneName", "Parent Bone Name"));
+
+	IDetailCategoryBuilder& ConstraintTransformsCat = DetailBuilder.EditCategory("ConstraintTransforms");
+
+	AddConstraintFrameTransform(EConstraintFrame::Frame1, ConstraintTransformsCat, ConstraintInstance);
+	AddConstraintFrameTransform(EConstraintFrame::Frame2, ConstraintTransformsCat, ConstraintInstance);
+
+	UpdateTransformProxyDisplayRelativeToDefault(ConstraintInstance);
+
+	// Hide the constraint transform values that are being represented by the proxy values in the details panel.
+	ChildPositionPropertyHandle->MarkHiddenByCustomization();
+	ChildPriAxisPropertyHandle->MarkHiddenByCustomization();
+	ChildSecAxisPropertyHandle->MarkHiddenByCustomization();
+
+	ParentPositionPropertyHandle->MarkHiddenByCustomization();
+	ParentPriAxisPropertyHandle->MarkHiddenByCustomization();
+	ParentSecAxisPropertyHandle->MarkHiddenByCustomization();
+}	
+
+TSharedRef<SWidget> FPhysicsConstraintComponentDetails::MakeEditSpaceToggleButtonWidget(TSharedPtr<IPropertyHandle> ConstraintInstancePropertyHandle, const EConstraintTransformComponentFlags ComponentFlags)
+{
+	auto ToolTipTextLambda = [this, ComponentFlags]()
+	{
+		const FText FrameText = EnumHasAnyFlags(ComponentFlags, EConstraintTransformComponentFlags::AllChild) ? LOCTEXT("Child", "Child") : LOCTEXT("Parent", "Parent");
+		const FText ComponentText = EnumHasAnyFlags(ComponentFlags, EConstraintTransformComponentFlags::AllPosition) ? LOCTEXT("Position", "Position") : LOCTEXT("Rotation", "Rotation");
+		const FText LocalFrameText = FText::Format(LOCTEXT("LocalFrameDescription", "in the frame of the {0} bone"), FrameText);
+		const FText SnapFrameText = LOCTEXT("SnapFrameDescription", "relative to the default (snapped) transforms");
+
+		const bool IsDisplayingRelativeToDefault = this->IsDisplayingConstraintTransformComponentRelativeToDefault(ComponentFlags);
+		const FText CurrentFrameText = IsDisplayingRelativeToDefault ? SnapFrameText : LocalFrameText;
+		const FText AlternativeFrameText = IsDisplayingRelativeToDefault ? LocalFrameText : SnapFrameText;
+
+		return FText::Format(LOCTEXT("ToggleDisplayConstraintTransformComponentRelativeToDefault", "{0} transform {1} component displayed {2}. Click here to switch to displaying {3}. Hold Shift to change all."), FrameText, ComponentText, CurrentFrameText, AlternativeFrameText);
+	};
+
+	return
+		SNew(SButton)
+		.HAlign(HAlign_Center)
+		.VAlign(VAlign_Center)
+		.OnClicked_Lambda([this, ConstraintInstancePropertyHandle, ComponentFlags]() { this->ToggleDisplayConstraintTransformComponentRelativeToDefault(ConstraintInstancePropertyHandle, ComponentFlags); return FReply::Handled(); })
+		.ButtonColorAndOpacity(FSlateColor::UseForeground())
+		.Content()
+		[
+			SNew(SImage)
+			.ColorAndOpacity(FSlateColor::UseForeground())
+			.Image_Lambda([this, ComponentFlags]() { return (IsDisplayingConstraintTransformComponentRelativeToDefault(ComponentFlags)) ? FAppStyle::GetBrush("Icons.Snap") : FAppStyle::GetBrush("Icons.Transform"); })
+			.ToolTipText_Lambda(ToolTipTextLambda)
+		];
+}
+
+void FPhysicsConstraintComponentDetails::AddConstraintFrameTransform(const EConstraintFrame::Type ConstraintFrameType, IDetailCategoryBuilder& ConstraintCat, TSharedPtr<IPropertyHandle> ConstraintInstancePropertyHandle)
+{
+	TSharedPtr<FConstraintTransformCustomization> TransformProxy;
+	TSharedPtr<IPropertyHandle> PositionPropertyHandle;
+	TSharedPtr<IPropertyHandle> PriAxisPropertyHandle;
+	TSharedPtr<IPropertyHandle> SecAxisPropertyHandle;
+	EConstraintTransformComponentFlags PositionSnapFlag = EConstraintTransformComponentFlags::None;
+	EConstraintTransformComponentFlags RotationSnapFlag = EConstraintTransformComponentFlags::None;
+	FText KeyboardShortcutText;
+
+	if (ConstraintFrameType == EConstraintFrame::Frame1) // Child Frame
+	{
+		TransformProxy = ChildTransformProxy;
+		TransformProxy->SetFrameLabelText(LOCTEXT("Child", "Child"));
+		KeyboardShortcutText = LOCTEXT("ShiftCtl", "[Shift + Alt]");
+		PositionPropertyHandle = ChildPositionPropertyHandle;
+		PriAxisPropertyHandle = ChildPriAxisPropertyHandle;
+		SecAxisPropertyHandle = ChildSecAxisPropertyHandle;
+		PositionSnapFlag = EConstraintTransformComponentFlags::ChildPosition;
+		RotationSnapFlag = EConstraintTransformComponentFlags::ChildRotation;
+	}
+	
+	if (ConstraintFrameType == EConstraintFrame::Frame2) // Parent Frame
+	{
+		TransformProxy = ParentTransformProxy;
+		TransformProxy->SetFrameLabelText(LOCTEXT("Parent", "Parent"));
+		KeyboardShortcutText = LOCTEXT("Alt", "[Alt]");
+		PositionPropertyHandle = ParentPositionPropertyHandle;
+		PriAxisPropertyHandle = ParentPriAxisPropertyHandle;
+		SecAxisPropertyHandle = ParentSecAxisPropertyHandle;
+		PositionSnapFlag = EConstraintTransformComponentFlags::ParentPosition;
+		RotationSnapFlag = EConstraintTransformComponentFlags::ParentRotation;
+	}
+
+	const EConstraintTransformComponentFlags SnapFlag = PositionSnapFlag | RotationSnapFlag;
+		
+	if (TransformProxy && PositionPropertyHandle && PriAxisPropertyHandle && SecAxisPropertyHandle)
+	{
+		// Create list of properties to be associated with the proxy custom rows.
+		TArray<TSharedPtr<IPropertyHandle>> TransformProxyPropertyHandleList;
+		TransformProxyPropertyHandleList.Add(ConstraintInstancePropertyHandle);
+
+		FDetailWidgetRow& HeaderRow = ConstraintCat.AddCustomRow(TransformProxy->GetFrameLabelText());
+
+		HeaderRow
+			.WholeRowContent()
+			[
+				SNew(SHorizontalBox)
+				+ SHorizontalBox::Slot()
+				.VAlign(VAlign_Center)
+				.HAlign(HAlign_Left)
+				.Padding(0.0f, 0.0f, 4.0f, 0.0f)
+				[
+					SNew(STextBlock)
+					.Text(TransformProxy->GetFrameLabelText())
+					.TextStyle(FAppStyle::Get(), "DetailsView.CategoryTextStyle")
+					.ColorAndOpacity(this, &FPhysicsConstraintComponentDetails::GetConstraintTransformColorAndOpacity, SnapFlag)
+				]
+
+				+ SHorizontalBox::Slot()
+				.VAlign(VAlign_Center)
+				.HAlign(HAlign_Right)
+				.Padding(0.0f, 0.0f, 4.0f, 0.0f)
+				[
+					SNew(STextBlock)
+					.Text(KeyboardShortcutText)
+					.TextStyle(FAppStyle::Get(), "DetailsView.CategoryTextStyle")
+					.ColorAndOpacity(this, &FPhysicsConstraintComponentDetails::GetConstraintTransformColorAndOpacity, SnapFlag)
+				]
+			];
+
+		HeaderRow.PropertyHandleList(TransformProxyPropertyHandleList);
+
+		HeaderRow.CopyAction(FUIAction(FExecuteAction::CreateSP(this, &FPhysicsConstraintComponentDetails::OnCopyConstraintTransform, PositionPropertyHandle, PriAxisPropertyHandle, SecAxisPropertyHandle, TransformProxy)));
+		HeaderRow.PasteAction(FUIAction(FExecuteAction::CreateSP(this, &FPhysicsConstraintComponentDetails::OnPasteConstraintTransform, PositionPropertyHandle, PriAxisPropertyHandle, SecAxisPropertyHandle, TransformProxy)));
+
+		HeaderRow.OverrideResetToDefault(
+			FResetToDefaultOverride::Create(
+				FIsResetToDefaultVisible::CreateSP(this, &FPhysicsConstraintComponentDetails::IsSnapConstraintTransformComponentVisible, SnapFlag),
+				FResetToDefaultHandler::CreateSP(this, &FPhysicsConstraintComponentDetails::SnapConstraintTransformComponentsToDefault, SnapFlag)
+			));
+
+		PositionPropertyHandle->MarkHiddenByCustomization();
+		PriAxisPropertyHandle->MarkHiddenByCustomization();
+		SecAxisPropertyHandle->MarkHiddenByCustomization();
+
+		TSharedRef<IPropertyHandle> PositionPropertyHandleRef = PositionPropertyHandle.ToSharedRef();
+		TSharedRef<IPropertyHandle> PriAxisPropertyHandleRef = PriAxisPropertyHandle.ToSharedRef();
+		TSharedRef<IPropertyHandle> SecAxisPropertyHandleRef = SecAxisPropertyHandle.ToSharedRef();
+
+		// Create Position Row
+		{
+			// Create a new, empty row in the details panel.
+			FDetailWidgetRow& ProxyRow = ConstraintCat.AddCustomRow(LOCTEXT("ConstraintPositionProxy", "ConstraintPositionProxy")); // TODO - non localised ?
+
+			// Associate the constraint property with the new row.
+			ProxyRow.PropertyHandleList(TransformProxyPropertyHandleList);
+
+			// Populate the new row with widgets generated by the rotation proxy object.
+			TransformProxy->MakePositionRow(PositionPropertyHandleRef, ProxyRow, MakeEditSpaceToggleButtonWidget(ConstraintInstancePropertyHandle, PositionSnapFlag));
+
+			// Set the behavior of the reset button for the new row, so that it snaps the orientation back to the default.
+			ProxyRow.OverrideResetToDefault(
+				FResetToDefaultOverride::Create(
+					FIsResetToDefaultVisible::CreateSP(this, &FPhysicsConstraintComponentDetails::IsSnapConstraintTransformComponentVisible, PositionSnapFlag),
+					FResetToDefaultHandler::CreateSP(this, &FPhysicsConstraintComponentDetails::SnapConstraintTransformComponentsToDefault, PositionSnapFlag)
+				));
+
+			ProxyRow.IsEnabled(TAttribute<bool>::Create(TAttribute<bool>::FGetter::CreateSP(this, &FPhysicsConstraintComponentDetails::IsConstraintTransformComponentEnabled, SnapFlag)));
+		}
+
+		// Create Rotation Row
+		{
+			// Create a new, empty row in the details panel.
+			FDetailWidgetRow& ProxyRow = ConstraintCat.AddCustomRow(LOCTEXT("ConstraintTransformProxy", "ConstraintTransformProxy"));
+
+			// Associate the constraint property with the new row.
+			ProxyRow.PropertyHandleList(TransformProxyPropertyHandleList);
+
+			// Populate the new row with widgets generated by the rotation proxy object.
+			TransformProxy->MakeRotationRow(PriAxisPropertyHandleRef, SecAxisPropertyHandleRef, ProxyRow, MakeEditSpaceToggleButtonWidget(ConstraintInstancePropertyHandle, RotationSnapFlag));
+
+			// Set the behavior of the reset button for the new row, so that it snaps the orientation back to the default.
+			ProxyRow.OverrideResetToDefault(
+				FResetToDefaultOverride::Create(
+					FIsResetToDefaultVisible::CreateSP(this, &FPhysicsConstraintComponentDetails::IsSnapConstraintTransformComponentVisible, RotationSnapFlag),
+					FResetToDefaultHandler::CreateSP(this, &FPhysicsConstraintComponentDetails::SnapConstraintTransformComponentsToDefault, RotationSnapFlag)
+				));
+
+			ProxyRow.IsEnabled(TAttribute<bool>::Create(TAttribute<bool>::FGetter::CreateSP(this, &FPhysicsConstraintComponentDetails::IsConstraintTransformComponentEnabled, SnapFlag)));
+		}
+	}
+}
+
+void FPhysicsConstraintComponentDetails::SnapConstraintTransformComponentsToDefault(TSharedPtr<IPropertyHandle> ConstraintInstancePropertyHandle, const EConstraintTransformComponentFlags SnapFlags)
+{
+	if (FConstraintInstance* const ConstraintInstance = GetConstraintInstance())
+	{
+		ConstraintInstance->SnapTransformsToDefault(SnapFlags, ParentPhysicsAsset);
+	}
+}
+
+bool FPhysicsConstraintComponentDetails::IsSnapConstraintTransformComponentVisible(TSharedPtr<IPropertyHandle> ConstraintInstancePropertyHandle, const EConstraintTransformComponentFlags SnapFlags)
+{
+	bool Result = false;
+
+	const FTransform ParentTransform = ParentTransformProxy->GetDefaultTransform();
+	const FTransform ChildTransform = ChildTransformProxy->GetDefaultTransform();
+
+	if (EnumHasAnyFlags(SnapFlags, EConstraintTransformComponentFlags::ChildPosition))
+	{
+		Result |= !IsChildPropertyNearlyEqual< FVector >(ConstraintInstancePropertyHandle, GET_MEMBER_NAME_CHECKED(FConstraintInstance, Pos1), ChildTransform.GetLocation());
+	}
+
+	if (EnumHasAnyFlags(SnapFlags, EConstraintTransformComponentFlags::ChildRotation))
+	{
+		Result |= !(IsChildPropertyNearlyEqual< FVector >(ConstraintInstancePropertyHandle, GET_MEMBER_NAME_CHECKED(FConstraintInstance, PriAxis1), ChildTransform.GetUnitAxis(EAxis::X)) && IsChildPropertyNearlyEqual< FVector >(ConstraintInstancePropertyHandle, GET_MEMBER_NAME_CHECKED(FConstraintInstance, SecAxis1), ChildTransform.GetUnitAxis(EAxis::Y)));
+	}
+
+	if (EnumHasAnyFlags(SnapFlags, EConstraintTransformComponentFlags::ParentPosition))
+	{
+		Result |= !IsChildPropertyNearlyEqual< FVector >(ConstraintInstancePropertyHandle, GET_MEMBER_NAME_CHECKED(FConstraintInstance, Pos2), ParentTransform.GetLocation());
+	}
+
+	if (EnumHasAnyFlags(SnapFlags, EConstraintTransformComponentFlags::ParentRotation))
+	{
+		Result |= !(IsChildPropertyNearlyEqual< FVector >(ConstraintInstancePropertyHandle, GET_MEMBER_NAME_CHECKED(FConstraintInstance, PriAxis2), ParentTransform.GetUnitAxis(EAxis::X)) && IsChildPropertyNearlyEqual< FVector >(ConstraintInstancePropertyHandle, GET_MEMBER_NAME_CHECKED(FConstraintInstance, SecAxis2), ParentTransform.GetUnitAxis(EAxis::Y)));
+	}
+
+	return Result;
+}
+
+void FPhysicsConstraintComponentDetails::ToggleDisplayConstraintTransformComponentRelativeToDefault(TSharedPtr<IPropertyHandle> ConstraintInstancePropertyHandle, const EConstraintTransformComponentFlags ComponentFlags)
+{
+	IPhysicsAssetRenderInterface& PhysicsAssetRenderInterface = IModularFeatures::Get().GetModularFeature<IPhysicsAssetRenderInterface>("PhysicsAssetRenderInterface");
+	const bool bCurrentState = IsDisplayingConstraintTransformComponentRelativeToDefault(ComponentFlags);
+	const EConstraintTransformComponentFlags ComponentFlagsToModify = (FSlateApplication::Get().GetModifierKeys().IsShiftDown()) ? EConstraintTransformComponentFlags::All : ComponentFlags;
+	PhysicsAssetRenderInterface.SetDisplayConstraintTransformComponentRelativeToDefault(ParentPhysicsAsset, ComponentFlagsToModify, !bCurrentState);
+	UpdateTransformProxyDisplayRelativeToDefault(ConstraintInstancePropertyHandle);
+}
+
+bool FPhysicsConstraintComponentDetails::IsConstraintTransformComponentEnabled(const EConstraintTransformComponentFlags ComponentFlags) const
+{
+	IPhysicsAssetRenderInterface& PhysicsAssetRenderInterface = IModularFeatures::Get().GetModularFeature<IPhysicsAssetRenderInterface>("PhysicsAssetRenderInterface");
+	const EConstraintTransformComponentFlags ComponentManipulationFlags = PhysicsAssetRenderInterface.GetConstraintViewportManipulationFlags(ParentPhysicsAsset);
+
+	return EnumHasAllFlags(ComponentManipulationFlags, ComponentFlags);
+}
+
+FSlateColor FPhysicsConstraintComponentDetails::GetConstraintTransformColorAndOpacity(const EConstraintTransformComponentFlags ComponentFlags) const
+{
+	return IsConstraintTransformComponentEnabled(ComponentFlags) ? FSlateColor::UseForeground() : FSlateColor::UseSubduedForeground();
+}
+
+void FPhysicsConstraintComponentDetails::UpdateTransformProxyDisplayRelativeToDefault(TSharedPtr<IPropertyHandle> ConstraintInstancePropertyHandle)
+{
+	IPhysicsAssetRenderInterface& PhysicsAssetRenderInterface = IModularFeatures::Get().GetModularFeature<IPhysicsAssetRenderInterface>("PhysicsAssetRenderInterface");
+
+	ChildTransformProxy->SetPositionDisplayRelativeToDefault(PhysicsAssetRenderInterface.IsDisplayingConstraintTransformComponentRelativeToDefault(ParentPhysicsAsset, EConstraintTransformComponentFlags::ChildPosition));
+	ChildTransformProxy->SetRotationDisplayRelativeToDefault(PhysicsAssetRenderInterface.IsDisplayingConstraintTransformComponentRelativeToDefault(ParentPhysicsAsset, EConstraintTransformComponentFlags::ChildRotation));
+	ParentTransformProxy->SetPositionDisplayRelativeToDefault(PhysicsAssetRenderInterface.IsDisplayingConstraintTransformComponentRelativeToDefault(ParentPhysicsAsset, EConstraintTransformComponentFlags::ParentPosition));
+	ParentTransformProxy->SetRotationDisplayRelativeToDefault(PhysicsAssetRenderInterface.IsDisplayingConstraintTransformComponentRelativeToDefault(ParentPhysicsAsset, EConstraintTransformComponentFlags::ParentRotation));
+
+	if (FConstraintInstance* const ConstraintInstance = GetConstraintInstance())
+	{
+		ChildTransformProxy->SetDefaultTransform(ConstraintInstance->CalculateDefaultChildTransform(), ChildPositionPropertyHandle, ChildPriAxisPropertyHandle, ChildSecAxisPropertyHandle);
+		ParentTransformProxy->SetDefaultTransform(ConstraintInstance->CalculateDefaultParentTransform(ParentPhysicsAsset), ParentPositionPropertyHandle, ParentPriAxisPropertyHandle, ParentSecAxisPropertyHandle);
+	}
+}
+
+bool FPhysicsConstraintComponentDetails::IsDisplayingConstraintTransformComponentRelativeToDefault(const EConstraintTransformComponentFlags ComponentFlags)
+{
+	return IModularFeatures::Get().GetModularFeature<IPhysicsAssetRenderInterface>("PhysicsAssetRenderInterface").IsDisplayingConstraintTransformComponentRelativeToDefault(ParentPhysicsAsset, ComponentFlags);
+}
+
+void FPhysicsConstraintComponentDetails::OnCopyConstraintTransform(TSharedPtr<IPropertyHandle> PositionPropertyHandle, TSharedPtr<IPropertyHandle> PriAxisPropertyHandle, TSharedPtr<IPropertyHandle> SecAxisPropertyHandle, TSharedPtr<FConstraintTransformCustomization> TransformProxy)
+{
+	if (TransformProxy.IsValid())
+	{
+		TransformProxy->OnCopy(PositionPropertyHandle, PriAxisPropertyHandle, SecAxisPropertyHandle);
+	}
+}
+
+void FPhysicsConstraintComponentDetails::OnPasteConstraintTransform(TSharedPtr<IPropertyHandle> PositionPropertyHandle, TSharedPtr<IPropertyHandle> PriAxisPropertyHandle, TSharedPtr<IPropertyHandle> SecAxisPropertyHandle, TSharedPtr<FConstraintTransformCustomization> TransformProxy)
+{
+	if (TransformProxy.IsValid())
+	{
+		TransformProxy->OnPaste(PositionPropertyHandle, PriAxisPropertyHandle, SecAxisPropertyHandle);
+	}
+}
+
 void FPhysicsConstraintComponentDetails::AddConstraintBehaviorProperties(IDetailLayoutBuilder& DetailBuilder, TSharedPtr<IPropertyHandle> ConstraintInstance, TSharedPtr<IPropertyHandle> ProfilePropertiesProperty)
 {
 	IDetailCategoryBuilder& ConstraintCat = DetailBuilder.EditCategory("Constraint Behavior");
@@ -859,8 +1850,6 @@ void FPhysicsConstraintComponentDetails::AddAngularDrive(IDetailLayoutBuilder& D
 	[
 		MaxForceWidget
 	];
-
-
 }
 
 void FPhysicsConstraintComponentDetails::CustomizeDetails( IDetailLayoutBuilder& DetailBuilder )
@@ -868,8 +1857,7 @@ void FPhysicsConstraintComponentDetails::CustomizeDetails( IDetailLayoutBuilder&
 	TArray<TWeakObjectPtr<UObject>> Objects;
 	DetailBuilder.GetObjectsBeingCustomized(Objects);
 
-	TSharedPtr<IPropertyHandle> ConstraintInstance;
-	UPhysicsConstraintComponent* ConstraintComp = NULL;
+	TSharedPtr<IPropertyHandle> ConstraintInstanceProperty;
 	APhysicsConstraintActor* OwningConstraintActor = NULL;
 
 	bInPhat = false;
@@ -880,70 +1868,32 @@ void FPhysicsConstraintComponentDetails::CustomizeDetails( IDetailLayoutBuilder&
 
 		if (Objects[i]->IsA(UPhysicsConstraintTemplate::StaticClass()))
 		{
-			ConstraintInstance = DetailBuilder.GetProperty(GET_MEMBER_NAME_CHECKED(UPhysicsConstraintTemplate, DefaultInstance));
+			ConstraintInstanceProperty = DetailBuilder.GetProperty(GET_MEMBER_NAME_CHECKED(UPhysicsConstraintTemplate, DefaultInstance));
+			ConstraintTemplate = Cast<UPhysicsConstraintTemplate>(Objects[i].Get());
+			ParentPhysicsAsset = Cast<UPhysicsAsset>(Objects[i]->GetOuter());
 			bInPhat = true;
 			break;
 		}
 		else if (Objects[i]->IsA(UPhysicsConstraintComponent::StaticClass()))
 		{
-			ConstraintInstance = DetailBuilder.GetProperty(GET_MEMBER_NAME_CHECKED(UPhysicsConstraintComponent, ConstraintInstance));
+			ConstraintInstanceProperty = DetailBuilder.GetProperty(GET_MEMBER_NAME_CHECKED(UPhysicsConstraintComponent, ConstraintInstance));
 			ConstraintComp = (UPhysicsConstraintComponent*)Objects[i].Get();
 			OwningConstraintActor = Cast<APhysicsConstraintActor>(ConstraintComp->GetOwner());
 			break;
 		}
 	}
 
-	IDetailCategoryBuilder& ConstraintCategory = DetailBuilder.EditCategory("Constraint");	//Create this category first so it's at the top
-	DetailBuilder.EditCategory("Constraint Behavior");	//Create this category first so it's at the top
-
-	TSharedPtr<IPropertyHandle> ProfileInstance = ConstraintInstance->GetChildHandle(GET_MEMBER_NAME_CHECKED(FConstraintInstance, ProfileInstance));
-	AddLinearLimits(DetailBuilder, ConstraintInstance, ProfileInstance);
-	AddAngularLimits(DetailBuilder, ConstraintInstance, ProfileInstance);
-	AddLinearDrive(DetailBuilder, ConstraintInstance, ProfileInstance);
-	AddAngularDrive(DetailBuilder, ConstraintInstance, ProfileInstance);
+	AddConstraintProperties(DetailBuilder, ConstraintInstanceProperty, Objects);
 
-	AddConstraintBehaviorProperties(DetailBuilder, ConstraintInstance, ProfileInstance);	//Now we've added all the complex UI, just dump the rest into Constraint category
+	DetailBuilder.EditCategory("Constraint Behavior");	//Create this category first so it's at the top
 
-	if(bInPhat)
-	{
-		ConstraintCategory.HeaderContent(
-			SNew(SHorizontalBox)
-			+SHorizontalBox::Slot()
-			.FillWidth(1.0f)
-			.HAlign(HAlign_Right)
-			.VAlign(VAlign_Center)
-			[
-				SNew(SRichTextBlock)
-				.DecoratorStyleSet(&FAppStyle::Get())
-				.Text_Lambda([Objects]()
-				{
-					if (Objects.Num() > 0)
-					{
-						if (UPhysicsConstraintTemplate* Constraint = Cast<UPhysicsConstraintTemplate>(Objects[0].Get()))
-						{
-							FName CurrentProfileName = Constraint->GetCurrentConstraintProfileName();
-							if(CurrentProfileName != NAME_None)
-							{
-								if(Constraint->ContainsConstraintProfile(CurrentProfileName))
-								{
-									return FText::Format(LOCTEXT("ProfileFormatAssigned", "Assigned to Profile: <RichTextBlock.Bold>{0}</>"), FText::FromName(CurrentProfileName));
-								}
-								else
-								{
-									return FText::Format(LOCTEXT("ProfileFormatNotAssigned", "Not Assigned to Profile: <RichTextBlock.Bold>{0}</>"), FText::FromName(CurrentProfileName));
-								}
-							}
-							else
-							{
-								return LOCTEXT("ProfileFormatNone", "Current Profile: <RichTextBlock.Bold>None</>");
-							}
-						}
-					}
+	TSharedPtr<IPropertyHandle> ProfileInstance = ConstraintInstanceProperty->GetChildHandle(GET_MEMBER_NAME_CHECKED(FConstraintInstance, ProfileInstance));
+	AddLinearLimits(DetailBuilder, ConstraintInstanceProperty, ProfileInstance);
+	AddAngularLimits(DetailBuilder, ConstraintInstanceProperty, ProfileInstance);
+	AddLinearDrive(DetailBuilder, ConstraintInstanceProperty, ProfileInstance);
+	AddAngularDrive(DetailBuilder, ConstraintInstanceProperty, ProfileInstance);
 
-					return FText();
-				})
-			]);
-	}
+	AddConstraintBehaviorProperties(DetailBuilder, ConstraintInstanceProperty, ProfileInstance);	//Now we've added all the complex UI, just dump the rest into Constraint category
 }
 
 
@@ -991,4 +1941,21 @@ void FPhysicsConstraintComponentDetails::OnLimitRadioChanged( ECheckBoxState Che
 	}
 }
 
+FConstraintInstance* FPhysicsConstraintComponentDetails::GetConstraintInstance()
+{
+	FConstraintInstance* ConstraintInstance = nullptr;
+
+	if (ConstraintTemplate)
+	{
+		ConstraintInstance = &ConstraintTemplate->DefaultInstance;
+	}
+	else if (ConstraintComp)
+	{
+		ConstraintInstance = &ConstraintComp->ConstraintInstance;
+	}
+
+	return ConstraintInstance;
+}
+
+
 #undef LOCTEXT_NAMESPACE
diff --git a/Engine/Source/Editor/DetailCustomizations/Private/PhysicsConstraintComponentDetails.h b/Engine/Source/Editor/DetailCustomizations/Private/PhysicsConstraintComponentDetails.h
--- a/Engine/Source/Editor/DetailCustomizations/Private/PhysicsConstraintComponentDetails.h
+++ b/Engine/Source/Editor/DetailCustomizations/Private/PhysicsConstraintComponentDetails.h
@@ -5,10 +5,101 @@
 #include "CoreMinimal.h"
 #include "UObject/WeakObjectPtr.h"
 #include "IDetailCustomization.h"
+#include "Chaos/ChaosEngineInterface.h"
+#include "Customizations/MathStructProxyCustomizations.h"
 
 class IDetailLayoutBuilder;
+class IDetailCategoryBuilder;
 class IPropertyHandle;
+class UPhysicsAsset;
+class UPhysicsConstraintComponent;
+class UPhysicsConstraintTemplate;
 enum class ECheckBoxState : uint8;
+enum class EConstraintTransformComponentFlags : uint8;
+
+/**
+ * Proxy class customization that displays a constraint transform in a user friendly way.
+ * 
+ * Constraint transforms are described by a position vector and a pair of orthonormal vectors 
+ * representing the orientation. This class displays the transform as a position vector and a
+ * rotator in the UI details panel in a way that can be easily read and safely edited by users. 
+ */
+class FConstraintTransformCustomization : public FMathStructCustomization
+{
+public:
+
+	FConstraintTransformCustomization();
+
+	static TSharedRef<IPropertyTypeCustomization> MakeInstance();
+
+	void MakeRotationRow(TSharedRef<IPropertyHandle>& InPriAxisPropertyHandle, TSharedRef<IPropertyHandle>& InSecAxisPropertyHandle, FDetailWidgetRow& Row, TSharedRef<SWidget> EditSpaceToggleButtonWidget);
+	void MakePositionRow(TSharedRef<IPropertyHandle>& InPositionPropertyHandle, FDetailWidgetRow& Row, TSharedRef<SWidget> EditSpaceToggleButtonWidget);
+
+	template<typename ProxyType, typename NumericType> TSharedRef<SWidget> MakeNumericProxyWidget(TSharedRef<IPropertyHandle>& PriAxisPropertyHandle, TSharedRef<IPropertyHandle>& SecAxisPropertyHandle, TSharedRef< TProxyProperty<ProxyType, NumericType> >& ProxyValue, const FText& Label, bool bRotationInDegrees, const FLinearColor& LabelBackgroundColor);
+	template<typename ProxyType, typename NumericType> TSharedRef<SWidget> MakeNumericProxyWidget(TSharedRef<IPropertyHandle>& PositionPropertyHandle, TSharedRef< TProxyProperty<ProxyType, NumericType> >& ProxyValue, const FText& Label, const FLinearColor& LabelBackgroundColor);
+
+	void OnCopy(TWeakPtr<IPropertyHandle> PositionPropertyHandlePtr, TWeakPtr<IPropertyHandle> PriAxisPropertyHandlePtr, TWeakPtr<IPropertyHandle> SecAxisPropertyHandlePtr);
+	void OnCopyPosition(TWeakPtr<IPropertyHandle> PositionPropertyHandlePtr);
+	void OnCopyRotation(TWeakPtr<IPropertyHandle> PriAxisPropertyHandlePtr, TWeakPtr<IPropertyHandle> SecAxisPropertyHandlePtr);
+	void OnPaste(TWeakPtr<IPropertyHandle> PositionPropertyHandlePtr, TWeakPtr<IPropertyHandle> PriAxisPropertyHandlePtr, TWeakPtr<IPropertyHandle> SecAxisPropertyHandlePtr);
+	void OnPastePosition(TWeakPtr<IPropertyHandle> PositionPropertyHandlePtr);
+	void OnPasteRotation(TWeakPtr<IPropertyHandle> PriAxisPropertyHandlePtr, TWeakPtr<IPropertyHandle> SecAxisPropertyHandlePtr);
+
+	bool IsRotationValueEnabled(TWeakPtr<IPropertyHandle> PriAxisPropertyHandlePtr, TWeakPtr<IPropertyHandle> SecAxisPropertyHandlePtr) const;
+	template<typename ProxyType, typename NumericType> TOptional<NumericType> OnGetRotationValue(TWeakPtr<IPropertyHandle> PriAxisPropertyHandlePtr, TWeakPtr<IPropertyHandle> SecAxisPropertyHandlePtr, TSharedRef< TProxyProperty<ProxyType, NumericType> > ProxyValue) const;
+	template<typename ProxyType, typename NumericType> void OnRotationValueCommitted(NumericType NewValue, ETextCommit::Type CommitType, TWeakPtr<IPropertyHandle> PriAxisPropertyHandlePtr, TWeakPtr<IPropertyHandle> SecAxisPropertyHandlePtr, TSharedRef< TProxyProperty<ProxyType, NumericType> > ProxyValue);
+	template<typename ProxyType, typename NumericType> void OnRotationValueChanged(NumericType NewValue, TWeakPtr<IPropertyHandle> PriAxisPropertyHandlePtr, TWeakPtr<IPropertyHandle> SecAxisPropertyHandlePtr, TSharedRef< TProxyProperty<ProxyType, NumericType> > ProxyValue);
+
+	bool IsPositionValueEnabled(TWeakPtr<IPropertyHandle> PositionWeakHandlePtr) const;
+	template<typename ProxyType, typename NumericType> TOptional<NumericType> OnGetPositionValue(TWeakPtr<IPropertyHandle> PositionWeakHandlePtr, TSharedRef< TProxyProperty<ProxyType, NumericType> > ProxyValue) const;
+	template<typename ProxyType, typename NumericType> void OnPositionValueCommitted(NumericType NewValue, ETextCommit::Type CommitType, TWeakPtr<IPropertyHandle> PositionWeakHandlePtr, TSharedRef< TProxyProperty<ProxyType, NumericType> > ProxyValue);
+	template<typename ProxyType, typename NumericType> void OnPositionValueChanged(NumericType NewValue, TWeakPtr<IPropertyHandle> PositionWeakHandlePtr, TSharedRef< TProxyProperty<ProxyType, NumericType> > ProxyValue);
+
+	void SetFrameLabelText(const FText InText);
+	FText& GetFrameLabelText();
+
+	void GetPositionAsFormattedString(FString& OutString);
+	void GetRotationAsFormattedString(FString& OutString);
+	void GetValueAsFormattedString(FString& OutString);
+	bool SetPositionFromFormattedString(const FString& InString);
+	bool SetRotationFromFormattedString(const FString& InString);
+	bool SetValueFromFormattedString(const FString& InString);
+
+	virtual bool CacheValues(TWeakPtr<IPropertyHandle> PositionPropertyHandlePtr) const;
+	virtual bool CacheValues(TWeakPtr<IPropertyHandle> PriAxisPropertyHandlePtr, TWeakPtr<IPropertyHandle> SecAxisPropertyHandlePtr) const;
+	virtual bool FlushValues(TWeakPtr<IPropertyHandle> PositionPropertyHandlePtr) const;
+	virtual bool FlushValues(TWeakPtr<IPropertyHandle> PriAxisPropertyHandlePtr, TWeakPtr<IPropertyHandle> SecAxisPropertyHandlePtr) const;
+
+	const FTransform& GetDefaultTransform() const { return DefaultTransform; };
+
+	void SetDefaultTransform(const FTransform& InTransform, TWeakPtr<IPropertyHandle> PositionPropertyHandlePtr, TWeakPtr<IPropertyHandle> PriAxisPropertyHandlePtr, TWeakPtr<IPropertyHandle> SecAxisPropertyHandlePtr);
+	void SetPositionDisplayRelativeToDefault(const bool bValue) { bPositionDisplayRelativeToDefault = bValue; }
+	void SetRotationDisplayRelativeToDefault(const bool bValue) { bRotationDisplayRelativeToDefault = bValue; }
+
+protected:
+	void OrthonormalVectorPairToDisplayedRotator(const FVector& PriAxis, const FVector& SecAxis, FRotator& OutRotator) const;
+	void DisplayedRotatorToOrthonormalVectorPair(const FRotator& InRotator, FVector& OutPriAxis, FVector& OutSecAxis) const;
+
+	mutable TSharedRef< TProxyValue<FRotator> > CachedRotation;
+	mutable TSharedRef< TProxyProperty<FRotator, FRotator::FReal> > CachedRotationYaw;
+	mutable TSharedRef< TProxyProperty<FRotator, FRotator::FReal> > CachedRotationPitch;
+	mutable TSharedRef< TProxyProperty<FRotator, FRotator::FReal> > CachedRotationRoll;
+
+	mutable TSharedRef< TProxyValue<FVector> > CachedPosition;
+	mutable TSharedRef< TProxyProperty<FVector, FRotator::FReal> > CachedPositionX;
+	mutable TSharedRef< TProxyProperty<FVector, FRotator::FReal> > CachedPositionY;
+	mutable TSharedRef< TProxyProperty<FVector, FRotator::FReal> > CachedPositionZ;
+
+	// Constraint transforms are displayed relative to this default transform when the associated ...DisplayRelativeToDefault flag is true. Otherwise they are displayed in the frame of the associated bone.
+	FTransform DefaultTransform;
+	FTransform InverseDefaultTransform;
+
+	// When true, the Position/Rotation components are displayed as offsets from the default transforms.
+	bool bPositionDisplayRelativeToDefault;
+	bool bRotationDisplayRelativeToDefault;
+
+	FText FrameLabelText;
+};
 
 /** 
  * Detail customizer for PhysicsConstraintComponent and PhysicsConstraintTemplate
@@ -16,6 +107,8 @@ enum class ECheckBoxState : uint8;
 class FPhysicsConstraintComponentDetails : public IDetailCustomization
 {
 public:
+	FPhysicsConstraintComponentDetails();
+
 	/** Makes a new instance of this detail layout class for a specific detail view requesting it */
 	static TSharedRef<IDetailCustomization> MakeInstance();
 
@@ -50,13 +143,33 @@ private:
 	ECheckBoxState IsLimitRadioChecked(TSharedPtr<IPropertyHandle> Property, uint8 Value) const ;
 	void OnLimitRadioChanged(ECheckBoxState CheckType, TSharedPtr<IPropertyHandle> Property, uint8 Value);
 
-	void AddConstraintBehaviorProperties(IDetailLayoutBuilder& DetailBuilder, TSharedPtr<IPropertyHandle> ConstraintInstance, TSharedPtr<IPropertyHandle> ProfileInstance);
-	void AddLinearLimits(IDetailLayoutBuilder& DetailBuilder, TSharedPtr<IPropertyHandle> ConstraintInstance, TSharedPtr<IPropertyHandle> ProfileInstance);
-	void AddAngularLimits(IDetailLayoutBuilder& DetailBuilder, TSharedPtr<IPropertyHandle> ConstraintInstance, TSharedPtr<IPropertyHandle> ProfileInstance);
-	void AddLinearDrive(IDetailLayoutBuilder& DetailBuilder, TSharedPtr<IPropertyHandle> ConstraintInstance, TSharedPtr<IPropertyHandle> ProfileInstance);
-	void AddAngularDrive(IDetailLayoutBuilder& DetailBuilder, TSharedPtr<IPropertyHandle> ConstraintInstance, TSharedPtr<IPropertyHandle> ProfileInstance);
+	void AddConstraintProperties(IDetailLayoutBuilder& DetailBuilder, TSharedPtr<IPropertyHandle> ConstraintInstancePropertyHandle, TArray<TWeakObjectPtr<UObject>>& Objects);
+	void AddConstraintFrameTransform(const EConstraintFrame::Type ConstraintFrameType, IDetailCategoryBuilder& ConstraintCat, TSharedPtr<IPropertyHandle> ConstraintInstancePropertyHandle);
+	void AddConstraintBehaviorProperties(IDetailLayoutBuilder& DetailBuilder, TSharedPtr<IPropertyHandle> ConstraintInstancePropertyHandle, TSharedPtr<IPropertyHandle> ProfileInstance);
+	void AddLinearLimits(IDetailLayoutBuilder& DetailBuilder, TSharedPtr<IPropertyHandle> ConstraintInstancePropertyHandle, TSharedPtr<IPropertyHandle> ProfileInstance);
+	void AddAngularLimits(IDetailLayoutBuilder& DetailBuilder, TSharedPtr<IPropertyHandle> ConstraintInstancePropertyHandle, TSharedPtr<IPropertyHandle> ProfileInstance);
+	void AddLinearDrive(IDetailLayoutBuilder& DetailBuilder, TSharedPtr<IPropertyHandle> ConstraintInstancePropertyHandle, TSharedPtr<IPropertyHandle> ProfileInstance);
+	void AddAngularDrive(IDetailLayoutBuilder& DetailBuilder, TSharedPtr<IPropertyHandle> ConstraintInstancePropertyHandle, TSharedPtr<IPropertyHandle> ProfileInstance);
+
+	void OnCopyConstraintTransform(TSharedPtr<IPropertyHandle> PositionPropertyHandle, TSharedPtr<IPropertyHandle> PriAxisPropertyHandle, TSharedPtr<IPropertyHandle> SecAxisPropertyHandle, TSharedPtr<FConstraintTransformCustomization> RotationProxy);
+	void OnPasteConstraintTransform(TSharedPtr<IPropertyHandle> PositionPropertyHandle, TSharedPtr<IPropertyHandle> PriAxisPropertyHandle, TSharedPtr<IPropertyHandle> SecAxisPropertyHandle, TSharedPtr<FConstraintTransformCustomization> RotationProxy);
+
+	bool IsSnapConstraintTransformComponentVisible(TSharedPtr<IPropertyHandle> ConstraintInstancePropertyHandle, const EConstraintTransformComponentFlags SnapFlags);
+	void ToggleDisplayConstraintTransformComponentRelativeToDefault(TSharedPtr<IPropertyHandle> ConstraintInstance, const EConstraintTransformComponentFlags ComponentFlags);
+	bool IsDisplayingConstraintTransformComponentRelativeToDefault(const EConstraintTransformComponentFlags ComponentFlags);
+	bool IsConstraintTransformComponentEnabled(const EConstraintTransformComponentFlags ComponentFlags) const;
+	FSlateColor GetConstraintTransformColorAndOpacity(const EConstraintTransformComponentFlags ComponentFlags) const;
+
 private:
 
+	TSharedRef<SWidget> MakeEditSpaceToggleButtonWidget(TSharedPtr<IPropertyHandle> ConstraintInstancePropertyHandle, const EConstraintTransformComponentFlags ComponentFlags); // TODO - get this in the right place in src
+	void SnapConstraintTransformComponentsToDefault(TSharedPtr<IPropertyHandle> ConstraintInstancePropertyHandle, const EConstraintTransformComponentFlags SnapFlags);
+	void UpdateTransformProxyDisplayRelativeToDefault(TSharedPtr<IPropertyHandle> ConstraintInstancePropertyHandle);
+	FConstraintInstance* GetConstraintInstance();
+
+	TSharedRef<FConstraintTransformCustomization> ChildTransformProxy; // Creates a user friendly representation of the constraint's transform in the frame of the child bone in the UI details panel.
+	TSharedRef<FConstraintTransformCustomization> ParentTransformProxy;// Creates a user friendly representation of the constraint's transform in the frame of the parent bone in the UI details panel.
+
 	TSharedPtr<IPropertyHandle> LinearXPositionDriveProperty;
 	TSharedPtr<IPropertyHandle> LinearYPositionDriveProperty;
 	TSharedPtr<IPropertyHandle> LinearZPositionDriveProperty;
@@ -69,7 +182,19 @@ private:
 	TSharedPtr<IPropertyHandle> AngularSwing2MotionProperty;
 	TSharedPtr<IPropertyHandle> AngularTwistMotionProperty;
 
-	TWeakObjectPtr<UObject> ConstraintComponent;
+	TSharedPtr<IPropertyHandle> ChildPositionPropertyHandle;
+	TSharedPtr<IPropertyHandle> ChildPriAxisPropertyHandle;
+	TSharedPtr<IPropertyHandle> ChildSecAxisPropertyHandle;
 	
+	TSharedPtr<IPropertyHandle> ParentPositionPropertyHandle;
+	TSharedPtr<IPropertyHandle> ParentPriAxisPropertyHandle;
+	TSharedPtr<IPropertyHandle> ParentSecAxisPropertyHandle;
+
+//	TWeakObjectPtr<UObject> ConstraintComponent;
+	UPhysicsConstraintComponent* ConstraintComp;
+	UPhysicsConstraintTemplate* ConstraintTemplate;
+
+	UPhysicsAsset* ParentPhysicsAsset;
+
 	bool bInPhat;
 };
diff --git a/Engine/Source/Editor/EditorStyle/Private/StarshipStyle.cpp b/Engine/Source/Editor/EditorStyle/Private/StarshipStyle.cpp
--- a/Engine/Source/Editor/EditorStyle/Private/StarshipStyle.cpp
+++ b/Engine/Source/Editor/EditorStyle/Private/StarshipStyle.cpp
@@ -6907,6 +6907,10 @@ void FStarshipEditorStyle::FStyle::SetupToolkitStyles()
 		Set( "PhysicsAssetEditor.DuplicatePrimitive", new IMAGE_BRUSH( "PhysicsAssetEditor/icon_PhAT_DupePrim_40x", Icon40x40 ) );
 		Set( "PhysicsAssetEditor.ResetConstraint", new IMAGE_BRUSH( "PhysicsAssetEditor/icon_PhAT_ResetConstraint_40x", Icon40x40 ) );
 		Set( "PhysicsAssetEditor.SnapConstraint", new IMAGE_BRUSH( "PhysicsAssetEditor/icon_PhAT_SnapConstraint_40x", Icon40x40 ) );
+		Set( "PhysicsAssetEditor.SnapConstraintChildPosition", new IMAGE_BRUSH("PhysicsAssetEditor/icon_PhAT_SnapConstraint_40x", Icon40x40));
+		Set( "PhysicsAssetEditor.SnapConstraintChildOrientation", new IMAGE_BRUSH("PhysicsAssetEditor/icon_PhAT_SnapConstraint_40x", Icon40x40));
+		Set( "PhysicsAssetEditor.SnapConstraintParentPosition", new IMAGE_BRUSH("PhysicsAssetEditor/icon_PhAT_SnapConstraint_40x", Icon40x40));
+		Set( "PhysicsAssetEditor.SnapConstraintParentOrientation", new IMAGE_BRUSH("PhysicsAssetEditor/icon_PhAT_SnapConstraint_40x", Icon40x40));
 		Set( "PhysicsAssetEditor.SnapAllConstraints", new IMAGE_BRUSH( "PhysicsAssetEditor/icon_PhAT_SnapAll_40x", Icon40x40 ) );
 		Set( "PhysicsAssetEditor.ConvertToBallAndSocket", new IMAGE_BRUSH_SVG("Starship/Persona/AnimationToBallAndSocket", Icon20x20));
 		Set( "PhysicsAssetEditor.ConvertToHinge", new IMAGE_BRUSH_SVG("Starship/Persona/AnimationToHinge", Icon20x20));
diff --git a/Engine/Source/Editor/Persona/Private/PhysicsAssetRenderUtils.cpp b/Engine/Source/Editor/Persona/Private/PhysicsAssetRenderUtils.cpp
--- a/Engine/Source/Editor/Persona/Private/PhysicsAssetRenderUtils.cpp
+++ b/Engine/Source/Editor/Persona/Private/PhysicsAssetRenderUtils.cpp
@@ -45,6 +45,8 @@ FTransform GetConstraintMatrix(const USkeletalMeshComponent* const SkeletalMeshC
 FPhysicsAssetRenderSettings::FPhysicsAssetRenderSettings()
 	: CollisionViewMode(EPhysicsAssetEditorCollisionViewMode::Solid)
 	, ConstraintViewMode(EPhysicsAssetEditorConstraintViewMode::AllLimits)
+	, ConstraintViewportManipulationFlags(EConstraintTransformComponentFlags::All)
+	, ConstraintTransformComponentDisplayRelativeToDefaultFlags(EConstraintTransformComponentFlags::None)
 	, ConstraintDrawSize(1.0f)
 	, PhysicsBlend(1.0f)
 	, bHideKinematicBodies(false)
@@ -230,6 +232,21 @@ void FPhysicsAssetRenderSettings::SetHiddenConstraints(const TArray<int32>& InHi
 	HiddenConstraints.Append(InHiddenConstraints);
 }
 
+EConstraintTransformComponentFlags FPhysicsAssetRenderSettings::GetConstraintViewportManipulationFlags() const
+{
+	return ConstraintViewportManipulationFlags;
+}
+
+bool FPhysicsAssetRenderSettings::IsDisplayingConstraintTransformComponentRelativeToDefault(const EConstraintTransformComponentFlags ComponentFlags) const
+{
+	return EnumHasAllFlags(ConstraintTransformComponentDisplayRelativeToDefaultFlags, ComponentFlags);
+}
+
+void FPhysicsAssetRenderSettings::SetDisplayConstraintTransformComponentRelativeToDefault(const EConstraintTransformComponentFlags ComponentFlags, const bool bShouldDisplayRelativeToDefault)
+{
+	ConstraintTransformComponentDisplayRelativeToDefaultFlags = (bShouldDisplayRelativeToDefault) ? (ConstraintTransformComponentDisplayRelativeToDefaultFlags | ComponentFlags) : (ConstraintTransformComponentDisplayRelativeToDefaultFlags & ~ComponentFlags);
+}
+
 void FPhysicsAssetRenderSettings::ResetEditorViewportOptions()
 {
 	const FPhysicsAssetRenderSettings DefaultObject = FPhysicsAssetRenderSettings();
@@ -720,4 +737,32 @@ bool FPhysicsAssetRenderInterface::AreAnyConstraintsHidden(class UPhysicsAsset*
 	return false;
 }
 
+EConstraintTransformComponentFlags FPhysicsAssetRenderInterface::GetConstraintViewportManipulationFlags(class UPhysicsAsset* const PhysicsAsset)
+{
+	if (FPhysicsAssetRenderSettings* PhysicsAssetRenderSettings = UPhysicsAssetRenderUtilities::GetSettings(PhysicsAsset))
+	{
+		return PhysicsAssetRenderSettings->GetConstraintViewportManipulationFlags();
+	}
+
+	return EConstraintTransformComponentFlags::None;
+}
+
+bool FPhysicsAssetRenderInterface::IsDisplayingConstraintTransformComponentRelativeToDefault(class UPhysicsAsset* const PhysicsAsset, const EConstraintTransformComponentFlags ComponentFlags)
+{
+	if (FPhysicsAssetRenderSettings* PhysicsAssetRenderSettings = UPhysicsAssetRenderUtilities::GetSettings(PhysicsAsset))
+	{
+		return PhysicsAssetRenderSettings->IsDisplayingConstraintTransformComponentRelativeToDefault(ComponentFlags);
+	}
+
+	return false;
+}
+
+void FPhysicsAssetRenderInterface::SetDisplayConstraintTransformComponentRelativeToDefault(class UPhysicsAsset* const PhysicsAsset, const EConstraintTransformComponentFlags ComponentFlags, const bool bShouldDisplayRelativeToDefault)
+{
+	if (FPhysicsAssetRenderSettings* PhysicsAssetRenderSettings = UPhysicsAssetRenderUtilities::GetSettings(PhysicsAsset))
+	{
+		PhysicsAssetRenderSettings->SetDisplayConstraintTransformComponentRelativeToDefault(ComponentFlags, bShouldDisplayRelativeToDefault);
+	}
+}
+
 #undef LOCTEXT_NAMESPACE
diff --git a/Engine/Source/Editor/Persona/Public/IPhysicsAssetRenderInterface.h b/Engine/Source/Editor/Persona/Public/IPhysicsAssetRenderInterface.h
--- a/Engine/Source/Editor/Persona/Public/IPhysicsAssetRenderInterface.h
+++ b/Engine/Source/Editor/Persona/Public/IPhysicsAssetRenderInterface.h
@@ -6,6 +6,8 @@
 #include "Features/IModularFeature.h"
 #include "Modules/ModuleInterface.h"
 
+enum class EConstraintTransformComponentFlags : uint8;
+
 class IPhysicsAssetRenderInterface : public IModuleInterface, public IModularFeature
 {
 	// virtual ~IPhysicsAssetRenderInterface() = default;
@@ -22,4 +24,8 @@ public:
 	virtual void ToggleShowAllConstraints(class UPhysicsAsset* const PhysicsAsset) = 0;
 	virtual bool AreAnyBodiesHidden(class UPhysicsAsset* const PhysicsAsset) = 0;
 	virtual bool AreAnyConstraintsHidden(class UPhysicsAsset* const PhysicsAsset) = 0;
+
+	virtual EConstraintTransformComponentFlags GetConstraintViewportManipulationFlags(class UPhysicsAsset* const PhysicsAsset) = 0;
+	virtual bool IsDisplayingConstraintTransformComponentRelativeToDefault(class UPhysicsAsset* const PhysicsAsset, const EConstraintTransformComponentFlags ComponentFlags) = 0;
+	virtual void SetDisplayConstraintTransformComponentRelativeToDefault(class UPhysicsAsset* const PhysicsAsset, const EConstraintTransformComponentFlags ComponentFlags, const bool bShouldDisplayRelativeToDefault) = 0;
 };
diff --git a/Engine/Source/Editor/Persona/Public/PhysicsAssetRenderUtils.h b/Engine/Source/Editor/Persona/Public/PhysicsAssetRenderUtils.h
--- a/Engine/Source/Editor/Persona/Public/PhysicsAssetRenderUtils.h
+++ b/Engine/Source/Editor/Persona/Public/PhysicsAssetRenderUtils.h
@@ -6,6 +6,7 @@
 #include "IPhysicsAssetRenderInterface.h"
 #include "UObject/ObjectMacros.h"
 #include "Math/Color.h"
+#include "PhysicsEngine/ConstraintInstance.h"
 #include "PhysicsEngine/ShapeElem.h"
 #include "SceneManagement.h"
 #include "PhysicsAssetRenderUtils.generated.h"
@@ -61,6 +62,15 @@ public:
 	void SetHiddenBodies(const TArray<int32>& InHiddenBodies);
 	void SetHiddenConstraints(const TArray<int32>& InHiddenConstraints);
 
+	/** Returns a set of flags describing which components of the selected constraint's transforms are being manipulated in the view port. */
+	EConstraintTransformComponentFlags GetConstraintViewportManipulationFlags() const;
+
+	/** Returns true if all the constraint transform components specified by the flags should be displayed as an offset from the default (snapped) transforms. */
+	bool IsDisplayingConstraintTransformComponentRelativeToDefault(const EConstraintTransformComponentFlags ComponentFlags) const;
+
+	/** Set how the constraint transform components specified by the flags should be displayed, in the frame of their associated physics body (false) or as an offset from the default (snapped) transforms (true). */
+	void SetDisplayConstraintTransformComponentRelativeToDefault(const EConstraintTransformComponentFlags ComponentFlags, const bool bShouldDisplayRelativeToDefault);
+
 	void ResetEditorViewportOptions();
 
 	// Physics Asset Editor Viewport Options
@@ -70,6 +80,14 @@ public:
 	UPROPERTY()
 	EPhysicsAssetEditorConstraintViewMode ConstraintViewMode;
 
+	// Flags that determine which parts of the constraints transforms (parent frame, child frame, position and rotation) are currently begin manipulated in the viewport.
+	UPROPERTY(Transient)
+	EConstraintTransformComponentFlags ConstraintViewportManipulationFlags;
+
+	// Flags that determine which parts of the constraints transforms (parent/child position/rotation) should be displayed as an offset from the default (snapped) transforms.
+	UPROPERTY()
+	EConstraintTransformComponentFlags ConstraintTransformComponentDisplayRelativeToDefaultFlags;
+
 	UPROPERTY()
 	float ConstraintDrawSize;
 
@@ -207,4 +225,13 @@ public:
 	virtual void ToggleShowAllConstraints(UPhysicsAsset* const PhysicsAsset) override;
 	virtual bool AreAnyBodiesHidden(UPhysicsAsset* const PhysicsAsset) override;
 	virtual bool AreAnyConstraintsHidden(UPhysicsAsset* const PhysicsAsset) override;
+
+	/** Returns a set of flags describing which components of the selected constraint's transforms are being manipulated in the viewport. */
+	virtual EConstraintTransformComponentFlags GetConstraintViewportManipulationFlags(class UPhysicsAsset* const PhysicsAsset) override;
+
+	/** Returns true if the constraint transform component specified by the flags should be displayed as an offset from the default (snapped) transforms. */
+	virtual bool IsDisplayingConstraintTransformComponentRelativeToDefault(class UPhysicsAsset* const PhysicsAsset, const EConstraintTransformComponentFlags ComponentFlags) override;
+
+	/** Change how the constraint transform component specified by the flags should be displayed, in the frame of their associated physics body or as an offset from the default (snapped) transforms. */
+	virtual void SetDisplayConstraintTransformComponentRelativeToDefault(class UPhysicsAsset* const PhysicsAsset, const EConstraintTransformComponentFlags ComponentFlags, const bool bShouldDisplayRelativeToDefault) override;
 };
\ No newline at end of file
diff --git a/Engine/Source/Editor/PhysicsAssetEditor/Private/PhysicsAssetEditor.cpp b/Engine/Source/Editor/PhysicsAssetEditor/Private/PhysicsAssetEditor.cpp
--- a/Engine/Source/Editor/PhysicsAssetEditor/Private/PhysicsAssetEditor.cpp
+++ b/Engine/Source/Editor/PhysicsAssetEditor/Private/PhysicsAssetEditor.cpp
@@ -923,7 +923,27 @@ void FPhysicsAssetEditor::BindCommands()
 
 	ToolkitCommands->MapAction(
 		Commands.SnapConstraint,
-		FExecuteAction::CreateSP(this, &FPhysicsAssetEditor::OnSnapConstraint),
+		FExecuteAction::CreateSP(this, &FPhysicsAssetEditor::OnSnapConstraint, EConstraintTransformComponentFlags::All),
+		FCanExecuteAction::CreateSP(this, &FPhysicsAssetEditor::HasSelectedConstraintAndIsNotSimulation));
+
+	ToolkitCommands->MapAction(
+		Commands.SnapConstraintChildPosition,
+		FExecuteAction::CreateSP(this, &FPhysicsAssetEditor::OnSnapConstraint, EConstraintTransformComponentFlags::ChildPosition),
+		FCanExecuteAction::CreateSP(this, &FPhysicsAssetEditor::HasSelectedConstraintAndIsNotSimulation));
+
+	ToolkitCommands->MapAction(
+		Commands.SnapConstraintChildOrientation,
+		FExecuteAction::CreateSP(this, &FPhysicsAssetEditor::OnSnapConstraint, EConstraintTransformComponentFlags::ChildRotation),
+		FCanExecuteAction::CreateSP(this, &FPhysicsAssetEditor::HasSelectedConstraintAndIsNotSimulation));
+
+	ToolkitCommands->MapAction(
+		Commands.SnapConstraintParentPosition,
+		FExecuteAction::CreateSP(this, &FPhysicsAssetEditor::OnSnapConstraint, EConstraintTransformComponentFlags::ParentPosition),
+		FCanExecuteAction::CreateSP(this, &FPhysicsAssetEditor::HasSelectedConstraintAndIsNotSimulation));
+
+	ToolkitCommands->MapAction(
+		Commands.SnapConstraintParentOrientation,
+		FExecuteAction::CreateSP(this, &FPhysicsAssetEditor::OnSnapConstraint, EConstraintTransformComponentFlags::ParentRotation),
 		FCanExecuteAction::CreateSP(this, &FPhysicsAssetEditor::HasSelectedConstraintAndIsNotSimulation));
 
 	ToolkitCommands->MapAction(
@@ -1533,11 +1553,26 @@ void FPhysicsAssetEditor::BuildMenuWidgetConstraint(FMenuBuilder& InMenuBuilder)
 				InSubMenuBuilder.AddMenuEntry(PhysicsAssetEditorCommands.ConvertToSkeletal);
 				InSubMenuBuilder.EndSection();
 			}
+
+			static void FillSnapMenu(FMenuBuilder& InSubMenuBuilder)
+			{
+				const FPhysicsAssetEditorCommands& PhysicsAssetEditorCommands = FPhysicsAssetEditorCommands::Get();
+
+				InSubMenuBuilder.BeginSection("SnapHeader", LOCTEXT("SnapHeader", "Snap"));
+				InSubMenuBuilder.AddMenuEntry(PhysicsAssetEditorCommands.SnapConstraint);
+				InSubMenuBuilder.AddMenuEntry(PhysicsAssetEditorCommands.SnapConstraintChildPosition);
+				InSubMenuBuilder.AddMenuEntry(PhysicsAssetEditorCommands.SnapConstraintChildOrientation);
+				InSubMenuBuilder.AddMenuEntry(PhysicsAssetEditorCommands.SnapConstraintParentPosition);
+				InSubMenuBuilder.AddMenuEntry(PhysicsAssetEditorCommands.SnapConstraintParentOrientation);
+				InSubMenuBuilder.EndSection();
+			}
 		};
 
 		InMenuBuilder.BeginSection("EditTypeActions", LOCTEXT("ConstraintEditTypeHeader", "Edit"));
 
-		InMenuBuilder.AddMenuEntry(Commands.SnapConstraint);
+		InMenuBuilder.AddSubMenu(LOCTEXT("SnapMenu", "Snap"), LOCTEXT("SnapMenu_ToolTip", "Set constraint transforms to defaults"),
+			FNewMenuDelegate::CreateStatic(&FLocal::FillSnapMenu));
+
 		InMenuBuilder.AddMenuEntry(Commands.ResetConstraint);
 		
 		InMenuBuilder.AddSubMenu( LOCTEXT("AxesAndLimitsMenu", "Axes and Limits"), LOCTEXT("AxesAndLimitsMenu_ToolTip", "Edit axes and limits of this constraint"),
@@ -2732,13 +2767,13 @@ void FPhysicsAssetEditor::OnResetConstraint()
 	RefreshPreviewViewport();
 }
 
-void FPhysicsAssetEditor::OnSnapConstraint()
+void FPhysicsAssetEditor::OnSnapConstraint(const EConstraintTransformComponentFlags ComponentFlags)
 {
 	const FScopedTransaction Transaction( LOCTEXT( "SnapConstraints", "Snap Constraints" ) );
 
 	for(int32 i=0; i<SharedData->SelectedConstraints.Num(); ++i)
 	{
-		SnapConstraintToBone(&SharedData->SelectedConstraints[i]);
+		SharedData->SnapConstraintToBone(SharedData->SelectedConstraints[i].Index, ComponentFlags);
 	}
 	
 	RefreshPreviewViewport();
diff --git a/Engine/Source/Editor/PhysicsAssetEditor/Private/PhysicsAssetEditor.h b/Engine/Source/Editor/PhysicsAssetEditor/Private/PhysicsAssetEditor.h
--- a/Engine/Source/Editor/PhysicsAssetEditor/Private/PhysicsAssetEditor.h
+++ b/Engine/Source/Editor/PhysicsAssetEditor/Private/PhysicsAssetEditor.h
@@ -274,7 +274,7 @@ private:
 	bool CanDuplicatePrimitive() const;
 	void OnResetConstraint();
 	void OnConstrainChildBodiesToParentBody();
-	void OnSnapConstraint();
+	void OnSnapConstraint(const EConstraintTransformComponentFlags ComponentFlags);
 	void OnConvertToBallAndSocket();
 	void OnConvertToHinge();
 	void OnConvertToPrismatic();
diff --git a/Engine/Source/Editor/PhysicsAssetEditor/Private/PhysicsAssetEditorActions.cpp b/Engine/Source/Editor/PhysicsAssetEditor/Private/PhysicsAssetEditorActions.cpp
--- a/Engine/Source/Editor/PhysicsAssetEditor/Private/PhysicsAssetEditorActions.cpp
+++ b/Engine/Source/Editor/PhysicsAssetEditor/Private/PhysicsAssetEditorActions.cpp
@@ -70,7 +70,11 @@ void FPhysicsAssetEditorCommands::RegisterCommands()
 	UI_COMMAND(DuplicatePrimitive, "Duplicate", "Duplicate Selected Primitive(s)", EUserInterfaceActionType::Button, FInputChord());
 	UI_COMMAND(ConstrainChildBodiesToParentBody, "Constrain selected bodies", "Create a constraint between all selected bodies as children and the last selected body as parent", EUserInterfaceActionType::Button, FInputChord(EModifierKey::Control, EKeys::Y));
 	UI_COMMAND(ResetConstraint, "Reset", "Reset Constraint", EUserInterfaceActionType::Button, FInputChord(EModifierKey::Alt, EKeys::P));
-	UI_COMMAND(SnapConstraint, "Snap", "Snap Constraint Translation To Bone", EUserInterfaceActionType::Button, FInputChord(EModifierKey::Control, EKeys::P));
+	UI_COMMAND(SnapConstraint, "Snap All Transforms", "Snap Constraint Transforms To Bone", EUserInterfaceActionType::Button, FInputChord(EModifierKey::Control, EKeys::P));
+	UI_COMMAND(SnapConstraintChildPosition, "Snap Child Position", "Snap Constraint Child Positions To Bone", EUserInterfaceActionType::Button, FInputChord());
+	UI_COMMAND(SnapConstraintChildOrientation, "Snap Child Rotation", "Snap Constraint Child Rotations To Bone", EUserInterfaceActionType::Button, FInputChord());
+	UI_COMMAND(SnapConstraintParentPosition, "Snap Parent Position", "Snap Constraint Parent Positions To Bone", EUserInterfaceActionType::Button, FInputChord());
+	UI_COMMAND(SnapConstraintParentOrientation, "Snap Parent Rotation", "Snap Constraint Parent Rotations To Bone", EUserInterfaceActionType::Button, FInputChord());
 	UI_COMMAND(ConvertToBallAndSocket, "To Ball & Socket", "Convert Selected Constraint To Ball-And-Socket", EUserInterfaceActionType::Button, FInputChord());
 	UI_COMMAND(ConvertToHinge, "To Hinge", "Convert Selected Constraint To Hinge", EUserInterfaceActionType::Button, FInputChord());
 	UI_COMMAND(ConvertToPrismatic, "To Prismatic", "Convert Selected Constraint To Prismatic", EUserInterfaceActionType::Button, FInputChord());
diff --git a/Engine/Source/Editor/PhysicsAssetEditor/Private/PhysicsAssetEditorActions.h b/Engine/Source/Editor/PhysicsAssetEditor/Private/PhysicsAssetEditorActions.h
--- a/Engine/Source/Editor/PhysicsAssetEditor/Private/PhysicsAssetEditorActions.h
+++ b/Engine/Source/Editor/PhysicsAssetEditor/Private/PhysicsAssetEditorActions.h
@@ -83,6 +83,10 @@ public:
 	TSharedPtr<FUICommandInfo> ConstrainChildBodiesToParentBody;
 	TSharedPtr<FUICommandInfo> ResetConstraint;
 	TSharedPtr<FUICommandInfo> SnapConstraint;
+	TSharedPtr<FUICommandInfo> SnapConstraintChildPosition;
+	TSharedPtr<FUICommandInfo> SnapConstraintChildOrientation;
+	TSharedPtr<FUICommandInfo> SnapConstraintParentPosition;
+	TSharedPtr<FUICommandInfo> SnapConstraintParentOrientation;
 	TSharedPtr<FUICommandInfo> ConvertToBallAndSocket;
 	TSharedPtr<FUICommandInfo> ConvertToHinge;
 	TSharedPtr<FUICommandInfo> ConvertToPrismatic;
diff --git a/Engine/Source/Editor/PhysicsAssetEditor/Private/PhysicsAssetEditorEditMode.cpp b/Engine/Source/Editor/PhysicsAssetEditor/Private/PhysicsAssetEditorEditMode.cpp
--- a/Engine/Source/Editor/PhysicsAssetEditor/Private/PhysicsAssetEditorEditMode.cpp
+++ b/Engine/Source/Editor/PhysicsAssetEditor/Private/PhysicsAssetEditorEditMode.cpp
@@ -15,6 +15,7 @@
 #include "PhysicsAssetEditor.h"
 #include "PhysicsAssetEditorHitProxies.h"
 #include "PhysicsAssetEditorPhysicsHandleComponent.h"
+#include "PhysicsAssetRenderUtils.h"
 #include "DrawDebugHelpers.h"
 #include "SEditorViewport.h"
 #include "IPersonaToolkit.h"
@@ -447,23 +448,32 @@ bool FPhysicsAssetEditorEditMode::InputDelta(FEditorViewportClient* InViewportCl
 					SelectedObject.ManipulateTM = FTransform(Result);
 				}
 
-				UPhysicsConstraintTemplate* ConstraintSetup = SharedData->PhysicsAsset->ConstraintSetup[SelectedObject.Index];
-
-				ConstraintSetup->DefaultInstance.SetRefFrame(EConstraintFrame::Frame2, SelectedObject.ManipulateTM * StartManParentConTM[i]);
+				// Apply manipulations to Child or Parent or both transforms according to the constraint's view port manipulation flags.
+				{
+					UPhysicsConstraintTemplate* ConstraintSetup = SharedData->PhysicsAsset->ConstraintSetup[SelectedObject.Index];
+					FPhysicsAssetRenderSettings* const RenderSettings = SharedData->GetRenderSettings();
 
-				// Alt + Move will split frames and rotate or move them separately
-				bool bMultiFrame = !InViewportClient->IsAltPressed();
+					if (RenderSettings && !EnumHasAnyFlags(RenderSettings->ConstraintViewportManipulationFlags, EConstraintTransformComponentFlags::AllChild))
+					{
+						// Rotate or move the parent transform only.
+						ConstraintSetup->DefaultInstance.SetRefFrame(EConstraintFrame::Frame2, SelectedObject.ManipulateTM * StartManParentConTM[i]);
+						ConstraintSetup->DefaultInstance.SetRefFrame(EConstraintFrame::Frame1, FTransform(StartManChildConTM[i]));
+					}
+					else if (RenderSettings && !EnumHasAnyFlags(RenderSettings->ConstraintViewportManipulationFlags, EConstraintTransformComponentFlags::AllParent))
+					{
+						// Rotate or move the child transform only.
+						ConstraintSetup->DefaultInstance.SetRefFrame(EConstraintFrame::Frame1, SelectedObject.ManipulateTM * StartManChildConTM[i]);
+						ConstraintSetup->DefaultInstance.SetRefFrame(EConstraintFrame::Frame2, FTransform(StartManParentConTM[i]));
+					}
+					else
+					{
+						// Rotate or move both the parent and child transform.
+						ConstraintSetup->DefaultInstance.SetRefFrame(EConstraintFrame::Frame2, SelectedObject.ManipulateTM * StartManParentConTM[i]);
+						SharedData->SetConstraintRelTM(&SelectedObject, StartManRelConTM[i]);
+					}
 
-				if (bMultiFrame)
-				{
-					SharedData->SetConstraintRelTM(&SelectedObject, StartManRelConTM[i]);
-				}
-				else
-				{
-					ConstraintSetup->DefaultInstance.SetRefFrame(EConstraintFrame::Frame1, FTransform(StartManChildConTM[i]));
+					bHandled = true;
 				}
-
-				bHandled = true;
 			}
 		}
 	}
@@ -503,6 +513,21 @@ void FPhysicsAssetEditorEditMode::Tick(FEditorViewportClient* ViewportClient, fl
 
 		SharedData->PhysicalAnimationComponent->ApplyPhysicalAnimationProfileBelow(NAME_None, SharedData->PhysicsAsset->CurrentPhysicalAnimationProfileName, /*Include Self=*/true, /*Clear Not Found=*/true);
 	}
+
+	// Update the constraint view port manipulation flags from state of the modifier keys. These flags determine which parts of the constraint transform (Parent, child or both) should be modified when a view port widget is manipulated.
+	if (FPhysicsAssetRenderSettings* const RenderSettings = SharedData->GetRenderSettings())
+	{
+		RenderSettings->ConstraintViewportManipulationFlags = EConstraintTransformComponentFlags::All;
+
+		if (ViewportClient->IsShiftPressed() && ViewportClient->IsAltPressed()) // Shft + Alt + Move will rotate or move the child transform only.
+		{
+			EnumRemoveFlags(RenderSettings->ConstraintViewportManipulationFlags, EConstraintTransformComponentFlags::AllParent); // Remove Parent Frame flags.
+		}
+		else if (ViewportClient->IsAltPressed()) // Alt + Move will rotate or move the parent transform only.
+		{
+			EnumRemoveFlags(RenderSettings->ConstraintViewportManipulationFlags, EConstraintTransformComponentFlags::AllChild); // Remove Child Frame flags.
+		}	
+	}
 }
 
 void FPhysicsAssetEditorEditMode::Render(const FSceneView* View, FViewport* Viewport, FPrimitiveDrawInterface* PDI)
diff --git a/Engine/Source/Editor/PhysicsAssetEditor/Private/PhysicsAssetEditorSharedData.cpp b/Engine/Source/Editor/PhysicsAssetEditor/Private/PhysicsAssetEditorSharedData.cpp
--- a/Engine/Source/Editor/PhysicsAssetEditor/Private/PhysicsAssetEditorSharedData.cpp
+++ b/Engine/Source/Editor/PhysicsAssetEditor/Private/PhysicsAssetEditorSharedData.cpp
@@ -481,7 +481,7 @@ void FPhysicsAssetEditorSharedData::Mirror()
 					}
 				}
 
-				// Display an error notification if neccesary.
+				// Display an error notification if necessary.
 				if (!(MirrorCollisionsMissingBones.IsEmpty() && MirrorCollisionsMissingBodies.IsEmpty()))
 				{
 					// Construct error message for failed collision mirroring.
@@ -2293,35 +2293,16 @@ void FPhysicsAssetEditorSharedData::SetConstraintRelTM(const FPhysicsAssetEditor
 	}
 }
 
-void FPhysicsAssetEditorSharedData::SnapConstraintToBone(int32 ConstraintIndex)
+void FPhysicsAssetEditorSharedData::SnapConstraintToBone(const int32 ConstraintIndex, const EConstraintTransformComponentFlags ComponentFlags /* = EConstraintTransformComponentFlags::All */)
 {
 	UPhysicsConstraintTemplate* ConstraintSetup = PhysicsAsset->ConstraintSetup[ConstraintIndex];
 	ConstraintSetup->Modify();
-	SnapConstraintToBone(ConstraintSetup->DefaultInstance);
+	SnapConstraintToBone(ConstraintSetup->DefaultInstance, ComponentFlags);
 }
 
-void FPhysicsAssetEditorSharedData::SnapConstraintToBone(FConstraintInstance& ConstraintInstance)
+void FPhysicsAssetEditorSharedData::SnapConstraintToBone(FConstraintInstance& ConstraintInstance, const EConstraintTransformComponentFlags ComponentFlags /* = EConstraintTransformComponentFlags::All */)
 {
-	USkeletalMesh* EditorSkelMesh = PhysicsAsset->GetPreviewMesh();
-	if(EditorSkelMesh == nullptr)
-	{
-		return;
-	}
-
-	const int32 BoneIndex1 = EditorSkelMesh->GetRefSkeleton().FindBoneIndex(ConstraintInstance.ConstraintBone1);
-	const int32 BoneIndex2 = EditorSkelMesh->GetRefSkeleton().FindBoneIndex(ConstraintInstance.ConstraintBone2);
-
-	check(BoneIndex1 != INDEX_NONE);
-	check(BoneIndex2 != INDEX_NONE);
-
-	const FTransform BoneTransform1 = EditorSkelComp->GetBoneTransform(BoneIndex1);
-	const FTransform BoneTransform2 = EditorSkelComp->GetBoneTransform(BoneIndex2);
-
-	// Bone transforms are world space, and frame transforms are local space (local to bones).
-	// Frame 1 is the child frame, and set to identity.
-	// Frame 2 is the parent frame, and needs to be set relative to Frame1.
-	ConstraintInstance.SetRefFrame(EConstraintFrame::Frame2, BoneTransform1.GetRelativeTransform(BoneTransform2));
-	ConstraintInstance.SetRefFrame(EConstraintFrame::Frame1, FTransform::Identity);
+	ConstraintInstance.SnapTransformsToDefault(ComponentFlags, PhysicsAsset);
 }
 
 void FPhysicsAssetEditorSharedData::CopyConstraintProperties()
diff --git a/Engine/Source/Editor/PhysicsAssetEditor/Private/PhysicsAssetEditorSharedData.h b/Engine/Source/Editor/PhysicsAssetEditor/Private/PhysicsAssetEditorSharedData.h
--- a/Engine/Source/Editor/PhysicsAssetEditor/Private/PhysicsAssetEditorSharedData.h
+++ b/Engine/Source/Editor/PhysicsAssetEditor/Private/PhysicsAssetEditorSharedData.h
@@ -112,12 +112,12 @@ public:
     {
         SetConstraintRelTM(GetSelectedConstraint(), RelTM);
     }
-	
+
 	/** Snaps a constraint at the specified index to it's bone */
-	void SnapConstraintToBone(int32 ConstraintIndex);
+	void SnapConstraintToBone(const int32 ConstraintIndex, const EConstraintTransformComponentFlags ComponentFlags = EConstraintTransformComponentFlags::All);
 
 	/** Snaps the specified constraint to it's bone */
-	void SnapConstraintToBone(FConstraintInstance& ConstraintInstance);
+	void SnapConstraintToBone(FConstraintInstance& ConstraintInstance, const EConstraintTransformComponentFlags ComponentFlags = EConstraintTransformComponentFlags::All);
 
 	/** Deletes the currently selected constraints */
 	void DeleteCurrentConstraint();
diff --git a/Engine/Source/Runtime/Engine/Classes/PhysicsEngine/ConstraintInstance.h b/Engine/Source/Runtime/Engine/Classes/PhysicsEngine/ConstraintInstance.h
--- a/Engine/Source/Runtime/Engine/Classes/PhysicsEngine/ConstraintInstance.h
+++ b/Engine/Source/Runtime/Engine/Classes/PhysicsEngine/ConstraintInstance.h
@@ -2,6 +2,7 @@
 #pragma once
 
 #include "CoreMinimal.h"
+#include "Misc/EnumClassFlags.h"
 #include "UObject/ObjectMacros.h"
 #include "UObject/Class.h"
 #include "EngineDefines.h"
@@ -19,6 +20,26 @@ struct FBodyInstance;
 class FMaterialRenderProxy;
 class FPrimitiveDrawInterface;
 class FMaterialRenderProxy;
+class UPhysicsAsset;
+
+UENUM()
+enum class EConstraintTransformComponentFlags : uint8
+{
+	None = 0,
+
+	ChildPosition = 1 << 0,
+	ChildRotation = 1 << 1,
+	ParentPosition = 1 << 2,
+	ParentRotation = 1 << 3,
+
+	AllChild = ChildPosition | ChildRotation,
+	AllParent = ParentPosition | ParentRotation,
+	AllPosition = ChildPosition | ParentPosition,
+	AllRotation = ChildRotation | ParentRotation,
+	All = AllChild | AllParent
+};
+
+ENUM_CLASS_FLAGS(EConstraintTransformComponentFlags);
 
 /** Container for properties of a physics constraint that can be easily swapped at runtime. This is useful for switching different setups when going from ragdoll to standup for example */
 USTRUCT()
@@ -235,29 +256,29 @@ public:
 	///////////////////////////// Body1 ref frame
 	
 	/** Location of constraint in Body1 reference frame (usually the "child" body for skeletal meshes). */
-	UPROPERTY(EditAnywhere, Category=Constraint)
+	UPROPERTY(EditAnywhere, Category = ConstraintTransforms)
 	FVector Pos1;
 
 	/** Primary (twist) axis in Body1 reference frame. */
-	UPROPERTY()
+	UPROPERTY(EditAnywhere, Category = ConstraintTransforms)
 	FVector PriAxis1;
 
-	/** Seconday axis in Body1 reference frame. Orthogonal to PriAxis1. */
-	UPROPERTY()
+	/** Secondary axis in Body1 reference frame. Orthogonal to PriAxis1. */
+	UPROPERTY(EditAnywhere, Category = ConstraintTransforms)
 	FVector SecAxis1;
 
 	///////////////////////////// Body2 ref frame
 	
 	/** Location of constraint in Body2 reference frame (usually the "parent" body for skeletal meshes). */
-	UPROPERTY(EditAnywhere, Category=Constraint)
+	UPROPERTY(EditAnywhere, Category= ConstraintTransforms)
 	FVector Pos2;
 
 	/** Primary (twist) axis in Body2 reference frame. */
-	UPROPERTY()
+	UPROPERTY(EditAnywhere, Category = ConstraintTransforms)
 	FVector PriAxis2;
 
-	/** Seconday axis in Body2 reference frame. Orthogonal to PriAxis2. */
-	UPROPERTY()
+	/** Secondary axis in Body2 reference frame. Orthogonal to PriAxis2. */
+	UPROPERTY(EditAnywhere, Category = ConstraintTransforms)
 	FVector SecAxis2;
 
 	/** Specifies the angular offset between the two frames of reference. By default limit goes from (-Angle, +Angle)
@@ -991,6 +1012,18 @@ private:
 
 public:
 
+#if WITH_EDITORONLY_DATA
+
+	/** Returns this constraint's default transform relative to its parent bone. */
+	FTransform CalculateDefaultParentTransform(const UPhysicsAsset* const PhysicsAsset) const;
+
+	/** Returns this constraint's default transform relative to its child bone. */
+	FTransform CalculateDefaultChildTransform() const;
+
+	/** Set the constraint transform components specified by the SnapFlags to their default values (derived from the parent and child bone transforms). */
+	void SnapTransformsToDefault(const EConstraintTransformComponentFlags SnapFlags, const UPhysicsAsset* const PhysicsAsset);
+#endif // WITH_EDITORONLY_DATA
+
 	///////////////////////////// DEPRECATED
 	// Most of these properties have moved inside the ProfileInstance member (FConstraintProfileProperties struct)
 #if WITH_EDITORONLY_DATA
@@ -1172,4 +1205,4 @@ private:
 	/** Warning: it would be unwieldy to make the accessor move only, so the OnRelease callback should be safe to be called multiple times */
 	TFunction<void(void)> OnRelease;
 #endif
-};
\ No newline at end of file
+};
diff --git a/Engine/Source/Runtime/Engine/Private/PhysicsEngine/ConstraintInstance.cpp b/Engine/Source/Runtime/Engine/Private/PhysicsEngine/ConstraintInstance.cpp
--- a/Engine/Source/Runtime/Engine/Private/PhysicsEngine/ConstraintInstance.cpp
+++ b/Engine/Source/Runtime/Engine/Private/PhysicsEngine/ConstraintInstance.cpp
@@ -69,6 +69,30 @@ FVector RevolutionsToRads(const FVector Revolutions)
 	return Revolutions * 2.f * UE_PI;
 }
 
+#if WITH_EDITORONLY_DATA
+
+/** Returns the child bone's transform relative to the parent bone. */
+FTransform CalculateChildTransformRelativeToParent(const FName ChildBoneName, const FName ParentBoneName, const UPhysicsAsset* const PhysicsAsset)
+{
+	FReferenceSkeleton& ReferenceSkeleton = PhysicsAsset->PreviewSkeletalMesh->GetRefSkeleton();
+	const TArray<FTransform>& LocalPose = PhysicsAsset->PreviewSkeletalMesh->GetRefSkeleton().GetRefBonePose();
+	const int32 ChildBoneIndex = ReferenceSkeleton.FindBoneIndex(ChildBoneName);
+	const int32 ParentBoneIndex = ReferenceSkeleton.FindBoneIndex(ParentBoneName);
+
+	// Transform of child from parent is just child ref-pose entry.
+	FTransform RelTM = FTransform::Identity;
+
+	// Traverse the skeleton hierarchy from child to parent bone, accumulating transforms.
+	for (int32 ParentIndex = ChildBoneIndex; LocalPose.IsValidIndex(ParentIndex) && (ParentIndex != ParentBoneIndex); ParentIndex = ReferenceSkeleton.GetParentIndex(ParentIndex))
+	{
+		RelTM = RelTM * LocalPose[ParentIndex];
+	}
+
+	return RelTM;
+}
+
+#endif // WITH_EDITORONLY_DATA
+
 #if WITH_EDITOR
 void FConstraintProfileProperties::SyncChangedConstraintProperties(FPropertyChangedChainEvent& PropertyChangedEvent)
 {
@@ -223,6 +247,46 @@ void FConstraintInstance::UpdateDriveTarget()
 	});
 }
 
+#if WITH_EDITORONLY_DATA
+
+FTransform FConstraintInstance::CalculateDefaultParentTransform(const UPhysicsAsset* const PhysicsAsset) const
+{
+	return CalculateChildTransformRelativeToParent(ConstraintBone1, ConstraintBone2, PhysicsAsset);
+}
+
+FTransform FConstraintInstance::CalculateDefaultChildTransform() const
+{
+	return FTransform::Identity;
+} 
+
+void FConstraintInstance::SnapTransformsToDefault(const EConstraintTransformComponentFlags SnapFlags, const UPhysicsAsset* const PhysicsAsset)
+{
+	const FTransform ParentTransform = CalculateDefaultParentTransform(PhysicsAsset);
+	const FTransform ChildTransform = CalculateDefaultChildTransform();
+
+	if (EnumHasAnyFlags(SnapFlags, EConstraintTransformComponentFlags::ChildPosition))
+	{
+		SetRefPosition(EConstraintFrame::Frame1, ChildTransform.GetLocation());
+	}
+
+	if (EnumHasAnyFlags(SnapFlags, EConstraintTransformComponentFlags::ChildRotation))
+	{
+		SetRefOrientation(EConstraintFrame::Frame1, ChildTransform.GetUnitAxis(EAxis::X), ChildTransform.GetUnitAxis(EAxis::Y));
+	}
+
+	if (EnumHasAnyFlags(SnapFlags, EConstraintTransformComponentFlags::ParentPosition))
+	{
+		SetRefPosition(EConstraintFrame::Frame2, ParentTransform.GetLocation());
+	}
+
+	if (EnumHasAnyFlags(SnapFlags, EConstraintTransformComponentFlags::ParentRotation))
+	{
+		SetRefOrientation(EConstraintFrame::Frame2, ParentTransform.GetUnitAxis(EAxis::X), ParentTransform.GetUnitAxis(EAxis::Y));
+	}
+}
+
+#endif // WITH_EDITORONLY_DATA
+
 /** Constructor **/
 FConstraintInstanceBase::FConstraintInstanceBase()
 {
