[UE][FIX] Fix crash when creating a Constraint Actor. Also Fix conflicting LOCTEXT macro.

#rnx

diff --git a/Engine/Source/Editor/DetailCustomizations/Private/PhysicsConstraintComponentDetails.cpp b/Engine/Source/Editor/DetailCustomizations/Private/PhysicsConstraintComponentDetails.cpp
--- a/Engine/Source/Editor/DetailCustomizations/Private/PhysicsConstraintComponentDetails.cpp
+++ b/Engine/Source/Editor/DetailCustomizations/Private/PhysicsConstraintComponentDetails.cpp
@@ -187,7 +187,7 @@ void FConstraintTransformCustomization::MakePositionRow(TSharedRef<class IProper
 		+ SHorizontalBox::Slot()
 		.Padding(FMargin(0.0f, 2.0f, 0.0f, 2.0f))
 		[
-			MakeNumericProxyWidget<FVector, FVector::FReal>(InPositionPropertyHandle, CachedPositionZ, FText::Format(ComponentToolTipFormatText, LOCTEXT("PositionY", "Z"), FrameLabelText), SNumericEntryBox<FRotator::FReal>::BlueLabelBackgroundColor)
+			MakeNumericProxyWidget<FVector, FVector::FReal>(InPositionPropertyHandle, CachedPositionZ, FText::Format(ComponentToolTipFormatText, LOCTEXT("PositionZ", "Z"), FrameLabelText), SNumericEntryBox<FRotator::FReal>::BlueLabelBackgroundColor)
 		]
 	];
 }
@@ -1057,9 +1057,12 @@ void FPhysicsConstraintComponentDetails::AddConstraintFrameTransform(const ECons
 
 void FPhysicsConstraintComponentDetails::SnapConstraintTransformComponentsToDefault(TSharedPtr<IPropertyHandle> ConstraintInstancePropertyHandle, const EConstraintTransformComponentFlags SnapFlags)
 {
-	if (FConstraintInstance* const ConstraintInstance = GetConstraintInstance())
+	if (ParentPhysicsAsset)
 	{
-		ConstraintInstance->SnapTransformsToDefault(SnapFlags, ParentPhysicsAsset);
+		if (FConstraintInstance* const ConstraintInstance = GetConstraintInstance())
+		{
+			ConstraintInstance->SnapTransformsToDefault(SnapFlags, ParentPhysicsAsset);
+		}
 	}
 }
 
@@ -1095,19 +1098,27 @@ bool FPhysicsConstraintComponentDetails::IsSnapConstraintTransformComponentVisib
 
 void FPhysicsConstraintComponentDetails::ToggleDisplayConstraintTransformComponentRelativeToDefault(TSharedPtr<IPropertyHandle> ConstraintInstancePropertyHandle, const EConstraintTransformComponentFlags ComponentFlags)
 {
-	IPhysicsAssetRenderInterface& PhysicsAssetRenderInterface = IModularFeatures::Get().GetModularFeature<IPhysicsAssetRenderInterface>("PhysicsAssetRenderInterface");
-	const bool bCurrentState = IsDisplayingConstraintTransformComponentRelativeToDefault(ComponentFlags);
-	const EConstraintTransformComponentFlags ComponentFlagsToModify = (FSlateApplication::Get().GetModifierKeys().IsShiftDown()) ? EConstraintTransformComponentFlags::All : ComponentFlags;
-	PhysicsAssetRenderInterface.SetDisplayConstraintTransformComponentRelativeToDefault(ParentPhysicsAsset, ComponentFlagsToModify, !bCurrentState);
-	UpdateTransformProxyDisplayRelativeToDefault(ConstraintInstancePropertyHandle);
+	if (ParentPhysicsAsset)
+	{
+		IPhysicsAssetRenderInterface& PhysicsAssetRenderInterface = IModularFeatures::Get().GetModularFeature<IPhysicsAssetRenderInterface>("PhysicsAssetRenderInterface");
+		const bool bCurrentState = IsDisplayingConstraintTransformComponentRelativeToDefault(ComponentFlags);
+		const EConstraintTransformComponentFlags ComponentFlagsToModify = (FSlateApplication::Get().GetModifierKeys().IsShiftDown()) ? EConstraintTransformComponentFlags::All : ComponentFlags;
+		PhysicsAssetRenderInterface.SetDisplayConstraintTransformComponentRelativeToDefault(ParentPhysicsAsset, ComponentFlagsToModify, !bCurrentState);
+		UpdateTransformProxyDisplayRelativeToDefault(ConstraintInstancePropertyHandle);
+	}
 }
 
 bool FPhysicsConstraintComponentDetails::IsConstraintTransformComponentEnabled(const EConstraintTransformComponentFlags ComponentFlags) const
 {
-	IPhysicsAssetRenderInterface& PhysicsAssetRenderInterface = IModularFeatures::Get().GetModularFeature<IPhysicsAssetRenderInterface>("PhysicsAssetRenderInterface");
-	const EConstraintTransformComponentFlags ComponentManipulationFlags = PhysicsAssetRenderInterface.GetConstraintViewportManipulationFlags(ParentPhysicsAsset);
+	if (ParentPhysicsAsset)
+	{
+		IPhysicsAssetRenderInterface& PhysicsAssetRenderInterface = IModularFeatures::Get().GetModularFeature<IPhysicsAssetRenderInterface>("PhysicsAssetRenderInterface");
+		const EConstraintTransformComponentFlags ComponentManipulationFlags = PhysicsAssetRenderInterface.GetConstraintViewportManipulationFlags(ParentPhysicsAsset);
+
+		return EnumHasAllFlags(ComponentManipulationFlags, ComponentFlags);
+	}
 
-	return EnumHasAllFlags(ComponentManipulationFlags, ComponentFlags);
+	return true;
 }
 
 FSlateColor FPhysicsConstraintComponentDetails::GetConstraintTransformColorAndOpacity(const EConstraintTransformComponentFlags ComponentFlags) const
@@ -1117,23 +1128,31 @@ FSlateColor FPhysicsConstraintComponentDetails::GetConstraintTransformColorAndOp
 
 void FPhysicsConstraintComponentDetails::UpdateTransformProxyDisplayRelativeToDefault(TSharedPtr<IPropertyHandle> ConstraintInstancePropertyHandle)
 {
-	IPhysicsAssetRenderInterface& PhysicsAssetRenderInterface = IModularFeatures::Get().GetModularFeature<IPhysicsAssetRenderInterface>("PhysicsAssetRenderInterface");
+	if (ParentPhysicsAsset)
+	{
+		IPhysicsAssetRenderInterface& PhysicsAssetRenderInterface = IModularFeatures::Get().GetModularFeature<IPhysicsAssetRenderInterface>("PhysicsAssetRenderInterface");
 
-	ChildTransformProxy->SetPositionDisplayRelativeToDefault(PhysicsAssetRenderInterface.IsDisplayingConstraintTransformComponentRelativeToDefault(ParentPhysicsAsset, EConstraintTransformComponentFlags::ChildPosition));
-	ChildTransformProxy->SetRotationDisplayRelativeToDefault(PhysicsAssetRenderInterface.IsDisplayingConstraintTransformComponentRelativeToDefault(ParentPhysicsAsset, EConstraintTransformComponentFlags::ChildRotation));
-	ParentTransformProxy->SetPositionDisplayRelativeToDefault(PhysicsAssetRenderInterface.IsDisplayingConstraintTransformComponentRelativeToDefault(ParentPhysicsAsset, EConstraintTransformComponentFlags::ParentPosition));
-	ParentTransformProxy->SetRotationDisplayRelativeToDefault(PhysicsAssetRenderInterface.IsDisplayingConstraintTransformComponentRelativeToDefault(ParentPhysicsAsset, EConstraintTransformComponentFlags::ParentRotation));
+		ChildTransformProxy->SetPositionDisplayRelativeToDefault(PhysicsAssetRenderInterface.IsDisplayingConstraintTransformComponentRelativeToDefault(ParentPhysicsAsset, EConstraintTransformComponentFlags::ChildPosition));
+		ChildTransformProxy->SetRotationDisplayRelativeToDefault(PhysicsAssetRenderInterface.IsDisplayingConstraintTransformComponentRelativeToDefault(ParentPhysicsAsset, EConstraintTransformComponentFlags::ChildRotation));
+		ParentTransformProxy->SetPositionDisplayRelativeToDefault(PhysicsAssetRenderInterface.IsDisplayingConstraintTransformComponentRelativeToDefault(ParentPhysicsAsset, EConstraintTransformComponentFlags::ParentPosition));
+		ParentTransformProxy->SetRotationDisplayRelativeToDefault(PhysicsAssetRenderInterface.IsDisplayingConstraintTransformComponentRelativeToDefault(ParentPhysicsAsset, EConstraintTransformComponentFlags::ParentRotation));
 
-	if (FConstraintInstance* const ConstraintInstance = GetConstraintInstance())
-	{
-		ChildTransformProxy->SetDefaultTransform(ConstraintInstance->CalculateDefaultChildTransform(), ChildPositionPropertyHandle, ChildPriAxisPropertyHandle, ChildSecAxisPropertyHandle);
-		ParentTransformProxy->SetDefaultTransform(ConstraintInstance->CalculateDefaultParentTransform(ParentPhysicsAsset), ParentPositionPropertyHandle, ParentPriAxisPropertyHandle, ParentSecAxisPropertyHandle);
+		if (FConstraintInstance* const ConstraintInstance = GetConstraintInstance())
+		{
+			ChildTransformProxy->SetDefaultTransform(ConstraintInstance->CalculateDefaultChildTransform(), ChildPositionPropertyHandle, ChildPriAxisPropertyHandle, ChildSecAxisPropertyHandle);
+			ParentTransformProxy->SetDefaultTransform(ConstraintInstance->CalculateDefaultParentTransform(ParentPhysicsAsset), ParentPositionPropertyHandle, ParentPriAxisPropertyHandle, ParentSecAxisPropertyHandle);
+		}
 	}
 }
 
 bool FPhysicsConstraintComponentDetails::IsDisplayingConstraintTransformComponentRelativeToDefault(const EConstraintTransformComponentFlags ComponentFlags)
 {
-	return IModularFeatures::Get().GetModularFeature<IPhysicsAssetRenderInterface>("PhysicsAssetRenderInterface").IsDisplayingConstraintTransformComponentRelativeToDefault(ParentPhysicsAsset, ComponentFlags);
+	if (ParentPhysicsAsset)
+	{
+		return IModularFeatures::Get().GetModularFeature<IPhysicsAssetRenderInterface>("PhysicsAssetRenderInterface").IsDisplayingConstraintTransformComponentRelativeToDefault(ParentPhysicsAsset, ComponentFlags);
+	}
+
+	return false;
 }
 
 void FPhysicsConstraintComponentDetails::OnCopyConstraintTransform(TSharedPtr<IPropertyHandle> PositionPropertyHandle, TSharedPtr<IPropertyHandle> PriAxisPropertyHandle, TSharedPtr<IPropertyHandle> SecAxisPropertyHandle, TSharedPtr<FConstraintTransformCustomization> TransformProxy)
@@ -1859,6 +1878,7 @@ void FPhysicsConstraintComponentDetails::CustomizeDetails( IDetailLayoutBuilder&
 
 	TSharedPtr<IPropertyHandle> ConstraintInstanceProperty;
 	APhysicsConstraintActor* OwningConstraintActor = NULL;
+	ParentPhysicsAsset = NULL;
 
 	bInPhat = false;
 
diff --git a/Engine/Source/Editor/DetailCustomizations/Private/PhysicsConstraintComponentDetails.h b/Engine/Source/Editor/DetailCustomizations/Private/PhysicsConstraintComponentDetails.h
--- a/Engine/Source/Editor/DetailCustomizations/Private/PhysicsConstraintComponentDetails.h
+++ b/Engine/Source/Editor/DetailCustomizations/Private/PhysicsConstraintComponentDetails.h
@@ -190,7 +190,6 @@ private:
 	TSharedPtr<IPropertyHandle> ParentPriAxisPropertyHandle;
 	TSharedPtr<IPropertyHandle> ParentSecAxisPropertyHandle;
 
-//	TWeakObjectPtr<UObject> ConstraintComponent;
 	UPhysicsConstraintComponent* ConstraintComp;
 	UPhysicsConstraintTemplate* ConstraintTemplate;
 
