[UE][FIX] Ensure new physics constraint editing features intended for PhAT are not applied to details panel for ActorConstraints

#rnx

diff --git a/Engine/Source/Editor/DetailCustomizations/Private/PhysicsConstraintComponentDetails.cpp b/Engine/Source/Editor/DetailCustomizations/Private/PhysicsConstraintComponentDetails.cpp
--- a/Engine/Source/Editor/DetailCustomizations/Private/PhysicsConstraintComponentDetails.cpp
+++ b/Engine/Source/Editor/DetailCustomizations/Private/PhysicsConstraintComponentDetails.cpp
@@ -828,9 +828,9 @@ void FPhysicsConstraintComponentDetails::AddConstraintProperties(IDetailLayoutBu
 
 	IDetailCategoryBuilder& ConstraintCategory = DetailBuilder.EditCategory("Constraint");
 
-	// Add current profile name to the constraint category header.
 	if (bInPhat)
 	{
+		// Add current profile name to the constraint category header.
 		ConstraintCategory.HeaderContent(
 			SNew(SHorizontalBox)
 			+ SHorizontalBox::Slot()
@@ -868,17 +868,18 @@ void FPhysicsConstraintComponentDetails::AddConstraintProperties(IDetailLayoutBu
 					return FText();
 				})
 			]);
-	}
 
-	ConstraintCategory.AddProperty(ConstraintInstance->GetChildHandle(GET_MEMBER_NAME_CHECKED(FConstraintInstance, ConstraintBone1))).DisplayName(LOCTEXT("ConstraintChildBoneName", "Child Bone Name"));
-	ConstraintCategory.AddProperty(ConstraintInstance->GetChildHandle(GET_MEMBER_NAME_CHECKED(FConstraintInstance, ConstraintBone2))).DisplayName(LOCTEXT("ConstraintParentBoneName", "Parent Bone Name"));
+		// Add PhAT specific components to details panel.
+		ConstraintCategory.AddProperty(ConstraintInstance->GetChildHandle(GET_MEMBER_NAME_CHECKED(FConstraintInstance, ConstraintBone1))).DisplayName(LOCTEXT("ConstraintChildBoneName", "Child Bone Name"));
+		ConstraintCategory.AddProperty(ConstraintInstance->GetChildHandle(GET_MEMBER_NAME_CHECKED(FConstraintInstance, ConstraintBone2))).DisplayName(LOCTEXT("ConstraintParentBoneName", "Parent Bone Name"));
 
-	IDetailCategoryBuilder& ConstraintTransformsCat = DetailBuilder.EditCategory("ConstraintTransforms");
+		IDetailCategoryBuilder& ConstraintTransformsCat = DetailBuilder.EditCategory("ConstraintTransforms");
 
-	AddConstraintFrameTransform(EConstraintFrame::Frame1, ConstraintTransformsCat, ConstraintInstance);
-	AddConstraintFrameTransform(EConstraintFrame::Frame2, ConstraintTransformsCat, ConstraintInstance);
+		AddConstraintFrameTransform(EConstraintFrame::Frame1, ConstraintTransformsCat, ConstraintInstance);
+		AddConstraintFrameTransform(EConstraintFrame::Frame2, ConstraintTransformsCat, ConstraintInstance);
 
-	UpdateTransformProxyDisplayRelativeToDefault(ConstraintInstance);
+		UpdateTransformProxyDisplayRelativeToDefault(ConstraintInstance);
+	}
 
 	// Hide the constraint transform values that are being represented by the proxy values in the details panel.
 	ChildPositionPropertyHandle->MarkHiddenByCustomization();
