Crash fix when changing an anim composite segment's start time, end time, or play rate.

diff --git a/Engine/Source/Editor/DetailCustomizations/Private/AnimMontageSegmentDetails.cpp b/Engine/Source/Editor/DetailCustomizations/Private/AnimMontageSegmentDetails.cpp
--- a/Engine/Source/Editor/DetailCustomizations/Private/AnimMontageSegmentDetails.cpp
+++ b/Engine/Source/Editor/DetailCustomizations/Private/AnimMontageSegmentDetails.cpp
@@ -131,11 +131,11 @@ void FAnimMontageSegmentDetails::CustomizeDetails( IDetailLayoutBuilder& DetailB
 			.AllowSpin(true)
 			.MinSliderValue(0.f)
 			.MinValue(0.f)
-			.MaxSliderValue_Raw(this, &FAnimMontageSegmentDetails::GetAnimationAssetPlayLength)
-			.MaxValue_Raw(this, &FAnimMontageSegmentDetails::GetAnimationAssetPlayLength)
-			.Value_Raw(this, &FAnimMontageSegmentDetails::GetStartTime)
-			.OnValueChanged_Raw(this, &FAnimMontageSegmentDetails::OnStartTimeChanged, ETextCommit::Default, true)
-			.OnValueCommitted_Raw(this, &FAnimMontageSegmentDetails::OnStartTimeChanged, false)
+			.MaxSliderValue(this, &FAnimMontageSegmentDetails::GetAnimationAssetPlayLength)
+			.MaxValue(this, &FAnimMontageSegmentDetails::GetAnimationAssetPlayLength)
+			.Value(this, &FAnimMontageSegmentDetails::GetStartTime)
+			.OnValueChanged(this, &FAnimMontageSegmentDetails::OnStartTimeChanged, ETextCommit::Default, true)
+			.OnValueCommitted(this, &FAnimMontageSegmentDetails::OnStartTimeChanged, false)
 		];
 
 		FResetToDefaultOverride Handler = FResetToDefaultOverride::Create
@@ -170,11 +170,11 @@ void FAnimMontageSegmentDetails::CustomizeDetails( IDetailLayoutBuilder& DetailB
 			.AllowSpin(true)
 			.MinSliderValue(0.f)
 			.MinValue(0.f)
-			.MaxSliderValue_Raw(this, &FAnimMontageSegmentDetails::GetAnimationAssetPlayLength)
-			.MaxValue_Raw(this, &FAnimMontageSegmentDetails::GetAnimationAssetPlayLength)
-			.Value_Raw(this, &FAnimMontageSegmentDetails::GetEndTime)
-			.OnValueChanged_Raw(this, &FAnimMontageSegmentDetails::OnEndTimeChanged, ETextCommit::Default, true)
-			.OnValueCommitted_Raw(this, &FAnimMontageSegmentDetails::OnEndTimeChanged, false)
+			.MaxSliderValue(this, &FAnimMontageSegmentDetails::GetAnimationAssetPlayLength)
+			.MaxValue(this, &FAnimMontageSegmentDetails::GetAnimationAssetPlayLength)
+			.Value(this, &FAnimMontageSegmentDetails::GetEndTime)
+			.OnValueChanged(this, &FAnimMontageSegmentDetails::OnEndTimeChanged, ETextCommit::Default, true)
+			.OnValueCommitted(this, &FAnimMontageSegmentDetails::OnEndTimeChanged, false)
 		];
 
 		const TSharedPtr<IPropertyHandle> AnimPlayRateProperty = AnimSegmentHandle->GetChildHandle(GET_MEMBER_NAME_CHECKED(FAnimSegment, AnimPlayRate));
@@ -195,7 +195,7 @@ void FAnimMontageSegmentDetails::CustomizeDetails( IDetailLayoutBuilder& DetailB
 			.MinValue(-32.f)
 			.MaxSliderValue(1.f)
 			.MaxValue(32.f)
-			.Value_Raw(this, &FAnimMontageSegmentDetails::GetPlayRate)
+			.Value(this, &FAnimMontageSegmentDetails::GetPlayRate)
 			.OnValueChanged_Lambda([AnimPlayRateProperty](float InValue)
 			{
 				if(AnimPlayRateProperty.IsValid() && !FMath::IsNearlyZero(InValue))
