FDirectoryPath properties now correctly update the Enabled state of the "..." button.

Removed unused GetDisplayedText()

diff --git a/Engine/Source/Editor/DetailCustomizations/Private/DirectoryPathStructCustomization.cpp b/Engine/Source/Editor/DetailCustomizations/Private/DirectoryPathStructCustomization.cpp
--- a/Engine/Source/Editor/DetailCustomizations/Private/DirectoryPathStructCustomization.cpp
+++ b/Engine/Source/Editor/DetailCustomizations/Private/DirectoryPathStructCustomization.cpp
@@ -43,11 +43,11 @@ void FDirectoryPathStructCustomization::CustomizeHeader( TSharedRef<IPropertyHan
 			PickerWidget = SAssignNew(PickerButton, SButton)
 			.ButtonStyle( FAppStyle::Get(), "HoverHintOnly" )
 			.ToolTipText( LOCTEXT( "FolderComboToolTipText", "Choose a content directory") )
-			.OnClicked(FOnClicked::CreateSP(this, &FDirectoryPathStructCustomization::OnPickContent, PathProperty.ToSharedRef()))
+			.OnClicked( this, &FDirectoryPathStructCustomization::OnPickContent, PathProperty.ToSharedRef() )
 			.ContentPadding(2.0f)
 			.ForegroundColor( FSlateColor::UseForeground() )
 			.IsFocusable(false)
-			.IsEnabled(StructPropertyHandle->IsEditable())
+			.IsEnabled(this, &FDirectoryPathStructCustomization::IsBrowseEnabled, StructPropertyHandle)
 			[
 				SNew(SImage)
 				.Image(FAppStyle::GetBrush("PropertyWindow.Button_Ellipsis"))
@@ -60,11 +60,11 @@ void FDirectoryPathStructCustomization::CustomizeHeader( TSharedRef<IPropertyHan
 			PickerWidget = SAssignNew(BrowseButton, SButton)
 			.ButtonStyle( FAppStyle::Get(), "HoverHintOnly" )
 			.ToolTipText( LOCTEXT( "FolderButtonToolTipText", "Choose a directory from this computer") )
-			.OnClicked( FOnClicked::CreateSP(this, &FDirectoryPathStructCustomization::OnPickDirectory, PathProperty.ToSharedRef(), bRelativeToGameContentDir, bUseRelativePath) )
+			.OnClicked( this, &FDirectoryPathStructCustomization::OnPickDirectory, PathProperty.ToSharedRef(), bRelativeToGameContentDir, bUseRelativePath )
 			.ContentPadding( 2.0f )
 			.ForegroundColor( FSlateColor::UseForeground() )
 			.IsFocusable( false )
-			.IsEnabled(StructPropertyHandle->IsEditable())
+			.IsEnabled( this, &FDirectoryPathStructCustomization::IsBrowseEnabled, StructPropertyHandle )
 			[
 				SNew( SImage )
 				.Image( FAppStyle::GetBrush("PropertyWindow.Button_Ellipsis") )
@@ -207,4 +207,9 @@ void FDirectoryPathStructCustomization::OnPathPicked(const FString& Path, TShare
 	PropertyHandle->SetValue(Path);
 }
 
+bool FDirectoryPathStructCustomization::IsBrowseEnabled(TSharedRef<IPropertyHandle> PropertyHandle) const
+{
+	return PropertyHandle->IsEditable();
+}
+
 #undef LOCTEXT_NAMESPACE
diff --git a/Engine/Source/Editor/DetailCustomizations/Private/DirectoryPathStructCustomization.h b/Engine/Source/Editor/DetailCustomizations/Private/DirectoryPathStructCustomization.h
--- a/Engine/Source/Editor/DetailCustomizations/Private/DirectoryPathStructCustomization.h
+++ b/Engine/Source/Editor/DetailCustomizations/Private/DirectoryPathStructCustomization.h
@@ -23,9 +23,6 @@ public:
 
 private:
 
-	/** Delegate for displaying text value of path */
-	FText GetDisplayedText(TSharedRef<IPropertyHandle> PropertyHandle) const;
-
 	/** Delegate used to display a directory picker */
 	FReply OnPickContent(TSharedRef<IPropertyHandle> PropertyHandle) ;
 
@@ -38,6 +35,11 @@ private:
 	/** Called when a path is picked from the path picker */
 	void OnPathPicked(const FString& Path, TSharedRef<IPropertyHandle> PropertyHandle);
 
+	/** Delegate to determine whether the browse button should be enabled */
+	bool IsBrowseEnabled(TSharedRef<IPropertyHandle> PropertyHandle) const;
+
+private:
+
 	/** The browse button widget */
 	TSharedPtr<SButton> BrowseButton;
 
