Fixed a bug where if a spline details custimization was restricting spline point interpolation types it would fail to select the correct value.

diff --git a/Engine/Source/Editor/DetailCustomizations/Private/SplineComponentDetails.cpp b/Engine/Source/Editor/DetailCustomizations/Private/SplineComponentDetails.cpp
--- a/Engine/Source/Editor/DetailCustomizations/Private/SplineComponentDetails.cpp
+++ b/Engine/Source/Editor/DetailCustomizations/Private/SplineComponentDetails.cpp
@@ -1124,7 +1124,15 @@ void FSplinePointDetails::OnSplinePointTypeChanged(TSharedPtr<FString> NewValue,
 	const FScopedTransaction Transaction(LOCTEXT("SetSplinePointType", "Set spline point type"));
 	SplineComp->Modify();
 
-	EInterpCurveMode Mode = ConvertSplinePointTypeToInterpCurveMode((ESplinePointType::Type)SplinePointTypes.Find(NewValue));
+	EInterpCurveMode Mode = CIM_Unknown;
+	if (NewValue.IsValid() && SplinePointTypes.Contains(NewValue))
+	{
+		const UEnum* SplinePointTypeEnum = StaticEnum<ESplinePointType::Type>();
+		check(SplinePointTypeEnum);
+		const int64 SplinePointType = SplinePointTypeEnum->GetValueByNameString(*NewValue);
+
+		Mode = ConvertSplinePointTypeToInterpCurveMode((ESplinePointType::Type)SplinePointType);
+	}
 
 	for (int32 Index : SelectedKeys)
 	{
