Changed c-style cast in SplineComponentDetails to static_cast

diff --git a/Engine/Source/Editor/DetailCustomizations/Private/SplineComponentDetails.cpp b/Engine/Source/Editor/DetailCustomizations/Private/SplineComponentDetails.cpp
--- a/Engine/Source/Editor/DetailCustomizations/Private/SplineComponentDetails.cpp
+++ b/Engine/Source/Editor/DetailCustomizations/Private/SplineComponentDetails.cpp
@@ -1131,7 +1131,7 @@ void FSplinePointDetails::OnSplinePointTypeChanged(TSharedPtr<FString> NewValue,
 		check(SplinePointTypeEnum);
 		const int64 SplinePointType = SplinePointTypeEnum->GetValueByNameString(*NewValue);
 
-		Mode = ConvertSplinePointTypeToInterpCurveMode((ESplinePointType::Type)SplinePointType);
+		Mode = ConvertSplinePointTypeToInterpCurveMode(static_cast<ESplinePointType::Type>(SplinePointType));
 	}
 
 	for (int32 Index : SelectedKeys)
