Material Baking: Correctly apply TextureSizingType based customization to all texture size properties

diff --git a/Engine/Source/Editor/DetailCustomizations/Private/MaterialProxySettingsCustomizations.cpp b/Engine/Source/Editor/DetailCustomizations/Private/MaterialProxySettingsCustomizations.cpp
--- a/Engine/Source/Editor/DetailCustomizations/Private/MaterialProxySettingsCustomizations.cpp
+++ b/Engine/Source/Editor/DetailCustomizations/Private/MaterialProxySettingsCustomizations.cpp
@@ -86,6 +86,8 @@ void FMaterialProxySettingsCustomizations::CustomizeChildren(TSharedRef<IPropert
 	PropertyTextureSizeHandles.Add(PropertyHandles.FindChecked(GET_MEMBER_NAME_CHECKED(FMaterialProxySettings, OpacityTextureSize)));
 	PropertyTextureSizeHandles.Add(PropertyHandles.FindChecked(GET_MEMBER_NAME_CHECKED(FMaterialProxySettings, OpacityMaskTextureSize)));
 	PropertyTextureSizeHandles.Add(PropertyHandles.FindChecked(GET_MEMBER_NAME_CHECKED(FMaterialProxySettings, AmbientOcclusionTextureSize)));
+	PropertyTextureSizeHandles.Add(PropertyHandles.FindChecked(GET_MEMBER_NAME_CHECKED(FMaterialProxySettings, TangentTextureSize)));
+	PropertyTextureSizeHandles.Add(PropertyHandles.FindChecked(GET_MEMBER_NAME_CHECKED(FMaterialProxySettings, AnisotropyTextureSize)));
 
 	if (PropertyHandles.Contains(GET_MEMBER_NAME_CHECKED(FMaterialProxySettings, MaterialMergeType)))
 	{
