Arrays of FTimecodes now display the additional property buttons (eg. Delete) next to the rows.

diff --git a/Engine/Source/Editor/DetailCustomizations/Private/TimecodeDetailsCustomization.cpp b/Engine/Source/Editor/DetailCustomizations/Private/TimecodeDetailsCustomization.cpp
--- a/Engine/Source/Editor/DetailCustomizations/Private/TimecodeDetailsCustomization.cpp
+++ b/Engine/Source/Editor/DetailCustomizations/Private/TimecodeDetailsCustomization.cpp
@@ -32,13 +32,11 @@ void FTimecodeDetailsCustomization::CustomizeChildren(TSharedRef<IPropertyHandle
 {
 	TimecodeProperty = PropertyHandle;
 
-	ChildBuilder.AddCustomRow(LOCTEXT("TimecodeLabel", "Timecode"))
+	ChildBuilder.AddProperty(PropertyHandle)
+		.CustomWidget()
 		.NameContent()
 		[
-			SNew(STextBlock)
-			.Text(PropertyHandle->GetPropertyDisplayName())
-			.ToolTipText(LOCTEXT("TimecodeLabelTooltip", "Timecode"))
-			.Font(CustomizationUtils.GetRegularFont())
+			PropertyHandle->CreatePropertyNameWidget()
 		]
 		.ValueContent()
 		[
diff --git a/Engine/Source/Editor/UnrealEd/Classes/Editor/PropertyEditorTestObject.h b/Engine/Source/Editor/UnrealEd/Classes/Editor/PropertyEditorTestObject.h
--- a/Engine/Source/Editor/UnrealEd/Classes/Editor/PropertyEditorTestObject.h
+++ b/Engine/Source/Editor/UnrealEd/Classes/Editor/PropertyEditorTestObject.h
@@ -373,7 +373,10 @@ class UPropertyEditorTestObject : public UObject
 	TArray<FColor> ColorPropertyArray;
 
 	UPROPERTY(EditAnywhere, Category=ArraysOfProperties)
-	TArray<TEnumAsByte<enum EPropertyEditorTestEnum> > EnumPropertyArray;
+	TArray<FTimecode> TimecodePropertyArray;
+
+	UPROPERTY(EditAnywhere, Category=ArraysOfProperties)
+	TArray<TEnumAsByte<EPropertyEditorTestEnum> > EnumPropertyArray;
 
 	UPROPERTY(EditAnywhere, Category=ArraysOfProperties)
 	TArray<FPropertyEditorTestBasicStruct> StructPropertyArray;
