Hide Transformation field from SoundSourceBus. Partial fix for #jira UE-161069
#trivial

diff --git a/Engine/Source/Editor/DetailCustomizations/Private/SoundSourceBusDetails.cpp b/Engine/Source/Editor/DetailCustomizations/Private/SoundSourceBusDetails.cpp
--- a/Engine/Source/Editor/DetailCustomizations/Private/SoundSourceBusDetails.cpp
+++ b/Engine/Source/Editor/DetailCustomizations/Private/SoundSourceBusDetails.cpp
@@ -2,6 +2,8 @@
 
 #include "SoundSourceBusDetails.h"
 
+#include "DetailLayoutBuilder.h"
+
 #define LOCTEXT_NAMESPACE "FSoundSourceBusDetails"
 
 TSharedRef<IDetailCustomization> FSoundSourceBusDetails::MakeInstance()
@@ -12,7 +14,7 @@ TSharedRef<IDetailCustomization> FSoundSourceBusDetails::MakeInstance()
 void FSoundSourceBusDetails::CustomizeDetails(IDetailLayoutBuilder& DetailBuilder)
 {
 	// Hide details of sound base that aren't relevant for buses
-
+	DetailBuilder.HideCategory("Transformations");
 
 }
 
