Property Editor:

-Add customization for FLightingChannels

diff --git a/Engine/Plugins/Editor/ObjectMixer/LightMixer/Source/LightMixer/Public/LightMixerObjectFilter.h b/Engine/Plugins/Editor/ObjectMixer/LightMixer/Source/LightMixer/Public/LightMixerObjectFilter.h
--- a/Engine/Plugins/Editor/ObjectMixer/LightMixer/Source/LightMixer/Public/LightMixerObjectFilter.h
+++ b/Engine/Plugins/Editor/ObjectMixer/LightMixer/Source/LightMixer/Public/LightMixerObjectFilter.h
@@ -35,13 +35,13 @@ public:
 	{
 		return 
 		{
-			"Intensity", "LightColor", "AttenuationRadius"
+			"Intensity", "LightColor", "LightingChannels", "AttenuationRadius"
 		};
 	}
 
 	virtual TSet<FName> GetForceAddedColumns() const override
 	{
-		return {"LightColor"};
+		return {"LightColor", "LightingChannels"};
 	}
 
 	virtual bool ShouldIncludeUnsupportedProperties() const override
diff --git a/Engine/Source/Editor/DetailCustomizations/Private/DetailCustomizations.cpp b/Engine/Source/Editor/DetailCustomizations/Private/DetailCustomizations.cpp
--- a/Engine/Source/Editor/DetailCustomizations/Private/DetailCustomizations.cpp
+++ b/Engine/Source/Editor/DetailCustomizations/Private/DetailCustomizations.cpp
@@ -80,6 +80,7 @@
 #include "LevelSequenceActorDetails.h"
 #include "LevelSequenceBurnInOptionsCustomization.h"
 #include "LightComponentDetails.h"
+#include "LightingChannelsCustomization.h"
 #include "LinuxTargetSettingsDetails.h"
 #include "LocalLightComponentDetails.h"
 #include "MacTargetSettingsDetails.h"
@@ -348,6 +349,7 @@ void FDetailCustomizationsModule::RegisterPropertyTypeCustomizations()
 	RegisterCustomPropertyTypeLayout("DebugCameraControllerSettingsViewModeIndex", FOnGetPropertyTypeCustomizationInstance::CreateStatic(&FDebugCameraControllerSettingsViewModeIndexCustomization::MakeInstance));
 	RegisterCustomPropertyTypeLayout("CustomPrimitiveData", FOnGetPropertyTypeCustomizationInstance::CreateStatic(&FCustomPrimitiveDataCustomization::MakeInstance));
 	RegisterCustomPropertyTypeLayout("TemplateString", FOnGetPropertyTypeCustomizationInstance::CreateStatic(&FTemplateStringStructCustomization::MakeInstance));
+	RegisterCustomPropertyTypeLayout("LightingChannels", FOnGetPropertyTypeCustomizationInstance::CreateStatic(&FLightingChannelsCustomization::MakeInstance));
 }
 
 #undef REGISTER_UIMINMAX_CUSTOMIZATION
diff --git a/Engine/Source/Editor/DetailCustomizations/Private/LightingChannelsCustomization.cpp b/Engine/Source/Editor/DetailCustomizations/Private/LightingChannelsCustomization.cpp
new file mode 100644
--- /dev/null
+++ b/Engine/Source/Editor/DetailCustomizations/Private/LightingChannelsCustomization.cpp
@@ -0,0 +1,144 @@
+// Copyright Epic Games, Inc. All Rights Reserved.
+
+#include "LightingChannelsCustomization.h"
+
+#include "DetailLayoutBuilder.h"
+#include "DetailWidgetRow.h"
+#include "IDetailChildrenBuilder.h"
+#include "Internationalization/Text.h"
+#include "Styling/AppStyle.h"
+#include "Styling/StyleColors.h"
+#include "Widgets/Input/SCheckBox.h"
+#include "Widgets/Layout/SWrapBox.h"
+#include "Widgets/Text/STextBlock.h"
+
+#define LOCTEXT_NAMESPACE "FLightingChannelsCustomization"
+
+TSharedRef<IPropertyTypeCustomization> FLightingChannelsCustomization::MakeInstance()
+{
+	return MakeShareable( new FLightingChannelsCustomization );
+}
+
+void FLightingChannelsCustomization::CustomizeHeader(TSharedRef<IPropertyHandle> StructPropertyHandle,
+	FDetailWidgetRow& HeaderRow, IPropertyTypeCustomizationUtils& StructCustomizationUtils)
+{
+	LightingChannelsHandle = StructPropertyHandle;
+
+	TSharedRef<SHorizontalBox> ButtonOptionsPanel =
+		SNew(SHorizontalBox)
+		.Visibility(EVisibility::SelfHitTestInvisible)
+	;
+	
+	HeaderRow
+	.NameContent()
+	[
+		SNew(STextBlock)
+		.Text(this, &FLightingChannelsCustomization::GetStructPropertyNameText)
+		.ToolTipText(this, &FLightingChannelsCustomization::GetStructPropertyTooltipText)
+		.Font(IDetailLayoutBuilder::GetDetailFont())
+		.Visibility(EVisibility::SelfHitTestInvisible)
+	]
+	.ValueContent()
+	[
+		ButtonOptionsPanel
+	];
+
+	uint32 ChildCount;
+	LightingChannelsHandle->GetNumChildren(ChildCount);
+
+	for (uint32 ChildIndex = 0; ChildIndex < ChildCount; ChildIndex++)
+	{
+		const FText NumericText = FText::AsNumber(ChildIndex);
+		const FText SlotTooltipText = FText::Format(LOCTEXT("LightingChannelToggleFormat", "Toggle Lighting Channel {0}"), NumericText);
+		
+		ButtonOptionsPanel->AddSlot()
+		.HAlign(ChildIndex == 0 ? HAlign_Left : ChildIndex == ChildCount - 1 ? HAlign_Right : HAlign_Center)
+		.VAlign(VAlign_Center)
+		[
+			SNew(SCheckBox)
+			.Style(&FAppStyle::Get().GetWidgetStyle<FCheckBoxStyle>("DetailsView.ChannelToggleButton"))
+			.ToolTipText(SlotTooltipText)
+			.OnCheckStateChanged(this, &FLightingChannelsCustomization::OnButtonCheckedStateChanged, ChildIndex)
+			.IsChecked(this, &FLightingChannelsCustomization::GetButtonCheckedState, ChildIndex)
+			.Padding(FMargin(6,2))
+			.HAlign(HAlign_Center)
+			[
+				SNew(STextBlock)
+				.Font(FAppStyle::Get().GetFontStyle("SmallText"))
+				.Visibility(EVisibility::HitTestInvisible)
+				.Text(NumericText)
+				.ColorAndOpacity(FSlateColor::UseForeground())
+			]
+		];
+	}
+}
+
+void FLightingChannelsCustomization::CustomizeChildren(TSharedRef<IPropertyHandle> PropertyHandle,
+	IDetailChildrenBuilder& ChildBuilder, IPropertyTypeCustomizationUtils& CustomizationUtils)
+{
+	// Display channels as a normal foldout struct
+	uint32 ChildCount;
+	PropertyHandle->GetNumChildren(ChildCount);
+
+	for (uint32 ChildIndex = 0; ChildIndex < ChildCount; ChildIndex++)
+	{
+		if (TSharedPtr<IPropertyHandle> ChildHandle = PropertyHandle->GetChildHandle(ChildIndex))
+		{
+			ChildBuilder.AddProperty(ChildHandle.ToSharedRef());
+		}
+	}
+}
+
+FText FLightingChannelsCustomization::GetStructPropertyNameText() const
+{
+	if (LightingChannelsHandle.IsValid())
+	{
+		return LightingChannelsHandle->GetPropertyDisplayName();
+	}
+
+	return FText::GetEmpty();
+}
+
+FText FLightingChannelsCustomization::GetStructPropertyTooltipText() const
+{
+	if (LightingChannelsHandle.IsValid())
+	{
+		return LightingChannelsHandle->GetToolTipText();
+	}
+
+	return FText::GetEmpty();
+}
+
+void FLightingChannelsCustomization::OnButtonCheckedStateChanged(ECheckBoxState NewState, uint32 ChildIndex) const
+{
+	if (LightingChannelsHandle.IsValid())
+	{
+		uint32 OutNumChildren = 0;
+		LightingChannelsHandle->GetNumChildren(OutNumChildren);
+
+		if (ChildIndex < OutNumChildren)
+		{
+			LightingChannelsHandle->GetChildHandle(ChildIndex)->SetValue(NewState == ECheckBoxState::Checked);
+		}
+	}
+}
+
+ECheckBoxState FLightingChannelsCustomization::GetButtonCheckedState(uint32 ChildIndex) const
+{
+	if (LightingChannelsHandle.IsValid())
+	{
+		uint32 OutNumChildren = 0;
+		LightingChannelsHandle->GetNumChildren(OutNumChildren);
+
+		if (ChildIndex < OutNumChildren)
+		{
+			bool Value;
+			LightingChannelsHandle->GetChildHandle(ChildIndex)->GetValue(Value);
+			return Value ? ECheckBoxState::Checked : ECheckBoxState::Unchecked;
+		}
+	}
+
+	return ECheckBoxState::Undetermined;
+}
+
+#undef LOCTEXT_NAMESPACE
diff --git a/Engine/Source/Editor/DetailCustomizations/Private/LightingChannelsCustomization.h b/Engine/Source/Editor/DetailCustomizations/Private/LightingChannelsCustomization.h
new file mode 100644
--- /dev/null
+++ b/Engine/Source/Editor/DetailCustomizations/Private/LightingChannelsCustomization.h
@@ -0,0 +1,32 @@
+// Copyright Epic Games, Inc. All Rights Reserved.
+
+#pragma once
+
+#include "IPropertyTypeCustomization.h"
+#include "PropertyHandle.h"
+
+/**
+ * Customizes Lighting Channels as a horizontal row of buttons.
+ */
+class DETAILCUSTOMIZATIONS_API FLightingChannelsCustomization : public IPropertyTypeCustomization
+{
+public:
+	
+	static TSharedRef<IPropertyTypeCustomization> MakeInstance();
+
+	FLightingChannelsCustomization(){}
+
+	virtual void CustomizeHeader(TSharedRef<IPropertyHandle> StructPropertyHandle, class FDetailWidgetRow& HeaderRow, IPropertyTypeCustomizationUtils& StructCustomizationUtils) override;
+	virtual void CustomizeChildren(TSharedRef<IPropertyHandle> PropertyHandle, class IDetailChildrenBuilder& ChildBuilder, IPropertyTypeCustomizationUtils& CustomizationUtils) override;
+
+private:
+
+	FText GetStructPropertyNameText() const;
+	FText GetStructPropertyTooltipText() const;
+	
+	void OnButtonCheckedStateChanged(ECheckBoxState NewState, uint32 ChildIndex) const;
+	ECheckBoxState GetButtonCheckedState(uint32 ChildIndex) const;
+
+	TSharedPtr<IPropertyHandle> LightingChannelsHandle;
+	FCheckBoxStyle Style;
+};
diff --git a/Engine/Source/Editor/EditorStyle/Private/StarshipStyle.cpp b/Engine/Source/Editor/EditorStyle/Private/StarshipStyle.cpp
--- a/Engine/Source/Editor/EditorStyle/Private/StarshipStyle.cpp
+++ b/Engine/Source/Editor/EditorStyle/Private/StarshipStyle.cpp
@@ -2927,6 +2927,17 @@ void FStarshipEditorStyle::FStyle::SetupPropertyEditorStyles()
 			.SetPadding(FMargin(16, 4))
 		);
 
+		Set( "DetailsView.ChannelToggleButton", FCheckBoxStyle( FStarshipCoreStyle::GetCoreStyle().GetWidgetStyle<FCheckBoxStyle>("ToggleButtonCheckbox"))
+			.SetUncheckedImage(FSlateRoundedBoxBrush(FStyleColors::Input, 4.0f, FStyleColors::DropdownOutline, 1.0f))
+			.SetUncheckedHoveredImage(FSlateRoundedBoxBrush(FStyleColors::Input, 4.0f, FStyleColors::Hover, 1.0f))
+			.SetUncheckedPressedImage(FSlateRoundedBoxBrush(FStyleColors::Hover, 4.0f, FStyleColors::DropdownOutline, 1.0f))
+			.SetCheckedForegroundColor(FStyleColors::White)
+			.SetCheckedImage(FSlateRoundedBoxBrush(FStyleColors::Primary, 4.0f, FStyleColors::DropdownOutline, 1.0f))
+			.SetCheckedHoveredImage(FSlateRoundedBoxBrush(FStyleColors::Primary, 4.0f, FStyleColors::Hover, 1.0f))
+			.SetCheckedPressedImage(FSlateRoundedBoxBrush(FStyleColors::Hover, 4.0f, FStyleColors::DropdownOutline, 1.0f))
+			.SetPadding(FMargin(16, 4))
+		);
+
 		Set( "DetailsView.CategoryFontStyle", FStyleFonts::Get().SmallBold);
 		Set( "DetailsView.CategoryTextStyle", FTextBlockStyle(NormalText)
 			.SetFont(GetFontStyle("DetailsView.CategoryFontStyle"))
