[UE][FIX] - Display physics constraint transform properties in details panels outside PhAT
- Change Pos, PriAxis and SecAxis property's catagory
- Only hide Pos properties when in PhAT
- Fix typo in LocText tag (unrealted to above)

#rmx

diff --git a/Engine/Source/Editor/DetailCustomizations/Private/PhysicsConstraintComponentDetails.cpp b/Engine/Source/Editor/DetailCustomizations/Private/PhysicsConstraintComponentDetails.cpp
--- a/Engine/Source/Editor/DetailCustomizations/Private/PhysicsConstraintComponentDetails.cpp
+++ b/Engine/Source/Editor/DetailCustomizations/Private/PhysicsConstraintComponentDetails.cpp
@@ -904,14 +904,15 @@ void FPhysicsConstraintComponentDetails::AddConstraintProperties(IDetailLayoutBu
 		AddConstraintFrameTransform(EConstraintFrame::Frame2, ConstraintTransformsCat, ConstraintInstance);
 
 		UpdateTransformProxyDisplayRelativeToDefault(ConstraintInstance);
+
+		// Hide the constraint positions as they are represented by the proxy values in PhAT.
+		ChildPositionPropertyHandle->MarkHiddenByCustomization();
+		ParentPositionPropertyHandle->MarkHiddenByCustomization();
 	}
 
-	// Hide the constraint transform values that are being represented by the proxy values in the details panel.
-	ChildPositionPropertyHandle->MarkHiddenByCustomization();
+	// Always hide the constraint orientation values as they are represented by the proxy values in PhAT and hidden by design everywhere else.
 	ChildPriAxisPropertyHandle->MarkHiddenByCustomization();
 	ChildSecAxisPropertyHandle->MarkHiddenByCustomization();
-
-	ParentPositionPropertyHandle->MarkHiddenByCustomization();
 	ParentPriAxisPropertyHandle->MarkHiddenByCustomization();
 	ParentSecAxisPropertyHandle->MarkHiddenByCustomization();
 }	
@@ -961,7 +962,7 @@ void FPhysicsConstraintComponentDetails::AddConstraintFrameTransform(const ECons
 	{
 		TransformProxy = ChildTransformProxy;
 		TransformProxy->SetFrameLabelText(LOCTEXT("Child", "Child"));
-		KeyboardShortcutText = LOCTEXT("ShiftCtl", "[Shift + Alt]");
+		KeyboardShortcutText = LOCTEXT("ShiftAlt", "[Shift + Alt]");
 		PositionPropertyHandle = ChildPositionPropertyHandle;
 		PriAxisPropertyHandle = ChildPriAxisPropertyHandle;
 		SecAxisPropertyHandle = ChildSecAxisPropertyHandle;
diff --git a/Engine/Source/Runtime/Engine/Classes/PhysicsEngine/ConstraintInstance.h b/Engine/Source/Runtime/Engine/Classes/PhysicsEngine/ConstraintInstance.h
--- a/Engine/Source/Runtime/Engine/Classes/PhysicsEngine/ConstraintInstance.h
+++ b/Engine/Source/Runtime/Engine/Classes/PhysicsEngine/ConstraintInstance.h
@@ -256,29 +256,29 @@ public:
 	///////////////////////////// Body1 ref frame
 	
 	/** Location of constraint in Body1 reference frame (usually the "child" body for skeletal meshes). */
-	UPROPERTY(EditAnywhere, Category = ConstraintTransforms)
+	UPROPERTY(EditAnywhere, Category = Constraint)
 	FVector Pos1;
 
 	/** Primary (twist) axis in Body1 reference frame. */
-	UPROPERTY(EditAnywhere, Category = ConstraintTransforms)
+	UPROPERTY(EditAnywhere, Category = Constraint)
 	FVector PriAxis1;
 
 	/** Secondary axis in Body1 reference frame. Orthogonal to PriAxis1. */
-	UPROPERTY(EditAnywhere, Category = ConstraintTransforms)
+	UPROPERTY(EditAnywhere, Category = Constraint)
 	FVector SecAxis1;
 
 	///////////////////////////// Body2 ref frame
 	
 	/** Location of constraint in Body2 reference frame (usually the "parent" body for skeletal meshes). */
-	UPROPERTY(EditAnywhere, Category= ConstraintTransforms)
+	UPROPERTY(EditAnywhere, Category= Constraint)
 	FVector Pos2;
 
 	/** Primary (twist) axis in Body2 reference frame. */
-	UPROPERTY(EditAnywhere, Category = ConstraintTransforms)
+	UPROPERTY(EditAnywhere, Category = Constraint)
 	FVector PriAxis2;
 
 	/** Secondary axis in Body2 reference frame. Orthogonal to PriAxis2. */
-	UPROPERTY(EditAnywhere, Category = ConstraintTransforms)
+	UPROPERTY(EditAnywhere, Category = Constraint)
 	FVector SecAxis2;
 
 	/** Specifies the angular offset between the two frames of reference. By default limit goes from (-Angle, +Angle)
