Make mapped key tooltip more helpful when disabled from a combo trigger

diff --git a/Engine/Plugins/EnhancedInput/Source/InputEditor/Private/InputCustomizations.cpp b/Engine/Plugins/EnhancedInput/Source/InputEditor/Private/InputCustomizations.cpp
--- a/Engine/Plugins/EnhancedInput/Source/InputEditor/Private/InputCustomizations.cpp
+++ b/Engine/Plugins/EnhancedInput/Source/InputEditor/Private/InputCustomizations.cpp
@@ -60,8 +60,8 @@ void FEnhancedActionMappingCustomization::CustomizeHeader(TSharedRef<IPropertyHa
 	// Pass our header row into the key struct customizeheader method so it populates our row with the key struct header
 	KeyStructCustomization = StaticCastSharedPtr<FKeyStructCustomization>(KeyStructInstance);
 	const bool bContainsComboTrigger = DoesTriggerArrayContainCombo();
-	KeyStructCustomization->bDisplayIcon = bContainsComboTrigger;
-	KeyStructCustomization->bEnableKeySelector = !bContainsComboTrigger;
+	KeyStructCustomization->SetDisplayIcon(bContainsComboTrigger);
+	KeyStructCustomization->SetEnableKeySelector(!bContainsComboTrigger);
 	KeyStructCustomization->CustomizeHeaderOnlyWithButton(KeyHandle.ToSharedRef(), HeaderRow, CustomizationUtils, RemoveButton);
 }
 
@@ -93,8 +93,8 @@ void FEnhancedActionMappingCustomization::RemoveMappingButton_OnClick() const
 void FEnhancedActionMappingCustomization::OnTriggersChanged() const
 {
 	const bool bContainsComboTrigger = DoesTriggerArrayContainCombo();
-	KeyStructCustomization->bDisplayIcon = bContainsComboTrigger;
-	KeyStructCustomization->bEnableKeySelector = !bContainsComboTrigger;
+	KeyStructCustomization->SetDisplayIcon(bContainsComboTrigger);
+    KeyStructCustomization->SetEnableKeySelector(!bContainsComboTrigger);
 }
 
 bool FEnhancedActionMappingCustomization::DoesTriggerArrayContainCombo() const
diff --git a/Engine/Source/Editor/DetailCustomizations/Private/KeyStructCustomization.cpp b/Engine/Source/Editor/DetailCustomizations/Private/KeyStructCustomization.cpp
--- a/Engine/Source/Editor/DetailCustomizations/Private/KeyStructCustomization.cpp
+++ b/Engine/Source/Editor/DetailCustomizations/Private/KeyStructCustomization.cpp
@@ -12,7 +12,6 @@
 #include "Layout/Margin.h"
 #include "Misc/Attribute.h"
 #include "PropertyHandle.h"
-#include "SKeySelector.h"
 #include "SlotBase.h"
 #include "Types/SlateEnums.h"
 #include "UObject/ObjectMacros.h"
@@ -27,6 +26,15 @@ class SWidget;
 /* FKeyStructCustomization static interface
  *****************************************************************************/
 
+void FKeyStructCustomization::SetEnableKeySelector(bool bKeySelectorEnabled)
+{
+	bEnableKeySelector = bKeySelectorEnabled;
+	if (KeySelector)
+	{
+		KeySelector->SetEnabledFromKeyStructCustomization(bEnableKeySelector);
+	}
+}
+
 TSharedRef<IPropertyTypeCustomization> FKeyStructCustomization::MakeInstance( )
 {
 	return MakeShareable(new FKeyStructCustomization);
@@ -61,6 +69,19 @@ void FKeyStructCustomization::CustomizeHeaderOnlyWithButton(TSharedRef<class IPr
 {
 	PropertyHandle = StructPropertyHandle;
 
+	KeySelector = SNew(SKeySelector)
+		.CurrentKey(this, &FKeyStructCustomization::GetCurrentKey)
+		.OnKeyChanged(this, &FKeyStructCustomization::OnKeyChanged)
+		.Font(StructCustomizationUtils.GetRegularFont())
+		.AllowClear(!StructPropertyHandle->GetProperty()->HasAnyPropertyFlags(CPF_NoClear))
+		.FilterBlueprintBindable(false)
+		.IsEnabled_Lambda([this]() -> bool
+	    {
+		    return bEnableKeySelector;
+	    });
+	
+	KeySelector->SetEnabledFromKeyStructCustomization(bEnableKeySelector);
+	
 	// create struct header
 	HeaderRow.NameContent()
 	.MinDesiredWidth(125.0f)
@@ -71,16 +92,7 @@ void FKeyStructCustomization::CustomizeHeaderOnlyWithButton(TSharedRef<class IPr
 		.Padding(InputSettingsDetails::InputConstants::PropertyPadding)
 		//.AutoWidth()
 		[
-			SNew(SKeySelector)
-			.CurrentKey(this, &FKeyStructCustomization::GetCurrentKey)
-			.OnKeyChanged(this, &FKeyStructCustomization::OnKeyChanged)
-			.Font(StructCustomizationUtils.GetRegularFont())
-			.AllowClear(!StructPropertyHandle->GetProperty()->HasAnyPropertyFlags(CPF_NoClear))
-		    .FilterBlueprintBindable(false)
-		    .IsEnabled_Lambda([this]() -> bool
-		    {
-		    	return bEnableKeySelector;
-		    })
+			KeySelector.ToSharedRef()
 		]
 		+ SHorizontalBox::Slot()
 		.Padding(InputSettingsDetails::InputConstants::PropertyPadding)
diff --git a/Engine/Source/Editor/DetailCustomizations/Public/KeyStructCustomization.h b/Engine/Source/Editor/DetailCustomizations/Public/KeyStructCustomization.h
--- a/Engine/Source/Editor/DetailCustomizations/Public/KeyStructCustomization.h
+++ b/Engine/Source/Editor/DetailCustomizations/Public/KeyStructCustomization.h
@@ -7,6 +7,7 @@
 #include "InputCoreTypes.h"
 #include "Misc/Optional.h"
 #include "PropertyHandle.h"
+#include "SKeySelector.h"
 #include "Templates/SharedPointer.h"
 
 class IPropertyHandle;
@@ -30,8 +31,27 @@ public:
 	// Helper variant that generates the key struct in the header and appends a single button at the end
 	// TODO: Is there a better way?
 	void CustomizeHeaderOnlyWithButton(TSharedRef<class IPropertyHandle> StructPropertyHandle, class FDetailWidgetRow& HeaderRow, IPropertyTypeCustomizationUtils& StructCustomizationUtils, TSharedRef<SWidget> Button);
-	bool bDisplayIcon = false;
-	bool bEnableKeySelector = true;
+
+	/** Sets a bool for whether the key selector should be enabled or not */
+	void SetEnableKeySelector(bool bKeySelectorEnabled);
+	/** Sets a bool for whether the information Icon for a combo trigger is displayed or not */ 
+	void SetDisplayIcon(bool bShouldDisplayIcon)
+	{
+		bDisplayIcon = bShouldDisplayIcon;
+	}
+
+	/** Gets bEnableKeySelector bool */
+	bool GetEnableKeySelector() const
+	{
+		return bEnableKeySelector;
+	}
+	
+	/** Gets bDisplayIcon bool */
+    bool GetDisplayIcon() const
+    {
+    	return bDisplayIcon;
+    }
+
 public:
 
 	/**
@@ -42,6 +62,10 @@ public:
 	static TSharedRef<IPropertyTypeCustomization> MakeInstance( );
 
 protected:
+	/** Whether the information icon for combo triggers is displayed or not */
+	bool bDisplayIcon = false;
+	/** Whether the key selector should be enabled or not */
+	bool bEnableKeySelector = true;
 
 	/** Gets the current Key being edited. */
 	TOptional<FKey> GetCurrentKey() const;
@@ -51,4 +75,7 @@ protected:
 
 	/** Holds a handle to the property being edited. */
 	TSharedPtr<IPropertyHandle> PropertyHandle;
+	
+	/** shared pointer to the Key Selector. */
+    TSharedPtr<SKeySelector> KeySelector = nullptr;
 };
diff --git a/Engine/Source/Editor/UnrealEd/Private/SKeySelector.cpp b/Engine/Source/Editor/UnrealEd/Private/SKeySelector.cpp
--- a/Engine/Source/Editor/UnrealEd/Private/SKeySelector.cpp
+++ b/Engine/Source/Editor/UnrealEd/Private/SKeySelector.cpp
@@ -171,7 +171,7 @@ void SKeySelector::Construct(const FArguments& InArgs)
 			SAssignNew(KeyComboButton, SComboButton)
 			.OnGetMenuContent(this, &SKeySelector::GetMenuContent)
 			.ContentPadding(0)
-			.ToolTipText(this, &SKeySelector::GetKeyDescription)	// Longer key descriptions can overrun the visible space in the combo button if the parent width is constrained, so we reflect them in the tooltip too.
+			.ToolTipText(this, &SKeySelector::GetKeyDescriptionToolTip)	// Longer key descriptions can overrun the visible space in the combo button if the parent width is constrained, so we reflect them in the tooltip too.
 			.ButtonContent()
 			[
 				SNew(SHorizontalBox)
@@ -205,6 +205,21 @@ FText SKeySelector::GetKeyDescription() const
 	return LOCTEXT("MultipleValues", "Multiple Values");
 }
 
+FText SKeySelector::GetKeyDescriptionToolTip() const
+{
+	if(!bEnabledFromKeyStructCustomization)
+    {
+    	return LOCTEXT("KeySelectorDisabledToolTipText", "This Key cannot be part of the combo trigger so Key selection has been disabled.");
+    }
+    
+	TOptional<FKey> CurrentKeyValue = CurrentKey.Get();
+	if (CurrentKeyValue.IsSet())
+	{
+		return CurrentKeyValue.GetValue().GetDisplayName();
+	}
+	return LOCTEXT("MultipleValues", "Multiple Values");
+}
+
 const FSlateBrush* SKeySelector::GetKeyIconImage() const
 {
 	TOptional<FKey> CurrentKeyValue = CurrentKey.Get();
@@ -222,6 +237,11 @@ const FSlateBrush* SKeySelector::GetKeyIconImage() const
 
 FText SKeySelector::GetKeyTooltip() const
 {
+	if(!bEnabledFromKeyStructCustomization)
+	{
+		return LOCTEXT("KeySelectorDisabledToolTipText", "This Key cannot be part of the combo trigger so Key selection has been disabled.");
+	}
+	
 	TOptional<FKey> CurrentKeyValue = CurrentKey.Get();
 	if (CurrentKeyValue.IsSet())
 	{
diff --git a/Engine/Source/Editor/UnrealEd/Public/SKeySelector.h b/Engine/Source/Editor/UnrealEd/Public/SKeySelector.h
--- a/Engine/Source/Editor/UnrealEd/Public/SKeySelector.h
+++ b/Engine/Source/Editor/UnrealEd/Public/SKeySelector.h
@@ -75,6 +75,18 @@ public:
 public:
 	void Construct(const FArguments& InArgs);
 
+	/** Sets bool to produce tooltip notifying this key selector it was disabled from KeyStructCustomization */
+	void SetEnabledFromKeyStructCustomization(bool bIsEnabled)
+	{
+		bEnabledFromKeyStructCustomization = bIsEnabled;
+	}
+
+	/** Gets bEnabledFromKeyStructCustomization bool */
+	bool GetSetEnabledFromKeyStructCustomization() const
+	{
+		return bEnabledFromKeyStructCustomization;
+	}
+	
 protected:
 	/** Gets the icon for the key being manipulated */
 	const FSlateBrush* GetKeyIconImage() const;
@@ -83,6 +95,8 @@ protected:
 
 	/** Gets a succinct description for the key being manipulated */
 	FText GetKeyDescription() const;
+	/** Gets a description tooltip for the key being manipulated */
+	FText GetKeyDescriptionToolTip() const;
 	/** Gets a tooltip for the selected key */
 	FText GetKeyTooltip() const;
 
@@ -156,4 +170,5 @@ protected:
 	TArray<FKeyTreeItem>		FilteredKeyTreeRoot;
 
 	bool bListenForNextInput = false;
+	bool bEnabledFromKeyStructCustomization = true;
 };
