Object Mixer:

-Fix: Lighting Channels changes not applied to all selected rows
-Adjust spacing of light channel buttons

diff --git a/Engine/Plugins/Editor/ObjectMixer/ObjectMixer/Source/ObjectMixer/Private/Views/List/SObjectMixerEditorListRow.cpp b/Engine/Plugins/Editor/ObjectMixer/ObjectMixer/Source/ObjectMixer/Private/Views/List/SObjectMixerEditorListRow.cpp
--- a/Engine/Plugins/Editor/ObjectMixer/ObjectMixer/Source/ObjectMixer/Private/Views/List/SObjectMixerEditorListRow.cpp
+++ b/Engine/Plugins/Editor/ObjectMixer/ObjectMixer/Source/ObjectMixer/Private/Views/List/SObjectMixerEditorListRow.cpp
@@ -501,7 +501,14 @@ TSharedPtr<SWidget> SObjectMixerEditorListRow::GenerateCells(
 								this,
 								&SObjectMixerEditorListRow::OnPropertyChanged, Property, (void*)ObjectRef);
 					
-						SinglePropertyView->SetOnPropertyValueChanged(OnPropertyValueChanged);
+						Handle->SetOnPropertyValueChanged(OnPropertyValueChanged);
+						uint32 ChildHandleCount;
+						Handle->GetNumChildren(ChildHandleCount);
+
+						for (uint32 ChildHandleItr = 0; ChildHandleItr < ChildHandleCount; ChildHandleItr++)
+						{
+							Handle->GetChildHandle(ChildHandleItr)->SetOnPropertyValueChanged(OnPropertyValueChanged);
+						}
 
 						return SNew(SBox)
 								.Visibility(EVisibility::SelfHitTestInvisible)
diff --git a/Engine/Source/Editor/DetailCustomizations/Private/LightingChannelsCustomization.cpp b/Engine/Source/Editor/DetailCustomizations/Private/LightingChannelsCustomization.cpp
--- a/Engine/Source/Editor/DetailCustomizations/Private/LightingChannelsCustomization.cpp
+++ b/Engine/Source/Editor/DetailCustomizations/Private/LightingChannelsCustomization.cpp
@@ -9,6 +9,7 @@
 #include "Styling/AppStyle.h"
 #include "Styling/StyleColors.h"
 #include "Widgets/Input/SCheckBox.h"
+#include "Widgets/Layout/SBox.h"
 #include "Widgets/Layout/SWrapBox.h"
 #include "Widgets/Text/STextBlock.h"
 
@@ -52,22 +53,29 @@ void FLightingChannelsCustomization::CustomizeHeader(TSharedRef<IPropertyHandle>
 		const FText SlotTooltipText = FText::Format(LOCTEXT("LightingChannelToggleFormat", "Toggle Lighting Channel {0}"), NumericText);
 		
 		ButtonOptionsPanel->AddSlot()
-		.HAlign(ChildIndex == 0 ? HAlign_Left : ChildIndex == ChildCount - 1 ? HAlign_Right : HAlign_Center)
+		.HAlign(HAlign_Left)
 		.VAlign(VAlign_Center)
+		.AutoWidth()
+		.Padding(0, 0, ChildIndex == ChildCount -1 ? 16 : 0, 0)
 		[
-			SNew(SCheckBox)
-			.Style(&FAppStyle::Get().GetWidgetStyle<FCheckBoxStyle>("DetailsView.ChannelToggleButton"))
-			.ToolTipText(SlotTooltipText)
-			.OnCheckStateChanged(this, &FLightingChannelsCustomization::OnButtonCheckedStateChanged, ChildIndex)
-			.IsChecked(this, &FLightingChannelsCustomization::GetButtonCheckedState, ChildIndex)
-			.Padding(FMargin(6,2))
-			.HAlign(HAlign_Center)
+			SNew(SBox)
+			.WidthOverride(20)
+			.HAlign(HAlign_Fill)
 			[
-				SNew(STextBlock)
-				.Font(FAppStyle::Get().GetFontStyle("SmallText"))
-				.Visibility(EVisibility::HitTestInvisible)
-				.Text(NumericText)
-				.ColorAndOpacity(FSlateColor::UseForeground())
+				SNew(SCheckBox)
+				.Style(&FAppStyle::Get().GetWidgetStyle<FCheckBoxStyle>("DetailsView.ChannelToggleButton"))
+				.ToolTipText(SlotTooltipText)
+				.OnCheckStateChanged(this, &FLightingChannelsCustomization::OnButtonCheckedStateChanged, ChildIndex)
+				.IsChecked(this, &FLightingChannelsCustomization::GetButtonCheckedState, ChildIndex)
+				.HAlign(HAlign_Center)
+				.Padding(FMargin(0, 2))
+				[
+					SNew(STextBlock)
+					.Font(FAppStyle::Get().GetFontStyle("SmallText"))
+					.Visibility(EVisibility::HitTestInvisible)
+					.Text(NumericText)
+					.ColorAndOpacity(FSlateColor::UseForeground())
+				]
 			]
 		];
 	}
@@ -119,6 +127,8 @@ void FLightingChannelsCustomization::OnButtonCheckedStateChanged(ECheckBoxState
 		if (ChildIndex < OutNumChildren)
 		{
 			LightingChannelsHandle->GetChildHandle(ChildIndex)->SetValue(NewState == ECheckBoxState::Checked);
+			LightingChannelsHandle->GetChildHandle(ChildIndex)->NotifyFinishedChangingProperties();
+			LightingChannelsHandle->NotifyFinishedChangingProperties();
 		}
 	}
 }
