Material Proxy settings - Hide texture size property unless texture sizing type is "Use Single Texture Size" or "Use Automatic Biased Sizes"
* New automatic texture sizing options would show the "Texture Size" option, which was confusing

diff --git a/Engine/Source/Editor/DetailCustomizations/Private/MaterialProxySettingsCustomizations.cpp b/Engine/Source/Editor/DetailCustomizations/Private/MaterialProxySettingsCustomizations.cpp
--- a/Engine/Source/Editor/DetailCustomizations/Private/MaterialProxySettingsCustomizations.cpp
+++ b/Engine/Source/Editor/DetailCustomizations/Private/MaterialProxySettingsCustomizations.cpp
@@ -184,12 +184,12 @@ EVisibility FMaterialProxySettingsCustomizations::IsTextureSizeEnabled() const
 	uint8 TypeValue;
 	EnumHandle->GetValue(TypeValue);
 
-	if (TypeValue == TextureSizingType_UseSimplygonAutomaticSizing || TypeValue == TextureSizingType_UseManualOverrideTextureSize)
+	if (TypeValue == TextureSizingType_UseSingleTextureSize || TypeValue == TextureSizingType_UseAutomaticBiasedSizes)
 	{		
-		return EVisibility::Hidden;
+		return EVisibility::Visible;
 	}
 
-	return EVisibility::Visible;	
+	return EVisibility::Hidden;
 }
 
 EVisibility FMaterialProxySettingsCustomizations::IsSimplygonMaterialMergingVisible() const
