World Partition - HLOD: Auto configure the FMaterialProxySettings::MeshMinDrawDistance field when working with WP HLOD
* Field is now hidden when edited from an HLOD Layer
* When building HLODs, retrieve the expected draw distance from the FHLODBuildContext object

diff --git a/Engine/Plugins/Editor/WorldPartitionHLODUtilities/Source/Private/HLODBuilderMeshApproximate.cpp b/Engine/Plugins/Editor/WorldPartitionHLODUtilities/Source/Private/HLODBuilderMeshApproximate.cpp
--- a/Engine/Plugins/Editor/WorldPartitionHLODUtilities/Source/Private/HLODBuilderMeshApproximate.cpp
+++ b/Engine/Plugins/Editor/WorldPartitionHLODUtilities/Source/Private/HLODBuilderMeshApproximate.cpp
@@ -86,9 +86,15 @@ TArray<UActorComponent*> UHLODBuilderMeshApproximate::Build(const FHLODBuildCont
 	//
 
 	const UHLODBuilderMeshApproximateSettings* MeshApproximateSettings = CastChecked<UHLODBuilderMeshApproximateSettings>(HLODBuilderSettings);
-	const FMeshApproximationSettings& UseSettings = MeshApproximateSettings->MeshApproximationSettings;
+	FMeshApproximationSettings UseSettings = MeshApproximateSettings->MeshApproximationSettings; // Make a copy as we may tweak some values
 	UMaterialInterface* HLODMaterial = MeshApproximateSettings->HLODMaterial.LoadSynchronous();
 
+	// When using automatic textuse sizing based on draw distance, use the MinVisibleDistance for this HLOD.
+	if (UseSettings.MaterialSettings.TextureSizingType == TextureSizingType_AutomaticFromMeshDrawDistance)
+	{
+		UseSettings.MaterialSettings.MeshMinDrawDistance = InHLODBuildContext.MinVisibleDistance;
+	}
+
 	IGeometryProcessing_ApproximateActors::FOptions Options = ApproxActorsAPI->ConstructOptions(UseSettings);
 	Options.BasePackagePath = InHLODBuildContext.AssetsOuter->GetPackage()->GetName();
 	Options.bGenerateLightmapUVs = false;
diff --git a/Engine/Plugins/Editor/WorldPartitionHLODUtilities/Source/Private/HLODBuilderMeshMerge.cpp b/Engine/Plugins/Editor/WorldPartitionHLODUtilities/Source/Private/HLODBuilderMeshMerge.cpp
--- a/Engine/Plugins/Editor/WorldPartitionHLODUtilities/Source/Private/HLODBuilderMeshMerge.cpp
+++ b/Engine/Plugins/Editor/WorldPartitionHLODUtilities/Source/Private/HLODBuilderMeshMerge.cpp
@@ -77,9 +77,15 @@ TArray<UActorComponent*> UHLODBuilderMeshMerge::Build(const FHLODBuildContext& I
 	FVector MergedActorLocation;
 
 	const UHLODBuilderMeshMergeSettings* MeshMergeSettings = CastChecked<UHLODBuilderMeshMergeSettings>(HLODBuilderSettings);
-	const FMeshMergingSettings& UseSettings = MeshMergeSettings->MeshMergeSettings;
+	FMeshMergingSettings UseSettings = MeshMergeSettings->MeshMergeSettings; // Make a copy as we may tweak some values
 	UMaterialInterface* HLODMaterial = MeshMergeSettings->HLODMaterial.LoadSynchronous();
 
+	// When using automatic textuse sizing based on draw distance, use the MinVisibleDistance for this HLOD.
+	if (UseSettings.MaterialSettings.TextureSizingType == TextureSizingType_AutomaticFromMeshDrawDistance)
+	{
+		UseSettings.MaterialSettings.MeshMinDrawDistance = InHLODBuildContext.MinVisibleDistance;
+	}
+
 	const IMeshMergeUtilities& MeshMergeUtilities = FModuleManager::Get().LoadModuleChecked<IMeshMergeModule>("MeshMergeUtilities").GetUtilities();
 	MeshMergeUtilities.MergeComponentsToStaticMesh(SourcePrimitiveComponents, InHLODBuildContext.World, UseSettings, HLODMaterial, InHLODBuildContext.AssetsOuter->GetPackage(), InHLODBuildContext.AssetsBaseName, Assets, MergedActorLocation, 0.25f, false);
 
diff --git a/Engine/Plugins/Editor/WorldPartitionHLODUtilities/Source/Private/HLODBuilderMeshSimplify.cpp b/Engine/Plugins/Editor/WorldPartitionHLODUtilities/Source/Private/HLODBuilderMeshSimplify.cpp
--- a/Engine/Plugins/Editor/WorldPartitionHLODUtilities/Source/Private/HLODBuilderMeshSimplify.cpp
+++ b/Engine/Plugins/Editor/WorldPartitionHLODUtilities/Source/Private/HLODBuilderMeshSimplify.cpp
@@ -75,9 +75,15 @@ TArray<UActorComponent*> UHLODBuilderMeshSimplify::Build(const FHLODBuildContext
 	TArray<UStaticMeshComponent*> StaticMeshComponents = FilterComponents<UStaticMeshComponent>(InSourceComponents);
 
 	const UHLODBuilderMeshSimplifySettings* MeshSimplifySettings = CastChecked<UHLODBuilderMeshSimplifySettings>(HLODBuilderSettings);
-	const FMeshProxySettings& UseSettings = MeshSimplifySettings->MeshSimplifySettings;
+	FMeshProxySettings UseSettings = MeshSimplifySettings->MeshSimplifySettings; // Make a copy as we may tweak some values
 	UMaterialInterface* HLODMaterial = MeshSimplifySettings->HLODMaterial.LoadSynchronous();
 
+	// When using automatic textuse sizing based on draw distance, use the MinVisibleDistance for this HLOD.
+	if (UseSettings.MaterialSettings.TextureSizingType == TextureSizingType_AutomaticFromMeshDrawDistance)
+	{
+		UseSettings.MaterialSettings.MeshMinDrawDistance = InHLODBuildContext.MinVisibleDistance;
+	}
+
 	TArray<UObject*> Assets;
 	FCreateProxyDelegate ProxyDelegate;
 	ProxyDelegate.BindLambda([&Assets](const FGuid Guid, TArray<UObject*>& InAssetsCreated) { Assets = InAssetsCreated; });
diff --git a/Engine/Source/Editor/DetailCustomizations/Private/MaterialProxySettingsCustomizations.cpp b/Engine/Source/Editor/DetailCustomizations/Private/MaterialProxySettingsCustomizations.cpp
--- a/Engine/Source/Editor/DetailCustomizations/Private/MaterialProxySettingsCustomizations.cpp
+++ b/Engine/Source/Editor/DetailCustomizations/Private/MaterialProxySettingsCustomizations.cpp
@@ -2,6 +2,7 @@
 
 #include "MaterialProxySettingsCustomizations.h"
 
+#include "Algo/AnyOf.h"
 #include "Containers/Map.h"
 #include "Containers/UnrealString.h"
 #include "Delegates/Delegate.h"
@@ -27,6 +28,7 @@
 #include "UObject/NameTypes.h"
 #include "UObject/ReflectedTypeAccessors.h"
 #include "UObject/UnrealType.h"
+#include "WorldPartition/HLOD/HLODLayer.h"
 
 #define LOCTEXT_NAMESPACE "MaterialProxySettingsCustomizations"
 
@@ -75,8 +77,9 @@ void FMaterialProxySettingsCustomizations::CustomizeChildren(TSharedRef<IPropert
 	
 
 	// Retrieve special case properties
-	EnumHandle = PropertyHandles.FindChecked(GET_MEMBER_NAME_CHECKED(FMaterialProxySettings, TextureSizingType));
+	TextureSizingTypeHandle = PropertyHandles.FindChecked(GET_MEMBER_NAME_CHECKED(FMaterialProxySettings, TextureSizingType));
 	TextureSizeHandle = PropertyHandles.FindChecked(GET_MEMBER_NAME_CHECKED(FMaterialProxySettings, TextureSize));
+	MeshMinDrawDistanceHandle = PropertyHandles.FindChecked(GET_MEMBER_NAME_CHECKED(FMaterialProxySettings, MeshMinDrawDistance));
 	PropertyTextureSizeHandles.Add(PropertyHandles.FindChecked(GET_MEMBER_NAME_CHECKED(FMaterialProxySettings, DiffuseTextureSize)));
 	PropertyTextureSizeHandles.Add(PropertyHandles.FindChecked(GET_MEMBER_NAME_CHECKED(FMaterialProxySettings, NormalTextureSize)));
 	PropertyTextureSizeHandles.Add(PropertyHandles.FindChecked(GET_MEMBER_NAME_CHECKED(FMaterialProxySettings, MetallicTextureSize)));
@@ -118,7 +121,7 @@ void FMaterialProxySettingsCustomizations::CustomizeChildren(TSharedRef<IPropert
 			IDetailPropertyRow& SettingsRow = ChildBuilder.AddProperty(Iter.Value().ToSharedRef());
 			SettingsRow.Visibility(TAttribute<EVisibility>(this, &FMaterialProxySettingsCustomizations::IsSimplygonMaterialMergingVisible));
 		}
-		else if (Iter.Value() == EnumHandle)
+		else if (Iter.Value() == TextureSizingTypeHandle)
 		{
 			// Remove the simplygon specific option.
 			if (bUseNativeTool)
@@ -126,18 +129,25 @@ void FMaterialProxySettingsCustomizations::CustomizeChildren(TSharedRef<IPropert
 				TSharedPtr<FPropertyRestriction> EnumRestriction = MakeShareable(new FPropertyRestriction(LOCTEXT("NoSupport", "Unable to support this option in Merge Actor")));
 				const UEnum* const TextureSizingTypeEnum = StaticEnum<ETextureSizingType>();
 				EnumRestriction->AddHiddenValue(TextureSizingTypeEnum->GetNameStringByValue((uint8)ETextureSizingType::TextureSizingType_UseSimplygonAutomaticSizing));
-				EnumHandle->AddRestriction(EnumRestriction.ToSharedRef());
+				TextureSizingTypeHandle->AddRestriction(EnumRestriction.ToSharedRef());
 			}
 
-			IDetailPropertyRow& SettingsRow = ChildBuilder.AddProperty(Iter.Value().ToSharedRef());
+			ChildBuilder.AddProperty(Iter.Value().ToSharedRef());
 		}
-		// Do not show the merge type property
-		else if (Iter.Value() != MergeTypeHandle)
+		else if (Iter.Value() == MeshMinDrawDistanceHandle)
 		{
 			IDetailPropertyRow& SettingsRow = ChildBuilder.AddProperty(Iter.Value().ToSharedRef());
+			SettingsRow.Visibility(TAttribute<EVisibility>(this, &FMaterialProxySettingsCustomizations::IsMeshMinDrawDistanceVisible));
 		}
-	}	
-
+		else if (Iter.Value() == MergeTypeHandle)
+		{
+			// Do not show the merge type property
+		}
+		else
+		{
+			ChildBuilder.AddProperty(Iter.Value().ToSharedRef());
+		}		
+	}
 }
 
 void FMaterialProxySettingsCustomizations::AddTextureSizeClamping(TSharedPtr<IPropertyHandle> TextureSizeProperty)
@@ -169,7 +179,7 @@ void FMaterialProxySettingsCustomizations::AddTextureSizeClamping(TSharedPtr<IPr
 EVisibility FMaterialProxySettingsCustomizations::AreManualOverrideTextureSizesEnabled() const
 {
 	uint8 TypeValue;
-	EnumHandle->GetValue(TypeValue);
+	TextureSizingTypeHandle->GetValue(TypeValue);
 
 	if (TypeValue == TextureSizingType_UseManualOverrideTextureSize)
 	{
@@ -182,7 +192,7 @@ EVisibility FMaterialProxySettingsCustomizations::AreManualOverrideTextureSizesE
 EVisibility FMaterialProxySettingsCustomizations::IsTextureSizeEnabled() const
 {
 	uint8 TypeValue;
-	EnumHandle->GetValue(TypeValue);
+	TextureSizingTypeHandle->GetValue(TypeValue);
 
 	if (TypeValue == TextureSizingType_UseSingleTextureSize || TypeValue == TextureSizingType_UseAutomaticBiasedSizes)
 	{		
@@ -203,5 +213,20 @@ EVisibility FMaterialProxySettingsCustomizations::IsSimplygonMaterialMergingVisi
 	return ( MergeType == EMaterialMergeType::MaterialMergeType_Simplygon ) ? EVisibility::Visible : EVisibility::Hidden;
 }
 
+EVisibility FMaterialProxySettingsCustomizations::IsMeshMinDrawDistanceVisible() const
+{
+	uint8 TypeValue;
+	TextureSizingTypeHandle->GetValue(TypeValue);
+	if (TypeValue != TextureSizingType_AutomaticFromMeshDrawDistance)
+	{
+		return EVisibility::Hidden;
+	}
+
+	TArray<UObject*> OutersList;
+	MeshMinDrawDistanceHandle->GetOuterObjects(OutersList);
+	const bool bEditingHLODLayer = Algo::AnyOf(OutersList, [](UObject* Outer) { return Outer->IsInA(UHLODLayer::StaticClass()); });
+
+	return bEditingHLODLayer ? EVisibility::Hidden : EVisibility::Visible;
+}
 
 #undef LOCTEXT_NAMESPACE
diff --git a/Engine/Source/Editor/DetailCustomizations/Private/MaterialProxySettingsCustomizations.h b/Engine/Source/Editor/DetailCustomizations/Private/MaterialProxySettingsCustomizations.h
--- a/Engine/Source/Editor/DetailCustomizations/Private/MaterialProxySettingsCustomizations.h
+++ b/Engine/Source/Editor/DetailCustomizations/Private/MaterialProxySettingsCustomizations.h
@@ -27,12 +27,16 @@ protected:
 	EVisibility IsTextureSizeEnabled() const;
 
 	EVisibility IsSimplygonMaterialMergingVisible() const;
+
+	EVisibility IsMeshMinDrawDistanceVisible() const;
 	
-	TSharedPtr< IPropertyHandle > EnumHandle;
+	TSharedPtr< IPropertyHandle > TextureSizingTypeHandle;
 
 	TSharedPtr< IPropertyHandle > TextureSizeHandle;
 	TArray<TSharedPtr<IPropertyHandle>> PropertyTextureSizeHandles;
 
 	TSharedPtr< IPropertyHandle > MergeTypeHandle;
 	TSharedPtr< IPropertyHandle > GutterSpaceHandle;
+
+	TSharedPtr< IPropertyHandle > MeshMinDrawDistanceHandle;
 };
diff --git a/Engine/Source/Runtime/Engine/Classes/Engine/MaterialMerging.h b/Engine/Source/Runtime/Engine/Classes/Engine/MaterialMerging.h
--- a/Engine/Source/Runtime/Engine/Classes/Engine/MaterialMerging.h
+++ b/Engine/Source/Runtime/Engine/Classes/Engine/MaterialMerging.h
@@ -18,7 +18,7 @@ enum ETextureSizingType
 	TextureSizingType_UseSimplygonAutomaticSizing UMETA(DisplayName = "Use Simplygon's automatic texture sizing"),
 	TextureSizingType_AutomaticFromTexelDensity UMETA(DisplayName = "Automatic - From Texel Density"),
 	TextureSizingType_AutomaticFromMeshScreenSize UMETA(DisplayName = "Automatic - From Mesh Screen Size"),
-	TextureSizingType_AutomaticFromMeshDrawDistance UMETA(DisplayName = "Automatic - From Mesh Draw Distance"),
+	TextureSizingType_AutomaticFromMeshDrawDistance UMETA(DisplayName = "Automatic - From Mesh Draw Distance", ToolTip = "When working with World Partition HLODs, the draw distance is automatically deduced from the runtime grid loading range."),
 	TextureSizingType_MAX,
 };
 
