Setting a Margin variable in a parent blueprint now correctly updates child blueprint variables as well.

ChildPropertyHandles are stored as TSharedRef, so there's no reason to convert them to an intermediate TWeakPtr.

diff --git a/Engine/Source/Editor/DetailCustomizations/Private/MarginCustomization.cpp b/Engine/Source/Editor/DetailCustomizations/Private/MarginCustomization.cpp
--- a/Engine/Source/Editor/DetailCustomizations/Private/MarginCustomization.cpp
+++ b/Engine/Source/Editor/DetailCustomizations/Private/MarginCustomization.cpp
@@ -27,14 +27,12 @@
 #include "Widgets/Input/SNumericEntryBox.h"
 #include "Widgets/SBoxPanel.h"
 
-class SWidget;
-
 TSharedRef<IPropertyTypeCustomization> FMarginStructCustomization::MakeInstance() 
 {
 	return MakeShareable( new FMarginStructCustomization() );
 }
 
-void FMarginStructCustomization::CustomizeHeader( TSharedRef<class IPropertyHandle> InStructPropertyHandle, class FDetailWidgetRow& HeaderRow, IPropertyTypeCustomizationUtils& StructCustomizationUtils )
+void FMarginStructCustomization::CustomizeHeader( TSharedRef<IPropertyHandle> InStructPropertyHandle, FDetailWidgetRow& HeaderRow, IPropertyTypeCustomizationUtils& StructCustomizationUtils )
 {
 	StructPropertyHandle = InStructPropertyHandle;
 
@@ -71,7 +69,7 @@ void FMarginStructCustomization::CustomizeHeader( TSharedRef<class IPropertyHand
 	];
 }
 
-void FMarginStructCustomization::CustomizeChildren( TSharedRef<class IPropertyHandle> InStructPropertyHandle, class IDetailChildrenBuilder& StructBuilder, IPropertyTypeCustomizationUtils& StructCustomizationUtils )
+void FMarginStructCustomization::CustomizeChildren( TSharedRef<IPropertyHandle> InStructPropertyHandle, IDetailChildrenBuilder& StructBuilder, IPropertyTypeCustomizationUtils& StructCustomizationUtils )
 {
 	const uint32 NumChildren = ChildPropertyHandles.Num();
 
@@ -204,7 +202,7 @@ void FMarginStructCustomization::OnMarginTextCommitted( const FText& InText, ETe
 			}
 
 
-			if( !bError )
+			if (!bError)
 			{
 				if (bIsMarginUsingUVSpace)
 				{
@@ -219,27 +217,12 @@ void FMarginStructCustomization::OnMarginTextCommitted( const FText& InText, ETe
 					}
 				}
 
-				if ( StructPropertyHandle.IsValid() && StructPropertyHandle->IsValidHandle() )
-				{
-					TArray<void*> RawData;
-					StructPropertyHandle->AccessRawData(RawData);
-
-					if ( RawData.ContainsByPredicate([] (void* Data) { return Data != nullptr; }) )
-					{
-						FScopedTransaction Transaction(FText::Format(NSLOCTEXT("FMarginStructCustomization", "SetMarginProperty", "Edit {0}"), StructPropertyHandle->GetPropertyDisplayName()));
-						StructPropertyHandle->NotifyPreChange();
-
-						for ( void* Data : RawData )
-						{
-							if ( Data )
-							{
-								*(FMargin*)Data = NewMargin;
-							}
-						}
-
-						StructPropertyHandle->NotifyPostChange(EPropertyChangeType::ValueSet);
-					}
-				}
+				FScopedTransaction Transaction(FText::Format(NSLOCTEXT("FMarginStructCustomization", "SetMarginProperty", "Edit {0}"), StructPropertyHandle->GetPropertyDisplayName()));
+				
+				ChildPropertyHandles[0]->SetValue(NewMargin.Left);
+				ChildPropertyHandles[1]->SetValue(NewMargin.Top);
+				ChildPropertyHandles[2]->SetValue(NewMargin.Right);
+				ChildPropertyHandles[3]->SetValue(NewMargin.Bottom);
 
 				MarginEditableTextBox->SetError( FString() );
 			}
@@ -297,9 +280,9 @@ FString FMarginStructCustomization::GetMarginTextFromProperties() const
 
 TOptional<float> FMarginStructCustomization::OnGetValue( int32 PropertyIndex ) const
 {
-	TWeakPtr<IPropertyHandle> WeakHandlePtr = ChildPropertyHandles[ PropertyIndex ];
+	TSharedRef<IPropertyHandle> ChildHandle = ChildPropertyHandles[ PropertyIndex ];
 	float FloatVal = 0.0f;
-	if( WeakHandlePtr.Pin()->GetValue( FloatVal ) == FPropertyAccess::Success )
+	if( ChildHandle->GetValue( FloatVal ) == FPropertyAccess::Success )
 	{
 		return TOptional<float>( FloatVal );
 	}
@@ -323,17 +306,17 @@ void FMarginStructCustomization::OnEndSliderMovement( float NewValue )
 
 void FMarginStructCustomization::OnValueCommitted( float NewValue, ETextCommit::Type CommitType, int32 PropertyIndex )
 {
-	TWeakPtr<IPropertyHandle> WeakHandlePtr = ChildPropertyHandles[ PropertyIndex ];
-	WeakHandlePtr.Pin()->SetValue( NewValue );
+	TSharedRef<IPropertyHandle> ChildHandle = ChildPropertyHandles[ PropertyIndex ];
+	ChildHandle->SetValue( NewValue );
 }	
 
 void FMarginStructCustomization::OnValueChanged( float NewValue, int32 PropertyIndex )
 {
 	if( bIsUsingSlider )
 	{
-		TWeakPtr<IPropertyHandle> WeakHandlePtr = ChildPropertyHandles[ PropertyIndex ];
+		TSharedRef<IPropertyHandle> ChildHandle = ChildPropertyHandles[ PropertyIndex ];
 		EPropertyValueSetFlags::Type Flags = EPropertyValueSetFlags::InteractiveChange;
-		WeakHandlePtr.Pin()->SetValue( NewValue, Flags );
+		ChildHandle->SetValue( NewValue, Flags );
 	}
 }
 
