World Partition - HLOD: Automate ScreenSize computation for "Simplifiy" HLOD method
* ScreenSize option is now deduced from the meshes bounds and the draw distance.

diff --git a/Engine/Plugins/Editor/WorldPartitionHLODUtilities/Source/Private/HLODBuilderMeshSimplify.cpp b/Engine/Plugins/Editor/WorldPartitionHLODUtilities/Source/Private/HLODBuilderMeshSimplify.cpp
--- a/Engine/Plugins/Editor/WorldPartitionHLODUtilities/Source/Private/HLODBuilderMeshSimplify.cpp
+++ b/Engine/Plugins/Editor/WorldPartitionHLODUtilities/Source/Private/HLODBuilderMeshSimplify.cpp
@@ -84,6 +84,34 @@ TArray<UActorComponent*> UHLODBuilderMeshSimplify::Build(const FHLODBuildContext
 		UseSettings.MaterialSettings.MeshMinDrawDistance = InHLODBuildContext.MinVisibleDistance;
 	}
 
+	// Generate a projection matrix.
+	static const float ScreenX = 1920;
+	static const float ScreenY = 1080;
+	static const float HalfFOVRad = FMath::DegreesToRadians(45.0f);
+	static const FMatrix ProjectionMatrix = FPerspectiveMatrix(HalfFOVRad, ScreenX, ScreenY, 0.01f);
+
+	// Gather bounds of the input components
+	auto GetComponentsBounds = [&]() -> FBoxSphereBounds
+	{
+		FBoxSphereBounds Bounds;
+		bool bFirst = true;
+
+		for (UActorComponent* Component : InSourceComponents)
+		{
+			if (USceneComponent* SceneComponent = Cast<USceneComponent>(Component))
+			{
+				FBoxSphereBounds ComponentBounds = SceneComponent->Bounds;
+				Bounds = bFirst ? ComponentBounds : Bounds + ComponentBounds;
+				bFirst = false;
+			}
+		}
+
+		return Bounds;
+	};
+
+	float ScreenSizePercent = ComputeBoundsScreenSize(FVector::ZeroVector, GetComponentsBounds().SphereRadius, FVector(0.0f, 0.0f, InHLODBuildContext.MinVisibleDistance), ProjectionMatrix);
+	UseSettings.ScreenSize = ScreenSizePercent * ScreenX;
+
 	TArray<UObject*> Assets;
 	FCreateProxyDelegate ProxyDelegate;
 	ProxyDelegate.BindLambda([&Assets](const FGuid Guid, TArray<UObject*>& InAssetsCreated) { Assets = InAssetsCreated; });
diff --git a/Engine/Source/Editor/DetailCustomizations/Private/MeshProxySettingsCustomizations.cpp b/Engine/Source/Editor/DetailCustomizations/Private/MeshProxySettingsCustomizations.cpp
--- a/Engine/Source/Editor/DetailCustomizations/Private/MeshProxySettingsCustomizations.cpp
+++ b/Engine/Source/Editor/DetailCustomizations/Private/MeshProxySettingsCustomizations.cpp
@@ -2,6 +2,7 @@
 
 #include "MeshProxySettingsCustomizations.h"
 
+#include "Algo/AnyOf.h"
 #include "Containers/Map.h"
 #include "Containers/UnrealString.h"
 #include "Delegates/Delegate.h"
@@ -20,6 +21,7 @@
 #include "UObject/NameTypes.h"
 #include "UObject/UnrealNames.h"
 #include "UObject/UnrealType.h"
+#include "WorldPartition/HLOD/HLODLayer.h"
 
 #define LOCTEXT_NAMESPACE "MeshProxySettingsCustomizations"
 
@@ -54,6 +56,10 @@ void FMeshProxySettingsCustomizations::CustomizeChildren(TSharedRef<IPropertyHan
 		PropertyHandles.Add(PropertyName, ChildHandle);
 	}
 
+	TArray<UObject*> OutersList;
+	StructPropertyHandle->GetOuterObjects(OutersList);
+	bIsEditingHLODLayer = Algo::AnyOf(OutersList, [](UObject* Outer) { return Outer->IsInA(UHLODLayer::StaticClass()); });
+	
 	// Determine if we are using our native module  If so, we will supress some of the options used by the current thirdparty tool (simplygon).
 
 	IMeshReductionManagerModule& ModuleManager = FModuleManager::Get().LoadModuleChecked<IMeshReductionManagerModule>("MeshReductionInterface");
@@ -90,6 +96,7 @@ void FMeshProxySettingsCustomizations::CustomizeChildren(TSharedRef<IPropertyHan
 	TSharedPtr< IPropertyHandle > MergeDistanceHandle                     = PropertyHandles.FindChecked(GET_MEMBER_NAME_CHECKED(FMeshProxySettings, MergeDistance));
 	TSharedPtr< IPropertyHandle > UnresolvedGeometryColorHandle           = PropertyHandles.FindChecked(GET_MEMBER_NAME_CHECKED(FMeshProxySettings, UnresolvedGeometryColor));
 	TSharedPtr< IPropertyHandle > VoxelSizeHandle                         = PropertyHandles.FindChecked(GET_MEMBER_NAME_CHECKED(FMeshProxySettings, VoxelSize));
+	TSharedPtr< IPropertyHandle > ScreenSizeHandle						  = PropertyHandles.FindChecked(GET_MEMBER_NAME_CHECKED(FMeshProxySettings, ScreenSize));
 
 	for (auto Iter(PropertyHandles.CreateConstIterator()); Iter; ++Iter)
 	{
@@ -160,6 +167,13 @@ void FMeshProxySettingsCustomizations::CustomizeChildren(TSharedRef<IPropertyHan
 			
 			MeshProxySettingsRow.Visibility(TAttribute<EVisibility>(this, &FMeshProxySettingsCustomizations::IsVoxelSizeVisible));
 		}
+		else if (Iter.Value() == ScreenSizeHandle)
+		{
+			IDetailPropertyRow& MeshProxySettingsRow = AddPropertyToGroup(Iter.Value().ToSharedRef());
+			AddResetToDefaultOverrides(MeshProxySettingsRow);
+
+			MeshProxySettingsRow.Visibility(TAttribute<EVisibility>(this, &FMeshProxySettingsCustomizations::IsScreenSizeVisible));
+		}
 		else
 		{
 			IDetailPropertyRow& SettingsRow = AddPropertyToGroup(Iter.Value().ToSharedRef());
@@ -238,6 +252,11 @@ EVisibility FMeshProxySettingsCustomizations::IsVoxelSizeVisible() const
 {
 	return IsProxyLODSpecificVisible();
 }
+EVisibility FMeshProxySettingsCustomizations::IsScreenSizeVisible() const
+{
+	return bIsEditingHLODLayer ? EVisibility::Hidden : IsProxyLODSpecificVisible();
+}
+
 
 
 
diff --git a/Engine/Source/Editor/DetailCustomizations/Private/MeshProxySettingsCustomizations.h b/Engine/Source/Editor/DetailCustomizations/Private/MeshProxySettingsCustomizations.h
--- a/Engine/Source/Editor/DetailCustomizations/Private/MeshProxySettingsCustomizations.h
+++ b/Engine/Source/Editor/DetailCustomizations/Private/MeshProxySettingsCustomizations.h
@@ -29,10 +29,13 @@ protected:
 	EVisibility IsMergeDistanceVisible() const;
 	EVisibility IsUnresolvedGeometryColorVisible() const;
 	EVisibility IsVoxelSizeVisible() const;
+	EVisibility IsScreenSizeVisible() const;	
 	EVisibility IsNormalCalcMethodVisible() const;
 	EVisibility IsSearchDistanceVisible() const;
 
 	EVisibility IsThirdPartySpecificVisible() const;
 	EVisibility IsProxyLODSpecificVisible() const;
 	bool UseNativeProxyLODTool() const;
+
+	bool bIsEditingHLODLayer;
 };
