- Fixing perPlatform property customization to list all platforms, not just the ones that TargetPlatforms created
- Added PlatformGroup to WinGDK

diff --git a/Engine/Source/Editor/DetailCustomizations/Private/PerPlatformPropertyCustomization.cpp b/Engine/Source/Editor/DetailCustomizations/Private/PerPlatformPropertyCustomization.cpp
--- a/Engine/Source/Editor/DetailCustomizations/Private/PerPlatformPropertyCustomization.cpp
+++ b/Engine/Source/Editor/DetailCustomizations/Private/PerPlatformPropertyCustomization.cpp
@@ -16,7 +16,6 @@
 #include "Misc/AssertionMacros.h"
 #include "Misc/DataDrivenPlatformInfoRegistry.h"
 #include "PerPlatformProperties.h"
-#include "PlatformInfo.h"
 #include "PropertyEditorModule.h"
 #include "PropertyHandle.h"
 #include "SPerPlatformPropertiesWidget.h"
@@ -47,8 +46,6 @@ void FPerPlatformPropertyCustomization<PerPlatformType>::CustomizeChildren(TShar
 {
 	PropertyUtilities = StructCustomizationUtils.GetPropertyUtilities();
 
-	int32 PlatformNumber = PlatformInfo::GetAllPlatformGroupNames().Num();
-
 	TAttribute<TArray<FName>> PlatformOverrideNames = TAttribute<TArray<FName>>::Create(TAttribute<TArray<FName>>::FGetter::CreateSP(this, &FPerPlatformPropertyCustomization<PerPlatformType>::GetPlatformOverrideNames, StructPropertyHandle));
 
 	FPerPlatformPropertyCustomNodeBuilderArgs Args;
@@ -261,9 +258,9 @@ void FPerPlatformPropertyCustomNodeBuilder::GenerateHeaderRowContent(FDetailWidg
 	// Build Platform menu
 	FMenuBuilder AddPlatformMenuBuilder(true, nullptr, nullptr, true);
 
+	const TArray<const FDataDrivenPlatformInfo*>& SortedPlatforms = FDataDrivenPlatformInfoRegistry::GetSortedPlatformInfos(EPlatformInfoType::TruePlatformsOnly);
 	// Platform (group) names
-	const TArray<FName>& PlatformGroupNameArray = PlatformInfo::GetAllPlatformGroupNames();
-	const TArray<FName>& VanillaPlatformNameArray = PlatformInfo::GetAllVanillaPlatformNames();
+//	const TArray<FName>& PlatformGroupNameArray = PlatformInfo::GetAllPlatformGroupNames();
 
 	// Sanitized platform names
 	TArray<FName> BasePlatformNameArray;
@@ -271,15 +268,17 @@ void FPerPlatformPropertyCustomNodeBuilder::GenerateHeaderRowContent(FDetailWidg
 	TMultiMap<FName, FName> GroupToPlatform;
 
 	TArray<FName> PlatformOverrides = Args.PlatformOverrideNames.Get();
+	TArray<FName> PlatformGroupNameArray;
 
 	// Create mapping from platform to platform groups and remove postfixes and invalid platform names
-	for (const FName& PlatformName : VanillaPlatformNameArray)
+	for (const FDataDrivenPlatformInfo* DDPI : SortedPlatforms)
 	{
 		// Add platform name if it isn't already set, and also add to group mapping
-		if (!PlatformOverrides.Contains(PlatformName))
+		if (!PlatformOverrides.Contains(DDPI->IniPlatformName))
 		{
-			BasePlatformNameArray.AddUnique(PlatformName);
-			GroupToPlatform.AddUnique(PlatformInfo::FindPlatformInfo(PlatformName)->DataDrivenPlatformInfo->PlatformGroupName, PlatformName);
+			BasePlatformNameArray.AddUnique(DDPI->IniPlatformName);
+			GroupToPlatform.AddUnique(DDPI->PlatformGroupName, DDPI->IniPlatformName);
+			PlatformGroupNameArray.AddUnique(DDPI->PlatformGroupName);
 		}
 	}
 
@@ -303,6 +302,8 @@ void FPerPlatformPropertyCustomNodeBuilder::GenerateHeaderRowContent(FDetailWidg
 
 		TArray<FName> PlatformNames;
 		GroupToPlatform.MultiFind(GroupName, PlatformNames);
+		// these come out reversed for whatever MultiFind reason, even tho they went in sorted
+		Algo::Reverse(PlatformNames);
 
 		const FTextFormat Format = NSLOCTEXT("SPerPlatformPropertiesWidget", "AddOverrideFor", "Add Override specifically for {0}");
 		for (const FName& PlatformName : PlatformNames)
