Added slow task when building Landscape Nanite meshes through the details panel button

#rnx

diff --git a/Engine/Source/Editor/DetailCustomizations/Private/LandscapeUIDetails.cpp b/Engine/Source/Editor/DetailCustomizations/Private/LandscapeUIDetails.cpp
--- a/Engine/Source/Editor/DetailCustomizations/Private/LandscapeUIDetails.cpp
+++ b/Engine/Source/Editor/DetailCustomizations/Private/LandscapeUIDetails.cpp
@@ -35,6 +35,7 @@
 #include "Widgets/DeclarativeSyntaxSupport.h"
 #include "Widgets/Input/SButton.h"
 #include "Widgets/Input/SCheckBox.h"
+#include "Misc/ScopedSlowTask.h"
 
 class UObject;
 
@@ -124,8 +125,12 @@ void FLandscapeUIDetails::CustomizeDetails( IDetailLayoutBuilder& DetailBuilder
 						ULandscapeInfo* LandscapeInfo = Landscape->GetLandscapeInfo();
 						if (LandscapeInfo != nullptr)
 						{
+							FScopedSlowTask SlowTask(LandscapeInfo->StreamingProxies.Num(), (LOCTEXT("RebuildNaniteSlowTask", "Rebuilding Nanite Landscape Meshes")));
+							SlowTask.MakeDialog();
+
 							for (TWeakObjectPtr<ALandscapeProxy> ProxyPtr : LandscapeInfo->StreamingProxies)
 							{
+								SlowTask.EnterProgressFrame(1, FText::Format(LOCTEXT("RebuildNaniteSlowTaskProgress", "Building Nanite Landscape Mesh {0} of {1})"), FText::AsNumber(SlowTask.CompletedWork), FText::AsNumber(SlowTask.TotalAmountOfWork)));
 								if (ALandscapeProxy* Proxy = ProxyPtr.Get())
 								{
 									Proxy->UpdateNaniteRepresentation();
