Water: Fix spline point properties not showing in Restrictive Mode

diff --git a/Engine/Source/Editor/DetailCustomizations/Private/SplineComponentDetails.cpp b/Engine/Source/Editor/DetailCustomizations/Private/SplineComponentDetails.cpp
--- a/Engine/Source/Editor/DetailCustomizations/Private/SplineComponentDetails.cpp
+++ b/Engine/Source/Editor/DetailCustomizations/Private/SplineComponentDetails.cpp
@@ -480,6 +480,7 @@ void FSplinePointDetails::GenerateChildContent(IDetailChildrenBuilder& ChildrenB
 
 	// Message which is shown when no points are selected
 	ChildrenBuilder.AddCustomRow(LOCTEXT("NoneSelected", "None selected"))
+		.RowTag(TEXT("NoneSelected"))
 		.Visibility(TAttribute<EVisibility>(this, &FSplinePointDetails::IsDisabled))
 		[
 			SNew(SBox)
@@ -499,6 +500,7 @@ void FSplinePointDetails::GenerateChildContent(IDetailChildrenBuilder& ChildrenB
 
 	// Input key
 	ChildrenBuilder.AddCustomRow(LOCTEXT("InputKey", "Input Key"))
+		.RowTag(TEXT("InputKey"))
 		.Visibility(TAttribute<EVisibility>(this, &FSplinePointDetails::IsEnabled))
 		.NameContent()
 		.HAlign(HAlign_Left)
@@ -524,6 +526,7 @@ void FSplinePointDetails::GenerateChildContent(IDetailChildrenBuilder& ChildrenB
 	if (SplineComp->AllowsSpinePointLocationEditing())
 	{
 		ChildrenBuilder.AddCustomRow(LOCTEXT("Location", "Location"))
+			.RowTag(TEXT("Location"))
 			.CopyAction(CreateCopyAction(ESplinePointProperty::Location))
 			.PasteAction(CreatePasteAction(ESplinePointProperty::Location))
 			.Visibility(TAttribute<EVisibility>(this, &FSplinePointDetails::IsEnabled))
@@ -560,6 +563,7 @@ void FSplinePointDetails::GenerateChildContent(IDetailChildrenBuilder& ChildrenB
 	if (SplineComp->AllowsSplinePointRotationEditing())
 	{
 		ChildrenBuilder.AddCustomRow(LOCTEXT("Rotation", "Rotation"))
+			.RowTag(TEXT("Rotation"))
 			.CopyAction(CreateCopyAction(ESplinePointProperty::Rotation))
 			.PasteAction(CreatePasteAction(ESplinePointProperty::Rotation))
 			.Visibility(TAttribute<EVisibility>(this, &FSplinePointDetails::IsEnabled))
@@ -590,6 +594,7 @@ void FSplinePointDetails::GenerateChildContent(IDetailChildrenBuilder& ChildrenB
 	if (SplineComp->AllowsSplinePointScaleEditing())
 	{
 		ChildrenBuilder.AddCustomRow(LOCTEXT("Scale", "Scale"))
+			.RowTag(TEXT("Scale"))
 			.Visibility(TAttribute<EVisibility>(this, &FSplinePointDetails::IsEnabled))
 			.CopyAction(CreateCopyAction(ESplinePointProperty::Scale))
 			.PasteAction(CreatePasteAction(ESplinePointProperty::Scale))
@@ -627,6 +632,7 @@ void FSplinePointDetails::GenerateChildContent(IDetailChildrenBuilder& ChildrenB
 	if (SplineComp->AllowsSplinePointArriveTangentEditing())
 	{
 		ChildrenBuilder.AddCustomRow(LOCTEXT("ArriveTangent", "Arrive Tangent"))
+			.RowTag(TEXT("ArriveTangent"))
 			.Visibility(TAttribute<EVisibility>(this, &FSplinePointDetails::IsEnabled))
 			.CopyAction(CreateCopyAction(ESplinePointProperty::ArriveTangent))
 			.PasteAction(CreatePasteAction(ESplinePointProperty::ArriveTangent))
@@ -660,6 +666,7 @@ void FSplinePointDetails::GenerateChildContent(IDetailChildrenBuilder& ChildrenB
 	if (SplineComp->AllowsSplinePointLeaveTangentEditing())
 	{
 		ChildrenBuilder.AddCustomRow(LOCTEXT("LeaveTangent", "Leave Tangent"))
+			.RowTag(TEXT("LeaveTangent"))
 			.Visibility(TAttribute<EVisibility>(this, &FSplinePointDetails::IsEnabled))
 			.CopyAction(CreateCopyAction(ESplinePointProperty::LeaveTangent))
 			.PasteAction(CreatePasteAction(ESplinePointProperty::LeaveTangent))
@@ -692,6 +699,7 @@ void FSplinePointDetails::GenerateChildContent(IDetailChildrenBuilder& ChildrenB
 	if (SplineComp->GetEnabledSplinePointTypes().Num() > 1)
 	{
 		ChildrenBuilder.AddCustomRow(LOCTEXT("Type", "Type"))
+			.RowTag(TEXT("Type"))
 			.Visibility(TAttribute<EVisibility>(this, &FSplinePointDetails::IsEnabled))
 			.NameContent()
 			.HAlign(HAlign_Left)
