MaxLength fixes

- Updated text properties to respect the MaxLength meta-data
- Updated name properties to respect the MaxLength meta-data if it's smaller than the max length of an FName

#rnx

diff --git a/Engine/Source/Editor/DetailCustomizations/Private/TextCustomization.cpp b/Engine/Source/Editor/DetailCustomizations/Private/TextCustomization.cpp
--- a/Engine/Source/Editor/DetailCustomizations/Private/TextCustomization.cpp
+++ b/Engine/Source/Editor/DetailCustomizations/Private/TextCustomization.cpp
@@ -18,16 +18,20 @@ namespace
 			: PropertyHandle(InPropertyHandle)
 			, PropertyUtilities(InPropertyUtilities)
 		{
+			static const FName NAME_MaxLength = "MaxLength";
+			MaxLength = PropertyHandle->IsValidHandle() ? PropertyHandle->GetIntMetaData(NAME_MaxLength) : 0;
 		}
 
 		virtual bool IsMultiLineText() const override
 		{
-			return PropertyHandle->IsValidHandle() && PropertyHandle->GetBoolMetaData("MultiLine");
+			static const FName NAME_MultiLine = "MultiLine";
+			return PropertyHandle->IsValidHandle() && PropertyHandle->GetBoolMetaData(NAME_MultiLine);
 		}
 
 		virtual bool IsPassword() const override
 		{
-			return PropertyHandle->IsValidHandle() && PropertyHandle->GetBoolMetaData("PasswordField");
+			static const FName NAME_PasswordField = "PasswordField";
+			return PropertyHandle->IsValidHandle() && PropertyHandle->GetBoolMetaData(NAME_PasswordField);
 		}
 
 		virtual bool IsReadOnly() const override
@@ -84,6 +88,12 @@ namespace
 
 		virtual bool IsValidText(const FText& InText, FText& OutErrorMsg) const override
 		{
+			if (MaxLength > 0 && InText.ToString().Len() > MaxLength)
+			{
+				OutErrorMsg = FText::Format(NSLOCTEXT("PropertyEditor", "PropertyTextTooLongError", "This value is too long ({0}/{1} characters)"), InText.ToString().Len(), MaxLength);
+				return false;
+			}
+
 			return true;
 		}
 
@@ -105,6 +115,9 @@ namespace
 	private:
 		TSharedRef<IPropertyHandle> PropertyHandle;
 		TSharedPtr<IPropertyUtilities> PropertyUtilities;
+
+		/** The maximum length of the value that can be edited, or <=0 for unlimited */
+		int32 MaxLength = 0;
 	};
 }
 
diff --git a/Engine/Source/Editor/PropertyEditor/Private/UserInterface/PropertyEditor/SPropertyEditorText.cpp b/Engine/Source/Editor/PropertyEditor/Private/UserInterface/PropertyEditor/SPropertyEditorText.cpp
--- a/Engine/Source/Editor/PropertyEditor/Private/UserInterface/PropertyEditor/SPropertyEditorText.cpp
+++ b/Engine/Source/Editor/PropertyEditor/Private/UserInterface/PropertyEditor/SPropertyEditorText.cpp
@@ -37,7 +37,12 @@ void SPropertyEditorText::Construct( const FArguments& InArgs, const TSharedRef<
 	static const FName NAME_PasswordField = "PasswordField";
 
 	bIsMultiLine = InPropertyEditor->GetPropertyHandle()->GetBoolMetaData(NAME_MultiLine);
-	MaxLength = InPropertyEditor->PropertyIsA(FNameProperty::StaticClass()) ? NAME_SIZE - 1 : InPropertyEditor->GetPropertyHandle()->GetIntMetaData(NAME_MaxLength);
+
+	MaxLength = InPropertyEditor->GetPropertyHandle()->GetIntMetaData(NAME_MaxLength);
+	if (InPropertyEditor->PropertyIsA(FNameProperty::StaticClass()))
+	{
+		MaxLength = MaxLength <= 0 ? NAME_SIZE - 1 : FMath::Min(MaxLength, NAME_SIZE - 1);
+	}
 
 	const bool bIsPassword = InPropertyEditor->GetPropertyHandle()->GetBoolMetaData(NAME_PasswordField);
 	
