Make CallInEditor UFunction only use the default tooltip text ("Call an event on the selected object(s)") if there is no custom tooltip text written

diff --git a/Engine/Source/Editor/DetailCustomizations/Private/ObjectDetails.cpp b/Engine/Source/Editor/DetailCustomizations/Private/ObjectDetails.cpp
--- a/Engine/Source/Editor/DetailCustomizations/Private/ObjectDetails.cpp
+++ b/Engine/Source/Editor/DetailCustomizations/Private/ObjectDetails.cpp
@@ -200,7 +200,7 @@ void FObjectDetails::AddCallInEditorMethods(IDetailLayoutBuilder& DetailBuilder)
 				SNew(SButton)
 				.Text(ButtonCaption)
 				.OnClicked(FOnClicked::CreateSP(this, &FObjectDetails::OnExecuteCallInEditorFunction, WeakFunctionPtr))
-				.ToolTipText(FText::Format(LOCTEXT("CallInEditorTooltip", "Call an event on the selected object(s)\n\n\n{0}"), FunctionTooltip))
+				.ToolTipText(FunctionTooltip.IsEmptyOrWhitespace() ? LOCTEXT("CallInEditorTooltip", "Call an event on the selected object(s)") : FunctionTooltip)
 			];
 
 			CategoryEntry.RowTag = Function->GetFName();
