Fix debug cvar name for landscape grass detail

diff --git a/Engine/Source/Editor/DetailCustomizations/Private/LandscapeGrassTypeDetails.cpp b/Engine/Source/Editor/DetailCustomizations/Private/LandscapeGrassTypeDetails.cpp
--- a/Engine/Source/Editor/DetailCustomizations/Private/LandscapeGrassTypeDetails.cpp
+++ b/Engine/Source/Editor/DetailCustomizations/Private/LandscapeGrassTypeDetails.cpp
@@ -22,7 +22,7 @@
 class UObject;
 
 static bool GShowBothPerQualityAndPerPlaformProperties = false;
-static FAutoConsoleVariableRef CVarMinFramesToKeepGrass(
+static FAutoConsoleVariableRef CVarShowBothPerQualityAndPerPlaform(
 	TEXT("r.grass.ShowBothPerQualityAndPerPlaformProperties"),
 	GShowBothPerQualityAndPerPlaformProperties,
 	TEXT("Show both per platform and per quality properties in the editor."));
