Fixed typos

#rnx

diff --git a/Engine/Source/Editor/DetailCustomizations/Private/LandscapeUIDetails.cpp b/Engine/Source/Editor/DetailCustomizations/Private/LandscapeUIDetails.cpp
--- a/Engine/Source/Editor/DetailCustomizations/Private/LandscapeUIDetails.cpp
+++ b/Engine/Source/Editor/DetailCustomizations/Private/LandscapeUIDetails.cpp
@@ -130,7 +130,7 @@ void FLandscapeUIDetails::CustomizeDetails( IDetailLayoutBuilder& DetailBuilder
 
 							for (TWeakObjectPtr<ALandscapeProxy> ProxyPtr : LandscapeInfo->StreamingProxies)
 							{
-								SlowTask.EnterProgressFrame(1, FText::Format(LOCTEXT("RebuildNaniteSlowTaskProgress", "Building Nanite Landscape Mesh {0} of {1})"), FText::AsNumber(SlowTask.CompletedWork), FText::AsNumber(SlowTask.TotalAmountOfWork)));
+								SlowTask.EnterProgressFrame(1, FText::Format(LOCTEXT("RebuildNaniteSlowTaskProgress", "Building Nanite Landscape Mesh ({0} of {1})"), FText::AsNumber(SlowTask.CompletedWork), FText::AsNumber(SlowTask.TotalAmountOfWork)));
 								if (ALandscapeProxy* Proxy = ProxyPtr.Get())
 								{
 									Proxy->UpdateNaniteRepresentation();
diff --git a/Engine/Source/Runtime/Landscape/Private/LandscapeSubsystem.cpp b/Engine/Source/Runtime/Landscape/Private/LandscapeSubsystem.cpp
--- a/Engine/Source/Runtime/Landscape/Private/LandscapeSubsystem.cpp
+++ b/Engine/Source/Runtime/Landscape/Private/LandscapeSubsystem.cpp
@@ -239,7 +239,7 @@ void ULandscapeSubsystem::BuildNanite()
 
 	for (TWeakObjectPtr<ALandscapeProxy> ProxyPtr : Proxies)
 	{
-		SlowTask.EnterProgressFrame(1, FText::Format(LOCTEXT("Landscape_BuildNaniteProgress", "Building Nanite Landscape Mesh {0} of {1})"), FText::AsNumber(SlowTask.CompletedWork), FText::AsNumber(SlowTask.TotalAmountOfWork)));
+		SlowTask.EnterProgressFrame(1, FText::Format(LOCTEXT("Landscape_BuildNaniteProgress", "Building Nanite Landscape Mesh ({0} of {1})"), FText::AsNumber(SlowTask.CompletedWork), FText::AsNumber(SlowTask.TotalAmountOfWork)));
 		if (ALandscapeProxy* Proxy = ProxyPtr.Get())
 		{
 			Proxy->UpdateNaniteRepresentation();
