Ranges now use IPropertyHandle::GetMetaData() instead of FProperty::GetMetaData(), allowing them to function correctly when placed in collections.

diff --git a/Engine/Source/Editor/DetailCustomizations/Private/RangeStructCustomization.cpp b/Engine/Source/Editor/DetailCustomizations/Private/RangeStructCustomization.cpp
--- a/Engine/Source/Editor/DetailCustomizations/Private/RangeStructCustomization.cpp
+++ b/Engine/Source/Editor/DetailCustomizations/Private/RangeStructCustomization.cpp
@@ -52,7 +52,7 @@ namespace
 	template <>
 	struct FGetMetaDataHelper<float>
 	{
-		static float GetMetaData(const FProperty* Property, const TCHAR* Key)
+		static float GetMetaData(const TSharedRef<IPropertyHandle>& Property, const TCHAR* Key)
 		{
 			return Property->GetFloatMetaData(Key);
 		}
@@ -61,7 +61,7 @@ namespace
 	template <>
 	struct FGetMetaDataHelper<double>
 	{
-		static double GetMetaData(const FProperty* Property, const TCHAR* Key)
+		static double GetMetaData(const TSharedRef<IPropertyHandle>& Property, const TCHAR* Key)
 		{
 			return Property->GetDoubleMetaData(Key);
 		}
@@ -70,7 +70,7 @@ namespace
 	template <>
 	struct FGetMetaDataHelper<int32>
 	{
-		static int32 GetMetaData(const FProperty* Property, const TCHAR* Key)
+		static int32 GetMetaData(const TSharedRef<IPropertyHandle>& Property, const TCHAR* Key)
 		{
 			return Property->GetIntMetaData(Key);
 		}
@@ -110,18 +110,14 @@ void FRangeStructCustomization<NumericType>::CustomizeHeader(TSharedRef<IPropert
 	check(UpperBoundTypeHandle.IsValid());
 
 	// Get min/max metadata values if defined
-	auto Property = StructPropertyHandle->GetProperty();
-	if (Property != nullptr)
+	if (StructPropertyHandle->HasMetaData(TEXT("UIMin")))
 	{
-		if (Property->HasMetaData(TEXT("UIMin")))
-		{
-			MinAllowedValue = TOptional<NumericType>(FGetMetaDataHelper<NumericType>::GetMetaData(Property, TEXT("UIMin")));
-		}
+		MinAllowedValue = TOptional<NumericType>(FGetMetaDataHelper<NumericType>::GetMetaData(StructPropertyHandle, TEXT("UIMin")));
+	}
 
-		if (Property->HasMetaData(TEXT("UIMax")))
-		{
-			MaxAllowedValue = TOptional<NumericType>(FGetMetaDataHelper<NumericType>::GetMetaData(Property, TEXT("UIMax")));
-		}
+	if (StructPropertyHandle->HasMetaData(TEXT("UIMax")))
+	{
+		MaxAllowedValue = TOptional<NumericType>(FGetMetaDataHelper<NumericType>::GetMetaData(StructPropertyHandle, TEXT("UIMax")));
 	}
 
 	// Make weak pointers to be passed as payloads to the widgets
