Slate/ PropertyEditor:

-Adjust spacing of light channel buttons
-Static Lights Gray Out Light Channels (Details and Light Mixer)

Object Mixer:
-Implement list of property names that when changed require a list refresh

diff --git a/Engine/Plugins/Editor/ObjectMixer/ObjectMixer/Source/ObjectMixer/Private/ObjectMixerEditorModule.cpp b/Engine/Plugins/Editor/ObjectMixer/ObjectMixer/Source/ObjectMixer/Private/ObjectMixerEditorModule.cpp
--- a/Engine/Plugins/Editor/ObjectMixer/ObjectMixer/Source/ObjectMixer/Private/ObjectMixerEditorModule.cpp
+++ b/Engine/Plugins/Editor/ObjectMixer/ObjectMixer/Source/ObjectMixer/Private/ObjectMixerEditorModule.cpp
@@ -53,6 +53,11 @@ UWorld* FObjectMixerEditorModule::GetWorld()
 
 void FObjectMixerEditorModule::Initialize()
 {
+	PropertiesThatRequireRefresh =
+		{
+			"LightingChannels"
+		};
+	
 	BindDelegates();
 	
 	SetupMenuItemVariables();
@@ -352,10 +357,25 @@ void FObjectMixerEditorModule::BindDelegates()
 	}));
 
 	DelegateHandles.Add(FCoreUObjectDelegates::OnObjectTransacted.AddLambda([this](UObject*, const FTransactionObjectEvent& Event)
-	{
-		if (Event.GetEventType() == ETransactionObjectEventType::Finalized && Event.HasNameChange())
+	{		
+		if (Event.GetEventType() == ETransactionObjectEventType::Finalized)
 		{
-			RequestRebuildList();
+			if (Event.HasNameChange())
+			{
+				RequestRebuildList();
+			}
+			else
+			{
+				const TSet<FName> ChangedPropertyNames = TSet<FName>(Event.GetChangedProperties());
+				const bool bPropertyChangedThatRequiresRefresh =
+					ChangedPropertyNames.Difference(PropertiesThatRequireRefresh).Num() > 0 ||
+						PropertiesThatRequireRefresh.Difference(ChangedPropertyNames).Num() > 0;
+
+				if (bPropertyChangedThatRequiresRefresh)
+				{
+					RequestRebuildList();
+				}
+			}
 		}
 	}));
 
diff --git a/Engine/Plugins/Editor/ObjectMixer/ObjectMixer/Source/ObjectMixer/Public/ObjectMixerEditorModule.h b/Engine/Plugins/Editor/ObjectMixer/ObjectMixer/Source/ObjectMixer/Public/ObjectMixerEditorModule.h
--- a/Engine/Plugins/Editor/ObjectMixer/ObjectMixer/Source/ObjectMixer/Public/ObjectMixerEditorModule.h
+++ b/Engine/Plugins/Editor/ObjectMixer/ObjectMixer/Source/ObjectMixer/Public/ObjectMixerEditorModule.h
@@ -91,6 +91,9 @@ protected:
 	
 	TSet<FDelegateHandle> DelegateHandles;
 
+	/** If a property is changed that has a name found in this set, the panel will be refreshed. */
+	TSet<FName> PropertiesThatRequireRefresh;
+
 private:
 
 	TSharedPtr<FWorkspaceItem> WorkspaceGroup;
diff --git a/Engine/Source/Editor/DetailCustomizations/Private/LightingChannelsCustomization.cpp b/Engine/Source/Editor/DetailCustomizations/Private/LightingChannelsCustomization.cpp
--- a/Engine/Source/Editor/DetailCustomizations/Private/LightingChannelsCustomization.cpp
+++ b/Engine/Source/Editor/DetailCustomizations/Private/LightingChannelsCustomization.cpp
@@ -51,16 +51,19 @@ void FLightingChannelsCustomization::CustomizeHeader(TSharedRef<IPropertyHandle>
 	{
 		const FText NumericText = FText::AsNumber(ChildIndex);
 		const FText SlotTooltipText = FText::Format(LOCTEXT("LightingChannelToggleFormat", "Toggle Lighting Channel {0}"), NumericText);
+
+		const bool bIsLastChild = ChildIndex == ChildCount - 1;
 		
 		ButtonOptionsPanel->AddSlot()
 		.HAlign(HAlign_Left)
 		.VAlign(VAlign_Center)
 		.AutoWidth()
-		.Padding(0, 0, ChildIndex == ChildCount -1 ? 16 : 0, 0)
+		.Padding(0, 0, bIsLastChild ? 0 : 8, 0)
 		[
 			SNew(SBox)
 			.WidthOverride(20)
 			.HAlign(HAlign_Fill)
+			.IsEnabled(this, &FLightingChannelsCustomization::IsLightingChannelButtonEditable, ChildIndex)
 			[
 				SNew(SCheckBox)
 				.Style(&FAppStyle::Get().GetWidgetStyle<FCheckBoxStyle>("DetailsView.ChannelToggleButton"))
@@ -117,6 +120,18 @@ FText FLightingChannelsCustomization::GetStructPropertyTooltipText() const
 	return FText::GetEmpty();
 }
 
+bool FLightingChannelsCustomization::IsLightingChannelButtonEditable(uint32 ChildIndex) const
+{
+	if (LightingChannelsHandle.IsValid())
+	{
+		if (TSharedPtr<IPropertyHandle> ChildHandle = LightingChannelsHandle->GetChildHandle(ChildIndex))
+		{
+			return ChildHandle->IsEditable() && !ChildHandle->IsEditConst();
+		}
+	}
+	return false;
+}
+
 void FLightingChannelsCustomization::OnButtonCheckedStateChanged(ECheckBoxState NewState, uint32 ChildIndex) const
 {
 	if (LightingChannelsHandle.IsValid())
diff --git a/Engine/Source/Editor/DetailCustomizations/Private/LightingChannelsCustomization.h b/Engine/Source/Editor/DetailCustomizations/Private/LightingChannelsCustomization.h
--- a/Engine/Source/Editor/DetailCustomizations/Private/LightingChannelsCustomization.h
+++ b/Engine/Source/Editor/DetailCustomizations/Private/LightingChannelsCustomization.h
@@ -23,7 +23,8 @@ private:
 
 	FText GetStructPropertyNameText() const;
 	FText GetStructPropertyTooltipText() const;
-	
+
+	bool IsLightingChannelButtonEditable(uint32 ChildIndex) const;
 	void OnButtonCheckedStateChanged(ECheckBoxState NewState, uint32 ChildIndex) const;
 	ECheckBoxState GetButtonCheckedState(uint32 ChildIndex) const;
 
