[UE-82098] Shared Transitions are not usable by separate state machines within the same ABP

diff --git a/Engine/Source/Editor/AnimGraph/Private/AnimStateTransitionNode.cpp b/Engine/Source/Editor/AnimGraph/Private/AnimStateTransitionNode.cpp
--- a/Engine/Source/Editor/AnimGraph/Private/AnimStateTransitionNode.cpp
+++ b/Engine/Source/Editor/AnimGraph/Private/AnimStateTransitionNode.cpp
@@ -656,14 +656,18 @@ bool UAnimStateTransitionNode::IsBoundGraphShared() const
 {
 	if (BoundGraph)
 	{
-		//@TODO: O(N) search
-		UEdGraph* ParentGraph = GetGraph();
-		for (int32 NodeIdx = 0; NodeIdx < ParentGraph->Nodes.Num(); NodeIdx++)
+		if (UBlueprint* Blueprint = FBlueprintEditorUtils::FindBlueprintForGraph(GetGraph()))
 		{
-			UAnimStateNodeBase* AnimNode = Cast<UAnimStateNodeBase>(ParentGraph->Nodes[NodeIdx]);
-			if ((AnimNode != NULL) && (AnimNode != this) && (AnimNode->GetBoundGraph() == BoundGraph))
+			TArray<UAnimStateNodeBase*> StateNodes;
+			FBlueprintEditorUtils::GetAllNodesOfClassEx<UAnimStateNodeBase>(Blueprint, StateNodes);
+
+			for (int32 NodeIdx = 0; NodeIdx < StateNodes.Num(); NodeIdx++)
 			{
-				return true;
+				UAnimStateNodeBase* AnimNode = Cast<UAnimStateNodeBase>(StateNodes[NodeIdx]);
+				if ((AnimNode != NULL) && (AnimNode != this) && (AnimNode->GetBoundGraph() == BoundGraph))
+				{
+					return true;
+				}
 			}
 		}
 	}
diff --git a/Engine/Source/Editor/DetailCustomizations/Private/AnimTransitionNodeDetails.cpp b/Engine/Source/Editor/DetailCustomizations/Private/AnimTransitionNodeDetails.cpp
--- a/Engine/Source/Editor/DetailCustomizations/Private/AnimTransitionNodeDetails.cpp
+++ b/Engine/Source/Editor/DetailCustomizations/Private/AnimTransitionNodeDetails.cpp
@@ -367,33 +367,58 @@ TSharedRef<SWidget> FAnimTransitionNodeDetails::OnGetShareableNodesMenu(bool bSh
 
 	MenuBuilder.BeginSection("AnimTransitionSharableNodes", SectionText);
 
-	if (UAnimStateTransitionNode* TransNode = TransitionNode.Get())
+	if (UAnimStateTransitionNode* RawTransitionNode = TransitionNode.Get())
 	{
-		const UEdGraph* CurrentGraph = TransNode->GetGraph();
-
-		// Loop through the graph and build a list of the unique shared transitions
-		TMap<FString, UAnimStateTransitionNode*> SharedTransitions;
+		const UEdGraph* CurrentGraph = RawTransitionNode->GetGraph();
 
-		for (int32 NodeIdx=0; NodeIdx < CurrentGraph->Nodes.Num(); NodeIdx++)
+		// Collect all unique shared transitions and group them by their name.
+		TMultiMap<FString, UAnimStateTransitionNode*> SharedTransitions;
+		if (UBlueprint* Blueprint = FBlueprintEditorUtils::FindBlueprintForGraph(CurrentGraph))
 		{
-			if (UAnimStateTransitionNode* GraphTransNode = Cast<UAnimStateTransitionNode>(CurrentGraph->Nodes[NodeIdx]))
+			TArray<UAnimStateNodeBase*> StateNodes;
+			FBlueprintEditorUtils::GetAllNodesOfClassEx<UAnimStateNodeBase>(Blueprint, StateNodes);
+
+			for (UAnimStateNodeBase* StateNodeBase : StateNodes)
 			{
-				if (bShareRules && !GraphTransNode->SharedRulesName.IsEmpty())
+				if (UAnimStateTransitionNode* GraphTransNode = Cast<UAnimStateTransitionNode>(StateNodeBase))
 				{
-					SharedTransitions.Add(GraphTransNode->SharedRulesName, GraphTransNode);
-				}
+					if (bShareRules && !GraphTransNode->SharedRulesName.IsEmpty())
+					{
+						SharedTransitions.Add(GraphTransNode->SharedRulesName, GraphTransNode);
+					}
 
-				if (!bShareRules && !GraphTransNode->SharedCrossfadeName.IsEmpty())
-				{
-					SharedTransitions.Add(GraphTransNode->SharedCrossfadeName, GraphTransNode);
+					if (!bShareRules && !GraphTransNode->SharedCrossfadeName.IsEmpty())
+					{
+						SharedTransitions.Add(GraphTransNode->SharedCrossfadeName, GraphTransNode);
+					}
 				}
 			}
 		}
 
-		for (auto Iter = SharedTransitions.CreateIterator(); Iter; ++Iter)
+		// Get the unique shared transition names
+		TSet<FString> SharedTransitionKeys;
+		SharedTransitions.GetKeys(SharedTransitionKeys);
+
+		// Iterate through the unique shared transition names and list all the places where they are referenced in the tooltip.
+		TArray<UAnimStateTransitionNode*> UsedIn;
+		for (const FString& Key : SharedTransitionKeys)
 		{
-			FUIAction Action = FUIAction( FExecuteAction::CreateSP(this, &FAnimTransitionNodeDetails::BecomeSharedWith, Iter.Value(), bShareRules) );
-			MenuBuilder.AddMenuEntry( FText::FromString( Iter.Key() ), LOCTEXT("ShaerdTransitionToolTip", "Use this shared transition"), FSlateIcon(), Action);
+			UsedIn.Reset();
+			SharedTransitions.MultiFind(Key, UsedIn, /*bMaintainOrder=*/true);
+			if (UsedIn.IsEmpty())
+			{
+				continue;
+			}
+
+			FTextBuilder ToolTipBuilder;
+			ToolTipBuilder.AppendLine(LOCTEXT("AnimTransitionUsedBy", "Used by:"));
+			for (const UAnimStateTransitionNode* UsedInTransitionNode : UsedIn)
+			{
+				ToolTipBuilder.AppendLine(UsedInTransitionNode->GetGraph()->GetName());
+			}
+
+			FUIAction Action = FUIAction( FExecuteAction::CreateSP(this, &FAnimTransitionNodeDetails::BecomeSharedWith, UsedIn[0], bShareRules));
+			MenuBuilder.AddMenuEntry( FText::FromString(Key), ToolTipBuilder.ToText(), FSlateIcon(), Action);
 		}
 	}
 	MenuBuilder.EndSection();
