Fixed double localization string

#rnx

diff --git a/Engine/Source/Editor/DetailCustomizations/Private/LandscapeProxyUIDetails.cpp b/Engine/Source/Editor/DetailCustomizations/Private/LandscapeProxyUIDetails.cpp
--- a/Engine/Source/Editor/DetailCustomizations/Private/LandscapeProxyUIDetails.cpp
+++ b/Engine/Source/Editor/DetailCustomizations/Private/LandscapeProxyUIDetails.cpp
@@ -266,7 +266,7 @@ void FLandscapeProxyUIDetails::CustomizeDetails( IDetailLayoutBuilder& DetailBui
 
 		IDetailCategoryBuilder& CategoryBuilder = DetailBuilder.EditCategory("Nanite", FText::GetEmpty(), ECategoryPriority::Default);
 		TSharedRef<IPropertyHandle> EnableNaniteProperty = DetailBuilder.GetProperty(GET_MEMBER_NAME_CHECKED(ALandscape, bEnableNanite));
-		CategoryBuilder.AddCustomRow(LOCTEXT("RebuildNaniteData", "Rebuild Nanite Data"))
+		CategoryBuilder.AddCustomRow(FText::FromString(TEXT("Rebuild Nanite Data")))
 		[
 			SNew(SHorizontalBox)
 			.IsEnabled_Lambda([HasAtLeastOneNaniteLandscape] { return HasAtLeastOneNaniteLandscape(); })
