- One File Per Actor + Folder Objects: No longer experimental
- Remove support on single actor packaging mode

diff --git a/Engine/Source/Editor/DetailCustomizations/Private/ActorDetails.cpp b/Engine/Source/Editor/DetailCustomizations/Private/ActorDetails.cpp
--- a/Engine/Source/Editor/DetailCustomizations/Private/ActorDetails.cpp
+++ b/Engine/Source/Editor/DetailCustomizations/Private/ActorDetails.cpp
@@ -55,7 +55,6 @@
 #include "Misc/MessageDialog.h"
 #include "ScopedTransaction.h"
 #include "WorldPartition/WorldPartitionSubsystem.h"
-#include "Settings/EditorExperimentalSettings.h"
 #include "Algo/AnyOf.h"
 
 #include "Kismet2/BlueprintEditorUtils.h"
@@ -581,28 +580,6 @@ FReply FActorDetails::HandleAddOrViewEventForVariable(UBlueprint* BP, FMulticast
 	return FReply::Handled();
 }
 
-namespace ActorDetailsUtil
-{
-	constexpr int32 MultipleValuesIndicator = 2;
-
-	FText GetActorPackagingModeText(int32 Mode)
-	{
-		switch (Mode)
-		{
-			// false: internal
-		case 0:
-			return LOCTEXT("InternalActorPackaging", "Internal");
-			// true: external
-		case 1:
-			return LOCTEXT("ExternalActorPackaging", "External");
-		case MultipleValuesIndicator:
-			return LOCTEXT("MultipleValues", "Multiple Values");
-		default:
-			return LOCTEXT("InternalActorPackaging", "Internal");
-		}
-	}
-}
-
 void FActorDetails::AddActorCategory( IDetailLayoutBuilder& DetailBuilder, const TMap<ULevel*, int32>& ActorsPerLevelCount )
 {		
 	FLevelEditorModule& LevelEditor = FModuleManager::GetModuleChecked<FLevelEditorModule>( TEXT("LevelEditor") );
@@ -721,112 +698,6 @@ void FActorDetails::AddActorCategory( IDetailLayoutBuilder& DetailBuilder, const
 			MakeConvertMenu( SelectedActorInfo )
 		];
 	}
-
-	if (GetDefault<UEditorExperimentalSettings>()->bEnableOneFilePerActorSupport)
-	{
-		// WorldSettings should not be packaged externally
-		if (SelectedActorInfo.SelectionClass != AWorldSettings::StaticClass())
-		{
-			// Actor Packaging Mode
-			const bool bIsPartitionedWorld = UWorld::IsPartitionedWorld(SelectedActorInfo.SharedWorld);
-
-			auto OnGetMenuContent = [=]() -> TSharedRef<SWidget> {
-				FMenuBuilder MenuBuilder(true, nullptr);
-				MenuBuilder.AddMenuEntry(ActorDetailsUtil::GetActorPackagingModeText(0), FText(), FSlateIcon(), FExecuteAction::CreateSP(this, &FActorDetails::OnActorPackagingModeChanged, false));
-				MenuBuilder.AddMenuEntry(ActorDetailsUtil::GetActorPackagingModeText(1), FText(), FSlateIcon(), FExecuteAction::CreateSP(this, &FActorDetails::OnActorPackagingModeChanged, true));
-				return MenuBuilder.MakeWidget();
-			};
-
-			ActorCategory.AddCustomRow( LOCTEXT("ActorPackagingModeRow", "ActorPackagingMode"), true)
-				.NameContent()
-				[
-					SNew(STextBlock)
-					.Text(LOCTEXT("ActorPackagingMode", "Packaging Mode"))
-					.ToolTipText(LOCTEXT("ActorPackagingMode_ToolTip", "Change the actor packaging mode. This will indicate if the actor is packaged alongside its level or in an external package."))
-					.Font(IDetailLayoutBuilder::GetDetailFont())
-				]
-				.ValueContent()
-				[
-					SNew(SComboButton)
-					.ContentPadding(2)
-					.OnGetMenuContent_Lambda(OnGetMenuContent)
-					.IsEnabled(this, &FActorDetails::IsActorPackagingModeEditable)
-					.ButtonContent()
-					[
-						SNew(STextBlock)
-						.Text(this, &FActorDetails::GetCurrentActorPackagingMode)
-						.Font(IDetailLayoutBuilder::GetDetailFont())
-					]
-					.IsEnabled(!bIsPartitionedWorld)
-				];
-		}
-	}
-}
-
-bool FActorDetails::IsActorPackagingModeEditable() const
-{
-	for (TWeakObjectPtr<AActor> Actor : SelectedActors)
-	{
-		if (Actor.IsValid() && Actor->GetLevel())
-		{
-			if (!Actor->SupportsExternalPackaging())
-			{
-				return false;
-			}
-		}
-	}
-	return true;
-}
-
-FText FActorDetails::GetCurrentActorPackagingMode() const
-{
-	int32 PackagingMode = -1;
-	for (TWeakObjectPtr<AActor> Actor : SelectedActors)
-	{
-		if (Actor.IsValid())
-		{
-			// If loading strategy is `Count` initialize it, otherwise set it to `None` if different between selected actors
-			PackagingMode = PackagingMode == -1 ? (int32)Actor->IsPackageExternal() : (PackagingMode != (int32)Actor->IsPackageExternal() ? ActorDetailsUtil::MultipleValuesIndicator : PackagingMode);
-		}
-	}
-	return ActorDetailsUtil::GetActorPackagingModeText(PackagingMode);
-}
-
-void FActorDetails::OnActorPackagingModeChanged(bool bExternal)
-{
-	// Validate all actors are in a saved map
-	TArray<AActor*> ActorsToConvert;
-	for (TWeakObjectPtr<AActor> Actor : SelectedActors)
-	{
-		if (Actor.IsValid() && Actor->GetLevel())
-		{
-			ULevel* Level = Actor->GetLevel();
-			UPackage* LevelPackage = Level->GetOutermost();
-			FString LevelName = LevelPackage->GetName();
-			if (LevelPackage == GetTransientPackage()
-				|| LevelPackage->HasAnyFlags(RF_Transient)
-				|| !FPackageName::IsValidLongPackageName(LevelName))
-			{
-				FText Message = FText::Format(LOCTEXT("ActorPackagingModeSaveMap", "You need to save level {0} before changing packaging mode on its actors."), FText::FromString(LevelName));
-				FMessageDialog::Open(EAppMsgType::Ok, Message);
-				return;
-			}
-			else if (Actor->SupportsExternalPackaging())
-			{
-				ActorsToConvert.Add(Actor.Get());
-			}
-			else
-			{
-				UE_LOG(LogLevel, Warning, TEXT("Can't convert %s to external packaging."), *Actor->GetName());
-			}
-		}
-	}
-
-	FScopedTransaction Transaction(LOCTEXT("ActorSetPackageExternal", "Change Actors Assigned Package"));
-	for (AActor* Actor : ActorsToConvert)
-	{
-		Actor->SetPackageExternal(bExternal);
-	}
 }
 
 const TArray< TWeakObjectPtr<AActor> >& FActorDetails::GetSelectedActors() const
diff --git a/Engine/Source/Editor/DetailCustomizations/Private/ActorDetails.h b/Engine/Source/Editor/DetailCustomizations/Private/ActorDetails.h
--- a/Engine/Source/Editor/DetailCustomizations/Private/ActorDetails.h
+++ b/Engine/Source/Editor/DetailCustomizations/Private/ActorDetails.h
@@ -103,11 +103,6 @@ private:
 
 	// Functions for actor loading strategy details
 private:
-	// Functions to handle actor packaging mode (i.e. external or internal)
-	bool IsActorPackagingModeEditable() const;
-	FText GetCurrentActorPackagingMode() const;
-	void OnActorPackagingModeChanged(bool bExternal);
-
 
 	/** Bring up the menu for user to select the path to create blueprint at */
 	FReply OnPickBlueprintPathClicked(bool bHavest);
diff --git a/Engine/Source/Editor/DetailCustomizations/Private/WorldSettingsDetails.cpp b/Engine/Source/Editor/DetailCustomizations/Private/WorldSettingsDetails.cpp
--- a/Engine/Source/Editor/DetailCustomizations/Private/WorldSettingsDetails.cpp
+++ b/Engine/Source/Editor/DetailCustomizations/Private/WorldSettingsDetails.cpp
@@ -140,51 +140,46 @@ void FWorldSettingsDetails::AddWorldCustomization(IDetailLayoutBuilder& DetailBu
 
 	if (CustomizedLevel)
 	{
-		const bool bIsPartitionedWorld = UWorld::IsPartitionedWorld(CustomizedLevel->GetWorld());
-
 		IDetailCategoryBuilder& WorldCategory = DetailBuilder.EditCategory("World");
-		if (GetDefault<UEditorExperimentalSettings>()->bEnableOneFilePerActorSupport)
-		{
-			WorldCategory.AddCustomRow(LOCTEXT("LevelUseExternalActorsRow", "LevelUseExternalActors"), true)
-				.NameContent()
-				[
-					SNew(STextBlock)
-					.Text(LOCTEXT("LevelUseExternalActors", "Use External Actors"))
-					.ToolTipText(LOCTEXT("ActorPackagingMode_ToolTip", "Use external actors, new actor spawned in this level will be external and existing external actors will be loaded on load."))
-					.Font(IDetailLayoutBuilder::GetDetailFont())
-					.IsEnabled(!bIsPartitionedWorld)
-				]
-				.ValueContent()
-				[
-					SNew(SCheckBox)
-					.OnCheckStateChanged(this, &FWorldSettingsDetails::OnUseExternalActorsChanged, CustomizedLevel)
-					.IsChecked(this, &FWorldSettingsDetails::IsUseExternalActorsChecked, CustomizedLevel)
-					.IsEnabled(!bIsPartitionedWorld)
-				];
-		}
+		
+		WorldCategory.AddCustomRow(LOCTEXT("LevelUseExternalActorsRow", "LevelUseExternalActors"), true)
+			.NameContent()
+			[
+				SNew(STextBlock)
+				.Text(LOCTEXT("LevelUseExternalActors", "Use External Actors"))
+				.ToolTipText(LOCTEXT("ActorPackagingMode_ToolTip", "Use external actors, new actor spawned in this level will be external and existing external actors will be loaded on load."))
+				.Font(IDetailLayoutBuilder::GetDetailFont())
+				.IsEnabled(this, &FWorldSettingsDetails::IsUseExternalActorsEnabled, CustomizedLevel)
+			]
+			.ValueContent()
+			[
+				SNew(SCheckBox)
+				.OnCheckStateChanged(this, &FWorldSettingsDetails::OnUseExternalActorsChanged, CustomizedLevel)
+				.IsChecked(this, &FWorldSettingsDetails::IsUseExternalActorsChecked, CustomizedLevel)
+				.IsEnabled(this, &FWorldSettingsDetails::IsUseExternalActorsEnabled, CustomizedLevel)
+			];
+		
 
 		const bool bIsUsingActorFolders = CustomizedLevel->IsUsingActorFolders();
-		if (bIsUsingActorFolders || GetDefault<UEditorExperimentalSettings>()->bEnableActorFolderObjectSupport)
-		{
-			WorldCategory.AddCustomRow(LOCTEXT("LevelUseActorFoldersRow", "LevelUseActorFolders"), true)
-				.NameContent()
-				[
-					SNew(STextBlock)
-					.Text(LOCTEXT("LevelUseActorFolders", "Use Actor Folder Objects"))
-					.ToolTipText(LOCTEXT("LevelUseActorFolders_ToolTip", "Use actor folder objects, actor folders of this level will be persistent in their own object."))
-					.Font(IDetailLayoutBuilder::GetDetailFont())
-					.IsEnabled(!bIsUsingActorFolders)
-				]
-				.ValueContent()
-				[
-					SNew(SCheckBox)
-					.OnCheckStateChanged(this, &FWorldSettingsDetails::OnUseActorFoldersChanged, CustomizedLevel)
-					.IsChecked(this, &FWorldSettingsDetails::IsUsingActorFoldersChecked, CustomizedLevel)
-					.IsEnabled(!bIsUsingActorFolders)
-				];
-		}
+		WorldCategory.AddCustomRow(LOCTEXT("LevelUseActorFoldersRow", "LevelUseActorFolders"), true)
+			.NameContent()
+			[
+				SNew(STextBlock)
+				.Text(LOCTEXT("LevelUseActorFolders", "Use Actor Folder Objects"))
+				.ToolTipText(LOCTEXT("LevelUseActorFolders_ToolTip", "Use actor folder objects, actor folders of this level will be persistent in their own object."))
+				.Font(IDetailLayoutBuilder::GetDetailFont())
+				.IsEnabled(this, &FWorldSettingsDetails::IsUsingActorFoldersEnabled, CustomizedLevel)
+			]
+			.ValueContent()
+			[
+				SNew(SCheckBox)
+				.OnCheckStateChanged(this, &FWorldSettingsDetails::OnUseActorFoldersChanged, CustomizedLevel)
+				.IsChecked(this, &FWorldSettingsDetails::IsUsingActorFoldersChecked, CustomizedLevel)
+				.IsEnabled(this, &FWorldSettingsDetails::IsUsingActorFoldersEnabled, CustomizedLevel)
+			];
+		
 
-		if (bIsPartitionedWorld)
+		if (IsPartitionedWorld(CustomizedLevel))
 		{
 			IDetailCategoryBuilder& WorldPartitionCategory = DetailBuilder.EditCategory("WorldPartition");
 
@@ -195,7 +190,6 @@ void FWorldSettingsDetails::AddWorldCustomization(IDetailLayoutBuilder& DetailBu
 					.Text(LOCTEXT("DefaultWorldPartitionSettings", "Default World Partition Settings"))
 					.ToolTipText(LOCTEXT("DefaultWorldPartitionSettings_ToolTip", "Save or Reset the current World Partition default editor state"))
 					.Font(IDetailLayoutBuilder::GetDetailFont())
-					.IsEnabled(bIsPartitionedWorld)
 				]
 				.ValueContent()
 				[
@@ -230,7 +224,6 @@ void FWorldSettingsDetails::AddWorldCustomization(IDetailLayoutBuilder& DetailBu
 							.Text(LOCTEXT("SaveButtonText", "Save"))
 							.ToolTipText(LOCTEXT("SaveButtonToolTip", "Save current World Partition editor state as map default"))
 							.Font(IDetailLayoutBuilder::GetDetailFont())
-							.IsEnabled(bIsPartitionedWorld)
 						]
 					]
 				];
@@ -251,6 +244,11 @@ ECheckBoxState FWorldSettingsDetails::IsUsingActorFoldersChecked(ULevel* Level)
 	return Level->IsUsingActorFolders() ? ECheckBoxState::Checked : ECheckBoxState::Unchecked;
 }
 
+bool FWorldSettingsDetails::IsUsingActorFoldersEnabled(ULevel* Level) const
+{
+	return !Level->IsUsingActorFolders();
+}
+
 void FWorldSettingsDetails::OnUseExternalActorsChanged(ECheckBoxState BoxState, ULevel* Level)
 {
 	if (Level != nullptr)
@@ -265,20 +263,20 @@ void FWorldSettingsDetails::OnUseExternalActorsChanged(ECheckBoxState BoxState,
 			return;
 		}
 
-		FScopedTransaction Transaction(LOCTEXT("WorldUseExternalActors", "Change World Use External Actors"));
-
-		Level->Modify();
-		Level->SetUseExternalActors(BoxState == ECheckBoxState::Checked);
-		
 		FText MessageTitle(LOCTEXT("ConvertActorPackagingDialog", "Convert Actors Packaging"));
-		FText PackagingMode = Level->IsUsingExternalActors() ? LOCTEXT("ExternalActors", "External") : LOCTEXT("InternalActors", "Internal");
+		FText PackagingMode = !Level->IsUsingExternalActors() ? LOCTEXT("ExternalActors", "External") : LOCTEXT("InternalActors", "Internal");
 		FText Message = FText::Format(LOCTEXT("ConvertActorPackagingMsg", "Do you want to convert all actors to {0} packaging as well?"), PackagingMode);
 		EAppReturnType::Type ConvertAnswer = FMessageDialog::Open(EAppMsgType::YesNo, Message, &MessageTitle);
-
 		// if the user accepts, convert all actors to what the new packaging mode will be
 		if (ConvertAnswer == EAppReturnType::Yes)
 		{
-			Level->ConvertAllActorsToPackaging(Level->IsUsingExternalActors());
+			Level->Modify();
+			Level->ConvertAllActorsToPackaging(BoxState == ECheckBoxState::Checked);
+
+			if (Level->IsUsingExternalActors())
+			{
+				Level->SetUseActorFolders(true);
+			}
 		}
 	}
 }
@@ -288,6 +286,16 @@ ECheckBoxState FWorldSettingsDetails::IsUseExternalActorsChecked(ULevel* Level)
 	return Level->IsUsingExternalActors() ? ECheckBoxState::Checked : ECheckBoxState::Unchecked;
 }
 
+bool FWorldSettingsDetails::IsPartitionedWorld(ULevel* Level) const
+{
+	return UWorld::IsPartitionedWorld(Level->GetWorld());
+}
+
+bool FWorldSettingsDetails::IsUseExternalActorsEnabled(ULevel* Level) const
+{
+	return !IsPartitionedWorld(Level);
+}
+
 FLightmapCustomNodeBuilder::FLightmapCustomNodeBuilder(const TSharedPtr<FAssetThumbnailPool>& InThumbnailPool)
 {
 	ThumbnailPool = InThumbnailPool;
diff --git a/Engine/Source/Editor/DetailCustomizations/Private/WorldSettingsDetails.h b/Engine/Source/Editor/DetailCustomizations/Private/WorldSettingsDetails.h
--- a/Engine/Source/Editor/DetailCustomizations/Private/WorldSettingsDetails.h
+++ b/Engine/Source/Editor/DetailCustomizations/Private/WorldSettingsDetails.h
@@ -80,18 +80,27 @@ protected:
 	void AddWorldCustomization(IDetailLayoutBuilder& DetailBuilder);
 
 private:
+	// return true if level's owning world is partitioned.
+	bool IsPartitionedWorld(ULevel* Level) const;
+	
 	// Called when `ULevel::bUseExternalActors` changes.
 	void OnUseExternalActorsChanged(ECheckBoxState State, ULevel* Level);
 
-	// return the state of `ULevel::bUseExternalActors`
+	// return the state of `ULevel::bUseExternalActors`.
 	ECheckBoxState IsUseExternalActorsChecked(ULevel* Level) const;
 
+	// return true if the state of 'ULevel::bUseExternalActors' can be changed.
+	bool IsUseExternalActorsEnabled(ULevel* Level) const;
+
 	// Called when `ULevel::bUseActorFolders` changes.
 	void OnUseActorFoldersChanged(ECheckBoxState BoxState, ULevel* Level);
 
-	// return the state of `ULevel::bUseActorFolders`
+	// return the state of `ULevel::bUseActorFolders`.
 	ECheckBoxState IsUsingActorFoldersChecked(ULevel* Level) const;
 
+	// return true if the state of 'ULevel::bUseActorFolders' can be changed.
+	bool IsUsingActorFoldersEnabled(ULevel* Level) const;
+
 	// Handles checking whether a given asset is acceptable for drag-and-drop.
 	bool HandleAssetDropTargetIsAssetAcceptableForDrop( const UObject* InObject ) const;
 
diff --git a/Engine/Source/Editor/UnrealEd/Classes/Settings/EditorExperimentalSettings.h b/Engine/Source/Editor/UnrealEd/Classes/Settings/EditorExperimentalSettings.h
--- a/Engine/Source/Editor/UnrealEd/Classes/Settings/EditorExperimentalSettings.h
+++ b/Engine/Source/Editor/UnrealEd/Classes/Settings/EditorExperimentalSettings.h
@@ -115,14 +115,6 @@ public:
 	UPROPERTY(EditAnywhere, config, Category = Tools, meta = (ConfigRestartRequired = true))
 	bool bFacialAnimationImporter;
 
-	/** Whether to enable one file per actor support in the engine. */
-	UPROPERTY(EditAnywhere, config, Category = Tools, meta=(DisplayName="Enable One File Per Actor"))
-	bool bEnableOneFilePerActorSupport;
-
-	/** Whether to enable actor folder object support in the engine. */
-	UPROPERTY(EditAnywhere, config, Category = Tools, meta = (DisplayName = "Enable Actor Folder Objects"))
-	bool bEnableActorFolderObjectSupport;
-
 	/** Enable experimental PIE preview device launch */
 	UPROPERTY(EditAnywhere, config, Category = PIE, meta = (DisplayName = "Enable mobile PIE with preview device launch options."))
 	bool bMobilePIEPreviewDeviceLaunch;
diff --git a/Engine/Source/Editor/WorldBrowser/Private/StreamingLevels/StreamingLevelCollectionModel.cpp b/Engine/Source/Editor/WorldBrowser/Private/StreamingLevels/StreamingLevelCollectionModel.cpp
--- a/Engine/Source/Editor/WorldBrowser/Private/StreamingLevels/StreamingLevelCollectionModel.cpp
+++ b/Engine/Source/Editor/WorldBrowser/Private/StreamingLevels/StreamingLevelCollectionModel.cpp
@@ -6,7 +6,6 @@
 #include "Framework/MultiBox/MultiBoxBuilder.h"
 #include "Editor/EditorEngine.h"
 #include "Settings/LevelEditorMiscSettings.h"
-#include "Settings/EditorExperimentalSettings.h"
 #include "Engine/LevelStreamingAlwaysLoaded.h"
 #include "Engine/LevelStreamingDynamic.h"
 #include "Editor.h"
@@ -350,12 +349,9 @@ void FStreamingLevelCollectionModel::BuildHierarchyMenu(FMenuBuilder& InMenuBuil
 			InMenuBuilder.AddMenuEntry( Commands.MoveFoliageToSelected );
 		}
 
-		if (GetDefault<UEditorExperimentalSettings>()->bEnableOneFilePerActorSupport)
-		{
-			InMenuBuilder.AddMenuEntry(Commands.ConvertLevelToExternalActors);
-			InMenuBuilder.AddMenuEntry(Commands.ConvertLevelToInternalActors);
-		}
-
+		InMenuBuilder.AddMenuEntry(Commands.ConvertLevelToExternalActors);
+		InMenuBuilder.AddMenuEntry(Commands.ConvertLevelToInternalActors);
+		
 		if (AreAnyLevelsSelected() && !(IsOneLevelSelected() && SelectedLevelsList[0]->IsPersistent()))
 		{
 			InMenuBuilder.AddMenuEntry( Commands.SelectStreamingVolumes );
