Color Grading Vector Customization: Added metadata clamping to vector values.

#rnx

diff --git a/Engine/Source/Editor/DetailCustomizations/Private/ColorGradingVectorCustomization.cpp b/Engine/Source/Editor/DetailCustomizations/Private/ColorGradingVectorCustomization.cpp
--- a/Engine/Source/Editor/DetailCustomizations/Private/ColorGradingVectorCustomization.cpp
+++ b/Engine/Source/Editor/DetailCustomizations/Private/ColorGradingVectorCustomization.cpp
@@ -51,6 +51,45 @@ class IPropertyTypeCustomizationUtils;
 
 #define LOCTEXT_NAMESPACE "FColorGradingCustomization"
 
+
+namespace
+{
+	static FVector4 ClampValueFromMetaData(FVector4 InValue, FProperty* InProperty)
+	{
+		FVector4 RetVal = InValue;
+		if (InProperty)
+		{
+			//Enforce min
+			const FString& MinString = InProperty->GetMetaData(TEXT("ClampMin"));
+			if (MinString.Len())
+			{
+				checkSlow(MinString.IsNumeric());
+				double MinValue;
+				TTypeFromString<double>::FromString(MinValue, *MinString);
+				for (int Index = 0; Index < 4; Index++)
+				{
+					RetVal[Index] = FMath::Max<double>(MinValue, RetVal[Index]);
+				}
+			}
+			//Enforce max 
+			const FString& MaxString = InProperty->GetMetaData(TEXT("ClampMax"));
+			if (MaxString.Len())
+			{
+				checkSlow(MaxString.IsNumeric());
+				double MaxValue;
+				TTypeFromString<double>::FromString(MaxValue, *MaxString);
+
+				for (int Index = 0; Index < 4; Index++)
+				{
+					RetVal[Index] = FMath::Min<double>(MaxValue, RetVal[Index]);
+				}
+			}
+		}
+
+		return RetVal;
+	}
+}
+
 FColorGradingVectorCustomizationBase::FColorGradingVectorCustomizationBase(TWeakPtr<IPropertyHandle> InColorGradingPropertyHandle, const TArray<TWeakPtr<IPropertyHandle>>& InSortedChildArray)
 	: ColorGradingPropertyHandle(InColorGradingPropertyHandle)
 	, SortedChildArray(InSortedChildArray)
@@ -237,7 +276,7 @@ void FColorGradingVectorCustomizationBase::OnValueChanged(float NewValue, int32
 {
 	FVector4 CurrentValueVector;
 	verifySlow(ColorGradingPropertyHandle.Pin()->GetValue(CurrentValueVector) == FPropertyAccess::Success);
-
+	ClampValueFromMetaData(CurrentValueVector, ColorGradingPropertyHandle.Pin()->GetProperty());
 	FVector4 NewValueVector = CurrentValueVector;
 
 	if (IsRGBMode)
diff --git a/Engine/Source/Runtime/Engine/Classes/Engine/Scene.h b/Engine/Source/Runtime/Engine/Classes/Engine/Scene.h
--- a/Engine/Source/Runtime/Engine/Classes/Engine/Scene.h
+++ b/Engine/Source/Runtime/Engine/Classes/Engine/Scene.h
@@ -155,7 +155,7 @@ struct FColorGradePerRangeSettings
 	UPROPERTY(Interp, BlueprintReadWrite, Category = "Color Grading", meta = (UIMin = "0.0", UIMax = "2.0", Delta = "0.01", ColorGradingMode = "contrast", ShiftMouseMovePixelPerDelta = "10", SupportDynamicSliderMaxValue = "true", DisplayName = "Contrast"))
 	FVector4 Contrast;
 
-	UPROPERTY(Interp, BlueprintReadWrite, Category = "Color Grading", meta = (UIMin = "0.0", UIMax = "2.0", Delta = "0.01", ColorGradingMode = "gamma", ShiftMouseMovePixelPerDelta = "10", SupportDynamicSliderMaxValue = "true", DisplayName = "Gamma"))
+	UPROPERTY(Interp, BlueprintReadWrite, Category = "Color Grading", meta = (UIMin = "0.01", UIMax = "2.0", Delta = "0.01", ClampMin = "0.01", ColorGradingMode = "gamma", ShiftMouseMovePixelPerDelta = "10", SupportDynamicSliderMaxValue = "true", DisplayName = "Gamma"))
 	FVector4 Gamma;
 
 	UPROPERTY(Interp, BlueprintReadWrite, Category = "Color Grading", meta = (UIMin = "0.0", UIMax = "2.0", Delta = "0.01", ColorGradingMode = "gain", ShiftMouseMovePixelPerDelta = "10", SupportDynamicSliderMaxValue = "true", DisplayName = "Gain"))
