PR #9603: Fix unable to copy/paste materials on mesh components (Contributed by [REDACTED_PERSON])

diff --git a/Engine/Source/Editor/DetailCustomizations/DetailCustomizations.Build.cs b/Engine/Source/Editor/DetailCustomizations/DetailCustomizations.Build.cs
--- a/Engine/Source/Editor/DetailCustomizations/DetailCustomizations.Build.cs
+++ b/Engine/Source/Editor/DetailCustomizations/DetailCustomizations.Build.cs
@@ -64,6 +64,8 @@ public class DetailCustomizations : ModuleRules
 				"ToolWidgets",
 				"MaterialEditor",
 				"VirtualTexturingEditor",
+				"Json",
+				"JsonUtilities",
 			}
 		);
 
diff --git a/Engine/Source/Editor/DetailCustomizations/Private/ComponentMaterialCategory.cpp b/Engine/Source/Editor/DetailCustomizations/Private/ComponentMaterialCategory.cpp
--- a/Engine/Source/Editor/DetailCustomizations/Private/ComponentMaterialCategory.cpp
+++ b/Engine/Source/Editor/DetailCustomizations/Private/ComponentMaterialCategory.cpp
@@ -18,9 +18,11 @@
 #include "Editor/UnrealEdEngine.h"
 #include "HAL/PlatformCrt.h"
 #include "HAL/PlatformMath.h"
+#include "HAL/PlatformApplicationMisc.h"
 #include "IPropertyUtilities.h"
 #include "Internationalization/Internationalization.h"
 #include "Internationalization/Text.h"
+#include "JsonObjectConverter.h"
 #include "LandscapeComponent.h"
 #include "LandscapeProxy.h"
 #include "MaterialList.h"
@@ -46,7 +48,7 @@ class UWorld;
 class FMaterialIterator
 {
 public:
-	FMaterialIterator( TArray< TWeakObjectPtr<USceneComponent> >& InSelectedComponents )
+	FMaterialIterator( const TArray< TWeakObjectPtr<USceneComponent> >& InSelectedComponents )
 		: SelectedComponents( InSelectedComponents )
 		, CurMaterial( NULL )
 		, CurComponent( NULL )
@@ -154,7 +156,7 @@ public:
 
 private:
 	/** Reference to the selected components */
-	TArray< TWeakObjectPtr<USceneComponent> >& SelectedComponents;
+	const TArray< TWeakObjectPtr<USceneComponent> >& SelectedComponents;
 	/** The current material the iterator is stopped on */
 	UMaterialInterface* CurMaterial;
 	/** The current component using the current material */
@@ -181,6 +183,9 @@ void FComponentMaterialCategory::Create( IDetailLayoutBuilder& DetailBuilder )
 	FMaterialListDelegates MaterialListDelegates;
 	MaterialListDelegates.OnGetMaterials.BindSP( this, &FComponentMaterialCategory::OnGetMaterialsForView );
 	MaterialListDelegates.OnMaterialChanged.BindSP( this, &FComponentMaterialCategory::OnMaterialChanged );
+	MaterialListDelegates.OnCanCopyMaterialItem.BindSP( this, &FComponentMaterialCategory::OnCanCopyMaterialItem );
+	MaterialListDelegates.OnCopyMaterialItem.BindSP( this, &FComponentMaterialCategory::OnCopyMaterialItem );
+	MaterialListDelegates.OnPasteMaterialItem.BindSP(this, &FComponentMaterialCategory::OnPasteMaterialItem);
 	
 	//Pass an empty material list owner (owner can be use by the asset picker filter. In this case we do not need it)
 	TArray<FAssetData> MaterialListOwner;
@@ -426,4 +431,104 @@ void FComponentMaterialCategory::OnMaterialChanged( UMaterialInterface* NewMater
 	}
 }
 
+void FComponentMaterialCategory::OnCopyMaterialItem(int32 CurrentSlot)
+{
+	// For now we don't support multi-select
+	// TODO: Pass through more data to these delegates that allows us to know which
+	// material should be copied, or allow reaching into the generated material list from here
+	if (SelectedComponents.Num() != 1)
+	{
+		return;
+	}
+
+	for (FMaterialIterator It(SelectedComponents); It; ++It)
+	{
+		if (It.GetMaterialIndex() == CurrentSlot)
+		{
+			TSharedRef<FJsonObject> RootJsonObject = MakeShareable(new FJsonObject());
+
+			const FStaticMaterial StaticMaterial(It.GetMaterial());
+			FJsonObjectConverter::UStructToJsonObject(FStaticMaterial::StaticStruct(), &StaticMaterial, RootJsonObject, 0, 0);
+
+			typedef TJsonWriter<TCHAR, TPrettyJsonPrintPolicy<TCHAR>> FStringWriter;
+			typedef TJsonWriterFactory<TCHAR, TPrettyJsonPrintPolicy<TCHAR>> FStringWriterFactory;
+
+			FString CopyStr;
+			TSharedRef<FStringWriter> Writer = FStringWriterFactory::Create(&CopyStr);
+			FJsonSerializer::Serialize(RootJsonObject, Writer);
+
+			if (!CopyStr.IsEmpty())
+			{
+				FPlatformApplicationMisc::ClipboardCopy(*CopyStr);
+			}
+
+			break;
+		}
+	}
+}
+
+bool FComponentMaterialCategory::OnCanCopyMaterialItem(int32 CurrentSlot) const
+{
+	// For now we don't support multi-select
+	if (SelectedComponents.Num() != 1)
+	{
+		return false;
+	}
+
+	for (FMaterialIterator It(SelectedComponents); It; ++It)
+	{
+		const UActorComponent* CurrentComponent = It.GetComponent();
+
+		// We need to handle the same various component types as OnMaterialChanged,
+		// which are all primitive components except decals, which only have a single material slot anyway
+		if (const UPrimitiveComponent* PrimitiveComponent = Cast<UPrimitiveComponent>(CurrentComponent))
+		{
+			if (CurrentSlot >= PrimitiveComponent->GetNumMaterials())
+			{
+				return false;
+			}
+		}
+		else if (CurrentComponent->IsA<UDecalComponent>())
+		{
+			return CurrentSlot == 0;
+		}
+	}
+
+	return true;
+}
+
+void FComponentMaterialCategory::OnPasteMaterialItem(int32 CurrentSlot)
+{
+	// For now we don't support multi-select
+	if (SelectedComponents.Num() != 1)
+	{
+		return;
+	}
+
+	FString PastedText;
+	FPlatformApplicationMisc::ClipboardPaste(PastedText);
+
+	TSharedPtr<FJsonObject> RootJsonObject;
+	TSharedRef<TJsonReader<>> Reader = TJsonReaderFactory<>::Create(PastedText);
+	FJsonSerializer::Deserialize(Reader, RootJsonObject);
+
+	if (RootJsonObject.IsValid())
+	{
+		FStaticMaterial TmpStaticMaterial;
+		FJsonObjectConverter::JsonObjectToUStruct(RootJsonObject.ToSharedRef(), FStaticMaterial::StaticStruct(), &TmpStaticMaterial, 0, 0);
+		UMaterialInterface* NewMaterial = TmpStaticMaterial.MaterialInterface;
+
+		// Try to find the previous material used for this slot so we can piggy-back on the existing material change method
+		for (FMaterialIterator It(SelectedComponents); It; ++It)
+		{
+			if (It.GetMaterialIndex() == CurrentSlot)
+			{
+				UMaterialInterface* PreviousMaterial = It.GetMaterial();
+				OnMaterialChanged(NewMaterial, PreviousMaterial, CurrentSlot, /* bReplaceAll */ false);
+				break;
+			}
+		}
+	}
+}
+
 #undef LOCTEXT_NAMESPACE
diff --git a/Engine/Source/Editor/DetailCustomizations/Private/ComponentMaterialCategory.h b/Engine/Source/Editor/DetailCustomizations/Private/ComponentMaterialCategory.h
--- a/Engine/Source/Editor/DetailCustomizations/Private/ComponentMaterialCategory.h
+++ b/Engine/Source/Editor/DetailCustomizations/Private/ComponentMaterialCategory.h
@@ -47,6 +47,28 @@ private:
 	 */
 	void OnMaterialChanged( UMaterialInterface* NewMaterial, UMaterialInterface* PrevMaterial, int32 SlotIndex, bool bReplaceAll );
 
+	/**
+	 * Called when a user executes the copy command on a material list item
+	 *
+	 * @param CurrentSlot	The slot index that was copied
+	 */
+	void OnCopyMaterialItem(int32 CurrentSlot);
+
+	/**
+	 * Called to determine if the copy command should be enabled for a particular a material list item
+	 *
+	 * @param CurrentSlot	The slot index that was copied
+	 * @return	whether this slot index can be copied
+	 */
+	bool OnCanCopyMaterialItem(int32 CurrentSlot) const;
+
+	/**
+	 * Called when a user executes the paste command on a material list item
+	 *
+	 * @param CurrentSlot	The slot index that was pasted on
+	 */
+	void OnPasteMaterialItem(int32 CurrentSlot);
+
 	/** 
 	 * @return true if a component is editable (and visible in the view)
 	 */
