Optimus: Add UMeshDeformer property customization that allows filtering so that only deformers whose primary bindings match the owning component are shown.

diff --git a/Engine/Plugins/Animation/DeformerGraph/Source/OptimusCore/Private/OptimusDeformer.cpp b/Engine/Plugins/Animation/DeformerGraph/Source/OptimusCore/Private/OptimusDeformer.cpp
--- a/Engine/Plugins/Animation/DeformerGraph/Source/OptimusCore/Private/OptimusDeformer.cpp
+++ b/Engine/Plugins/Animation/DeformerGraph/Source/OptimusCore/Private/OptimusDeformer.cpp
@@ -2279,6 +2279,35 @@ void UOptimusDeformer::PostRename(UObject* OldOuter, const FName OldName)
 	}
 }
 
+void UOptimusDeformer::GetAssetRegistryTags(TArray<FAssetRegistryTag>& OutTags) const
+{
+	Super::GetAssetRegistryTags(OutTags);
+
+	UClass* BindingClass = nullptr;
+
+	if (Bindings != nullptr)
+	{
+		for (UOptimusComponentSourceBinding const* Binding : Bindings->Bindings)
+		{
+			if (Binding != nullptr && Binding->IsPrimaryBinding())
+			{
+				if (Binding->GetComponentSource() != nullptr)
+				{
+					BindingClass = Binding->GetComponentSource()->GetComponentClass();
+					break;
+				}
+			}
+		}
+	}
+
+	if (BindingClass != nullptr)
+	{
+		FSoftClassPath ClassPath(BindingClass);
+		OutTags.Add(FAssetRegistryTag(TEXT("PrimaryBindingClass"), *ClassPath.ToString(), FAssetRegistryTag::TT_Hidden));
+	}
+}
+
+
 #if WITH_EDITOR
 void UOptimusDeformer::PostEditChangeProperty(FPropertyChangedEvent& PropertyChangedEvent)
 {
diff --git a/Engine/Plugins/Animation/DeformerGraph/Source/OptimusCore/Public/OptimusDeformer.h b/Engine/Plugins/Animation/DeformerGraph/Source/OptimusCore/Public/OptimusDeformer.h
--- a/Engine/Plugins/Animation/DeformerGraph/Source/OptimusCore/Public/OptimusDeformer.h
+++ b/Engine/Plugins/Animation/DeformerGraph/Source/OptimusCore/Public/OptimusDeformer.h
@@ -254,7 +254,8 @@ public:
 	// are not moved to the new package automatically (see FAssetRenameManager), so we
 	// have to manually perform the move/rename, to avoid invalid reference to the old package
 	void PostRename(UObject* OldOuter, const FName OldName) override;
-	
+	void GetAssetRegistryTags(TArray<FAssetRegistryTag>& OutTags) const override;
+
 #if WITH_EDITOR
 	void PostEditChangeProperty(FPropertyChangedEvent& PropertyChangedEvent) override;
 #endif
diff --git a/Engine/Source/Editor/DetailCustomizations/Private/DetailCustomizations.cpp b/Engine/Source/Editor/DetailCustomizations/Private/DetailCustomizations.cpp
--- a/Engine/Source/Editor/DetailCustomizations/Private/DetailCustomizations.cpp
+++ b/Engine/Source/Editor/DetailCustomizations/Private/DetailCustomizations.cpp
@@ -95,6 +95,7 @@
 #include "Math/Quat.h"
 #include "Math/Transform.h"
 #include "MeshComponentDetails.h"
+#include "MeshDeformerCustomizations.h"
 #include "MeshMergingSettingsCustomization.h"
 #include "MeshProxySettingsCustomizations.h"
 #include "Misc/AssertionMacros.h"
@@ -353,6 +354,7 @@ void FDetailCustomizationsModule::RegisterPropertyTypeCustomizations()
 	RegisterCustomPropertyTypeLayout("CustomPrimitiveData", FOnGetPropertyTypeCustomizationInstance::CreateStatic(&FCustomPrimitiveDataCustomization::MakeInstance));
 	RegisterCustomPropertyTypeLayout("TemplateString", FOnGetPropertyTypeCustomizationInstance::CreateStatic(&FTemplateStringStructCustomization::MakeInstance));
 	RegisterCustomPropertyTypeLayout("LightingChannels", FOnGetPropertyTypeCustomizationInstance::CreateStatic(&FLightingChannelsCustomization::MakeInstance));
+	RegisterCustomPropertyTypeLayout("MeshDeformer", FOnGetPropertyTypeCustomizationInstance::CreateStatic(&FMeshDeformerCustomization::MakeInstance));
 }
 
 #undef REGISTER_UIMINMAX_CUSTOMIZATION
diff --git a/Engine/Source/Editor/DetailCustomizations/Private/MeshDeformerCustomizations.cpp b/Engine/Source/Editor/DetailCustomizations/Private/MeshDeformerCustomizations.cpp
new file mode 100644
--- /dev/null
+++ b/Engine/Source/Editor/DetailCustomizations/Private/MeshDeformerCustomizations.cpp
@@ -0,0 +1,88 @@
+// Copyright Epic Games, Inc. All Rights Reserved.
+
+#include "MeshDeformerCustomizations.h"
+
+#include "Animation/MeshDeformer.h"
+#include "Components/SkinnedMeshComponent.h"
+#include "Engine/SkeletalMesh.h"
+#include "DetailWidgetRow.h"
+#include "PropertyCustomizationHelpers.h"
+#include "PropertyHandle.h"
+
+
+namespace
+{
+	/** Pull the class type from the PrimaryBindingClass tag in some AssetData. */
+	UClass* GetPrimaryBindingClassFromAssetData(FAssetData const& AssetData)
+	{
+		UClass* BindingClass = nullptr;
+		FString BindingClassPath;
+		if (AssetData.GetTagValue<FString>("PrimaryBindingClass", BindingClassPath))
+		{
+			FSoftClassPath BindingSoftClassPath(BindingClassPath);
+			BindingClass = BindingSoftClassPath.ResolveClass();
+		}
+		return BindingClass;
+	}
+}
+
+
+TSharedRef<IPropertyTypeCustomization> FMeshDeformerCustomization::MakeInstance()
+{
+	return MakeShareable(new FMeshDeformerCustomization);
+}
+
+void FMeshDeformerCustomization::CustomizeHeader(
+	TSharedRef<IPropertyHandle> InPropertyHandle,
+	FDetailWidgetRow& InHeaderRow,
+	IPropertyTypeCustomizationUtils& InCustomizationUtils)
+{
+	TArray<UObject*> OwningObjects;
+	InPropertyHandle->GetOuterObjects(OwningObjects);
+	if (OwningObjects.Num() == 0)
+	{
+		return;
+	}
+
+	UClass* OwnerClass = OwningObjects[0]->GetClass();
+
+	// Filter in USkeletalMesh context should assume USkinnedMeshComponent binding :(
+	if (OwnerClass == USkeletalMesh::StaticClass())
+	{
+		OwnerClass = USkinnedMeshComponent::StaticClass();
+	}
+
+	InHeaderRow.NameContent()
+	[
+		InPropertyHandle->CreatePropertyNameWidget()
+	]
+	.ValueContent()
+	[
+		SNew(SObjectPropertyEntryBox)
+		.PropertyHandle(InPropertyHandle)
+		.AllowedClass(UMeshDeformer::StaticClass())
+		.OnShouldFilterAsset_Lambda([OwnerClass](FAssetData const& AssetData)
+		{
+			// Filter depending on whether the PrimaryBindingClass matches our owning object.
+			// First try to load the class type from the asset registry.
+			UClass* BindingClass = GetPrimaryBindingClassFromAssetData(AssetData);
+
+			// If we can't find the tag in the registry then load the full object to get the class type (slow).
+			if (BindingClass == nullptr)
+			{
+				if (UObject* Object = AssetData.GetAsset())
+				{
+					FAssetData AssetDataWithObjectRegistryTags;
+					Object->GetAssetRegistryTags(AssetDataWithObjectRegistryTags);
+					BindingClass = GetPrimaryBindingClassFromAssetData(AssetDataWithObjectRegistryTags);
+				}
+			}
+
+			const bool bFoundBindingClass = BindingClass != nullptr;
+			const bool bCanBind = OwnerClass != nullptr && BindingClass != nullptr && OwnerClass->IsChildOf(BindingClass);
+			const bool bHideEntry = bFoundBindingClass && !bCanBind;
+
+			return bHideEntry;
+		})
+	];
+}
diff --git a/Engine/Source/Editor/DetailCustomizations/Private/MeshDeformerCustomizations.h b/Engine/Source/Editor/DetailCustomizations/Private/MeshDeformerCustomizations.h
new file mode 100644
--- /dev/null
+++ b/Engine/Source/Editor/DetailCustomizations/Private/MeshDeformerCustomizations.h
@@ -0,0 +1,22 @@
+// Copyright Epic Games, Inc. All Rights Reserved.
+
+#pragma once
+
+#include "IPropertyTypeCustomization.h"
+
+class FMeshDeformerCustomization : public IPropertyTypeCustomization
+{
+public:
+	static TSharedRef<IPropertyTypeCustomization> MakeInstance();
+
+	// IPropertyTypeCustomization overrides
+	void CustomizeHeader(
+		TSharedRef<IPropertyHandle> InPropertyHandle,
+		FDetailWidgetRow& InHeaderRow,
+		IPropertyTypeCustomizationUtils& InCustomizationUtils) override;
+
+	void CustomizeChildren(
+		TSharedRef<IPropertyHandle> InPropertyHandle,
+		IDetailChildrenBuilder& InChildBuilder,
+		IPropertyTypeCustomizationUtils& InCustomizationUtils) override {}
+};
