Optimus: Temp disable of mesh deformer asset picker filtering. Can reenable once we move shader compilation out of post load.

diff --git a/Engine/Source/Editor/DetailCustomizations/Private/MeshDeformerCustomizations.cpp b/Engine/Source/Editor/DetailCustomizations/Private/MeshDeformerCustomizations.cpp
--- a/Engine/Source/Editor/DetailCustomizations/Private/MeshDeformerCustomizations.cpp
+++ b/Engine/Source/Editor/DetailCustomizations/Private/MeshDeformerCustomizations.cpp
@@ -61,6 +61,11 @@ void FMeshDeformerCustomization::CustomizeHeader(
 		SNew(SObjectPropertyEntryBox)
 		.PropertyHandle(InPropertyHandle)
 		.AllowedClass(UMeshDeformer::StaticClass())
+
+		// Disable filtering for now because of slow asset loading. 
+		// This is due to kernel compilation happening in PostLoad().
+		// Need to implement deferred compilation on demand and can then re-enable here.
+#if 0 
 		.OnShouldFilterAsset_Lambda([OwnerClass](FAssetData const& AssetData)
 		{
 			// Filter depending on whether the PrimaryBindingClass matches our owning object.
@@ -84,5 +89,6 @@ void FMeshDeformerCustomization::CustomizeHeader(
 
 			return bHideEntry;
 		})
+#endif
 	];
 }
