Fix for 'Anim to Play' property not showing when selecting a skeletal mesh actor

diff --git a/Engine/Source/Editor/DetailCustomizations/Private/SkeletalMeshComponentDetails.cpp b/Engine/Source/Editor/DetailCustomizations/Private/SkeletalMeshComponentDetails.cpp
--- a/Engine/Source/Editor/DetailCustomizations/Private/SkeletalMeshComponentDetails.cpp
+++ b/Engine/Source/Editor/DetailCustomizations/Private/SkeletalMeshComponentDetails.cpp
@@ -217,8 +217,9 @@ void FSkeletalMeshComponentDetails::UpdateAnimationCategory(IDetailLayoutBuilder
 
 			TAttribute<bool> AnimPickerEnabledAttr(this, &FSkeletalMeshComponentDetails::AnimPickerIsEnabled);
 
-			AnimationCategory.AddCustomRow(ChildHandle->GetPropertyDisplayName())
+			AnimationCategory.AddProperty(ChildHandle)
 				.Visibility(SingleAnimVisibility)
+				.CustomWidget()
 				.IsEnabled(AnimPickerEnabledAttr)
 				.NameContent()
 				[
