Property Editor: Use property meta data for spin deltas

Related to this jira so that Sequence's key editors can have the same spin delta behavior

diff --git a/Engine/Source/Editor/DetailCustomizations/Private/ComponentTransformDetails.cpp b/Engine/Source/Editor/DetailCustomizations/Private/ComponentTransformDetails.cpp
--- a/Engine/Source/Editor/DetailCustomizations/Private/ComponentTransformDetails.cpp
+++ b/Engine/Source/Editor/DetailCustomizations/Private/ComponentTransformDetails.cpp
@@ -379,6 +379,21 @@ void FComponentTransformDetails::GenerateChildContent( IDetailChildrenBuilder& C
 	const bool bHideRotationField = ( HiddenFieldMask & ( 1 << ETransformField::Rotation ) ) != 0;
 	const bool bHideScaleField = ( HiddenFieldMask & ( 1 << ETransformField::Scale ) ) != 0;
 
+	IDetailLayoutBuilder& LayoutBuilder = ChildrenBuilder.GetParentCategory().GetParentLayout();
+	TSharedPtr<IPropertyHandle> LocationPropertyHandle = LayoutBuilder.GetProperty(USceneComponent::GetRelativeLocationPropertyName(), USceneComponent::StaticClass());
+	TSharedPtr<IPropertyHandle> RotationPropertyHandle = LayoutBuilder.GetProperty(USceneComponent::GetRelativeRotationPropertyName(), USceneComponent::StaticClass());
+	TSharedPtr<IPropertyHandle> ScalePropertyHandle = LayoutBuilder.GetProperty(USceneComponent::GetRelativeScale3DPropertyName(), USceneComponent::StaticClass());
+
+	const FString& MetaLocationDeltaString = LocationPropertyHandle->GetMetaData("Delta");
+	const FString& MetaRotationMinString = RotationPropertyHandle->GetMetaData("UIMin");
+	const FString& MetaRotationMaxString = RotationPropertyHandle->GetMetaData("UIMax");
+	const FString& MetaScaleDeltaString = ScalePropertyHandle->GetMetaData("Delta");
+
+	float LocationSpinDelta = !MetaLocationDeltaString.IsEmpty() ? FCString::Atof(*MetaLocationDeltaString) : 1.f;
+	float RotationMin = !MetaRotationMinString.IsEmpty() ? FCString::Atof(*MetaRotationMinString) : 0.f;
+	float RotationMax = !MetaRotationMaxString.IsEmpty() ? FCString::Atof(*MetaRotationMaxString) : 359.999f;
+	float ScaleSpinDelta = !MetaScaleDeltaString.IsEmpty() ? FCString::Atof(*MetaScaleDeltaString) : 0.0025f;
+
 	// Location
 	if(!bHideLocationField)
 	{
@@ -419,7 +434,7 @@ void FComponentTransformDetails::GenerateChildContent( IDetailChildrenBuilder& C
 			.Font(FontInfo)
 			.TypeInterface(TypeInterface)
 			.AllowSpin(SelectedObjects.Num() == 1)
-			.SpinDelta(1)
+			.SpinDelta(LocationSpinDelta)
 			.OnBeginSliderMovement(this, &FComponentTransformDetails::OnBeginLocationSlider)
 			.OnEndSliderMovement(this, &FComponentTransformDetails::OnEndLocationSlider)
 		];
@@ -452,6 +467,8 @@ void FComponentTransformDetails::GenerateChildContent( IDetailChildrenBuilder& C
 		[
 			SNew( SNumericRotatorInputBox<FRotator::FReal> )
 			.AllowSpin( SelectedObjects.Num() == 1 ) 
+			.MinSliderValue(RotationMin)
+			.MaxSliderValue(RotationMax)
 			.Roll( this, &FComponentTransformDetails::GetRotationX )
 			.Pitch( this, &FComponentTransformDetails::GetRotationY )
 			.Yaw( this, &FComponentTransformDetails::GetRotationZ )
@@ -506,7 +523,7 @@ void FComponentTransformDetails::GenerateChildContent( IDetailChildrenBuilder& C
 			.ContextMenuExtenderZ( this, &FComponentTransformDetails::ExtendZScaleContextMenu )
 			.Font( FontInfo )
 			.AllowSpin( SelectedObjects.Num() == 1 )
-			.SpinDelta( 0.0025f )
+			.SpinDelta( ScaleSpinDelta )
 			.OnBeginSliderMovement( this, &FComponentTransformDetails::OnBeginScaleSlider )
 			.OnEndSliderMovement(this, &FComponentTransformDetails::OnEndScaleSlider)
 		];
diff --git a/Engine/Source/Runtime/Engine/Classes/Components/SceneComponent.h b/Engine/Source/Runtime/Engine/Classes/Components/SceneComponent.h
--- a/Engine/Source/Runtime/Engine/Classes/Components/SceneComponent.h
+++ b/Engine/Source/Runtime/Engine/Classes/Components/SceneComponent.h
@@ -133,18 +133,18 @@ public:
 
 private:
 	/** Location of the component relative to its parent */
-	UPROPERTY(EditAnywhere, BlueprintReadOnly, ReplicatedUsing=OnRep_Transform, Category = Transform, meta=(AllowPrivateAccess="true"))
+	UPROPERTY(EditAnywhere, BlueprintReadOnly, ReplicatedUsing=OnRep_Transform, Category = Transform, meta=(AllowPrivateAccess="true", LinearDeltaSensitivity = "1", Delta = "1.0"))
 	FVector RelativeLocation;
 
 	/** Rotation of the component relative to its parent */
-	UPROPERTY(EditAnywhere, BlueprintReadOnly, ReplicatedUsing=OnRep_Transform, Category=Transform, meta=(AllowPrivateAccess="true"))
+	UPROPERTY(EditAnywhere, BlueprintReadOnly, ReplicatedUsing=OnRep_Transform, Category=Transform, meta=(AllowPrivateAccess="true", UIMin = "0.0", UIMax = "359.999"))
 	FRotator RelativeRotation;
 
 	/**
 	*	Non-uniform scaling of the component relative to its parent.
 	*	Note that scaling is always applied in local space (no shearing etc)
 	*/
-	UPROPERTY(EditAnywhere, BlueprintReadOnly, ReplicatedUsing=OnRep_Transform, Category=Transform, meta=(AllowPrivateAccess="true"))
+	UPROPERTY(EditAnywhere, BlueprintReadOnly, ReplicatedUsing=OnRep_Transform, Category=Transform, meta=(AllowPrivateAccess="true", LinearDeltaSensitivity = "1", Delta = "0.0025"))
 	FVector RelativeScale3D;
 
 public:
diff --git a/Engine/Source/Runtime/Slate/Public/Widgets/Input/SRotatorInputBox.h b/Engine/Source/Runtime/Slate/Public/Widgets/Input/SRotatorInputBox.h
--- a/Engine/Source/Runtime/Slate/Public/Widgets/Input/SRotatorInputBox.h
+++ b/Engine/Source/Runtime/Slate/Public/Widgets/Input/SRotatorInputBox.h
@@ -32,6 +32,8 @@ public:
 		: _bColorAxisLabels(false)
 		, _Font(FAppStyle::Get().GetFontStyle("NormalFont"))
 		, _AllowSpin(true)
+		, _MinSliderValue(0.0f)
+		, _MaxSliderValue(359.999f)
 		, _DisplayToggle( false )
 		, _TogglePitchChecked( ECheckBoxState::Checked )
 		, _ToggleYawChecked( ECheckBoxState::Checked )
@@ -63,6 +65,12 @@ public:
 		/** Whether or not values can be spun or if they should be typed in */
 		SLATE_ARGUMENT( bool, AllowSpin )
 
+		/** The minimum value that can be specified by using the slider */
+		SLATE_ATTRIBUTE( TOptional<NumericType>, MinSliderValue )
+
+		/** The maximum value that can be specified by using the slider */
+		SLATE_ATTRIBUTE( TOptional<NumericType>, MaxSliderValue )
+
 		/** Called when the pitch value is changed */
 		SLATE_EVENT( FOnNumericValueChanged, OnPitchChanged )
 
@@ -130,8 +138,8 @@ public:
 			[
 				SNew(SNumericEntryBox<NumericType>)
 				.AllowSpin(InArgs._AllowSpin)
-				.MinSliderValue(0.0f)
-				.MaxSliderValue(359.999f)
+				.MinSliderValue(InArgs._MinSliderValue)
+				.MaxSliderValue(InArgs._MaxSliderValue)
 				.LabelPadding(FMargin(3.0f))
 				.LabelLocation(SNumericEntryBox<NumericType>::ELabelLocation::Inside)
 				.Label()
@@ -162,8 +170,8 @@ public:
 			[
 				SNew(SNumericEntryBox<NumericType>)
 				.AllowSpin(InArgs._AllowSpin)
-				.MinSliderValue(0.0f)
-				.MaxSliderValue(359.999f)
+				.MinSliderValue(InArgs._MinSliderValue)
+				.MaxSliderValue(InArgs._MaxSliderValue)
 				.LabelPadding(FMargin(3.0f))
 				.LabelLocation(SNumericEntryBox<NumericType>::ELabelLocation::Inside)
 				.Label()
@@ -194,8 +202,8 @@ public:
 			[
 				SNew(SNumericEntryBox<NumericType>)
 				.AllowSpin(InArgs._AllowSpin)
-				.MinSliderValue(0.0f)
-				.MaxSliderValue(359.999f)
+				.MinSliderValue(InArgs._MinSliderValue)
+				.MaxSliderValue(InArgs._MaxSliderValue)
 				.LabelPadding(FMargin(3.0f))
 				.LabelLocation(SNumericEntryBox<NumericType>::ELabelLocation::Inside)
 				.Label()
diff --git a/Engine/Source/Runtime/Slate/Public/Widgets/Input/SVectorInputBox.h b/Engine/Source/Runtime/Slate/Public/Widgets/Input/SVectorInputBox.h
--- a/Engine/Source/Runtime/Slate/Public/Widgets/Input/SVectorInputBox.h
+++ b/Engine/Source/Runtime/Slate/Public/Widgets/Input/SVectorInputBox.h
@@ -176,6 +176,7 @@ public:
 			, _Font(FAppStyle::Get().GetFontStyle("NormalFont"))
 			, _AllowSpin(false)
 			, _SpinDelta(1)
+			, _LinearDeltaSensitivity(1)
 			, _bColorAxisLabels(false)
 			, _DisplayToggle(false)
 			, _TogglePadding(FMargin(1.f,0.f,1.f,0.f) )
@@ -217,6 +218,9 @@ public:
 		/** The delta amount to apply, per pixel, when the spinner is dragged. */
 		SLATE_ATTRIBUTE( NumericType, SpinDelta )
 
+		/** If we're an unbounded spinbox, what value do we divide mouse movement by before multiplying by Delta. Requires Delta to be set. */
+		SLATE_ATTRIBUTE( int32, LinearDeltaSensitivity )
+
 		/** Should the axis labels be colored */
 		SLATE_ARGUMENT( bool, bColorAxisLabels )		
 
@@ -424,7 +428,7 @@ private:
 			.MaxValue(CreatePerComponentGetter(ComponentIndex, TOptional<NumericType>(), InArgs._MaxVector))
 			.MinSliderValue(CreatePerComponentGetter(ComponentIndex, TOptional<NumericType>(), InArgs._MinSliderVector))
 			.MaxSliderValue(CreatePerComponentGetter(ComponentIndex, TOptional<NumericType>(), InArgs._MaxSliderVector))
-			.LinearDeltaSensitivity(1)
+			.LinearDeltaSensitivity(InArgs._LinearDeltaSensitivity)
 			.Delta(InArgs._SpinDelta)
 			.OnBeginSliderMovement(InArgs._OnBeginSliderMovement)
 			.OnEndSliderMovement(InArgs._OnEndSliderMovement)
