Enable filtering of mesh deformers to match primary component binding.
Was disabled because it can trigger mesh deformer asset sync load, which would trigger a slow graph compile.
The graph compilation is no longer done on PostLoad() so performance is now good.
Added specific filter metadata to the property in USkeletalMesh. This is needed for filtering with UObjects that aren't themselves primary binding types.

diff --git a/Engine/Source/Editor/DetailCustomizations/Private/MeshDeformerCustomizations.cpp b/Engine/Source/Editor/DetailCustomizations/Private/MeshDeformerCustomizations.cpp
--- a/Engine/Source/Editor/DetailCustomizations/Private/MeshDeformerCustomizations.cpp
+++ b/Engine/Source/Editor/DetailCustomizations/Private/MeshDeformerCustomizations.cpp
@@ -9,9 +9,40 @@
 #include "PropertyCustomizationHelpers.h"
 #include "PropertyHandle.h"
 
-
 namespace
 {
+	/** Pull the class type for filtering from the property. */
+	UClass* GetFilterClassFromProperty(TSharedRef<IPropertyHandle> InPropertyHandle)
+	{
+		TArray<UObject*> OwningObjects;
+		InPropertyHandle->GetOuterObjects(OwningObjects);
+		if (OwningObjects.Num() == 0)
+		{
+			return nullptr;
+		}
+
+		// Use the owner class by default.
+		UClass* FilterClass = OwningObjects[0]->GetClass();
+
+		// Override if we have metadata.
+		FProperty* DeformerProperty = InPropertyHandle->GetProperty();
+		static const FName FilterName(TEXT("Filter"));
+		const FString& FilterClassPath = DeformerProperty->GetMetaData(FilterName);
+		if (!FilterClassPath.IsEmpty())
+		{
+			FSoftClassPath FilterSoftClassPath(FilterClassPath);
+			FilterClass = FilterSoftClassPath.ResolveClass();
+		}
+
+		// Only UActorComponent classes are valid filters.
+		if (FilterClass != nullptr && !FilterClass->IsChildOf(UActorComponent::StaticClass()))
+		{
+			FilterClass = nullptr;
+		}
+
+		return FilterClass;
+	}
+
 	/** Pull the class type from the PrimaryBindingClass tag in some AssetData. */
 	UClass* GetPrimaryBindingClassFromAssetData(FAssetData const& AssetData)
 	{
@@ -37,20 +68,7 @@ void FMeshDeformerCustomization::CustomizeHeader(
 	FDetailWidgetRow& InHeaderRow,
 	IPropertyTypeCustomizationUtils& InCustomizationUtils)
 {
-	TArray<UObject*> OwningObjects;
-	InPropertyHandle->GetOuterObjects(OwningObjects);
-	if (OwningObjects.Num() == 0)
-	{
-		return;
-	}
-
-	UClass* OwnerClass = OwningObjects[0]->GetClass();
-
-	// Filter in USkeletalMesh context should assume USkinnedMeshComponent binding :(
-	if (OwnerClass == USkeletalMesh::StaticClass())
-	{
-		OwnerClass = USkinnedMeshComponent::StaticClass();
-	}
+	UClass* FilterClass = GetFilterClassFromProperty(InPropertyHandle);
 
 	InHeaderRow.NameContent()
 	[
@@ -61,12 +79,7 @@ void FMeshDeformerCustomization::CustomizeHeader(
 		SNew(SObjectPropertyEntryBox)
 		.PropertyHandle(InPropertyHandle)
 		.AllowedClass(UMeshDeformer::StaticClass())
-
-		// Disable filtering for now because of slow asset loading. 
-		// This is due to kernel compilation happening in PostLoad().
-		// Need to implement deferred compilation on demand and can then re-enable here.
-#if 0 
-		.OnShouldFilterAsset_Lambda([OwnerClass](FAssetData const& AssetData)
+		.OnShouldFilterAsset_Lambda([FilterClass](FAssetData const& AssetData)
 		{
 			// Filter depending on whether the PrimaryBindingClass matches our owning object.
 			// First try to load the class type from the asset registry.
@@ -83,12 +96,8 @@ void FMeshDeformerCustomization::CustomizeHeader(
 				}
 			}
 
-			const bool bFoundBindingClass = BindingClass != nullptr;
-			const bool bCanBind = OwnerClass != nullptr && BindingClass != nullptr && OwnerClass->IsChildOf(BindingClass);
-			const bool bHideEntry = bFoundBindingClass && !bCanBind;
-
+			const bool bHideEntry = FilterClass != nullptr && BindingClass != nullptr && !FilterClass->IsChildOf(BindingClass);
 			return bHideEntry;
 		})
-#endif
 	];
 }
diff --git a/Engine/Source/Runtime/Engine/Classes/Engine/SkeletalMesh.h b/Engine/Source/Runtime/Engine/Classes/Engine/SkeletalMesh.h
--- a/Engine/Source/Runtime/Engine/Classes/Engine/SkeletalMesh.h
+++ b/Engine/Source/Runtime/Engine/Classes/Engine/SkeletalMesh.h
@@ -2942,7 +2942,7 @@ protected:
 
 	/** Default mesh deformer to use with this mesh. */
 	UE_DEPRECATED(5.1, "This must be protected for async build, always use the accessors even internally.")
-	UPROPERTY(EditAnywhere, BlueprintReadOnly, Category = "Deformer")
+	UPROPERTY(EditAnywhere, BlueprintReadOnly, Category = "Deformer", Meta=(Filter="/Script/Engine.SkinnedMeshComponent"))
 	TObjectPtr<UMeshDeformer> DefaultMeshDeformer;
 
 private:
