PR #9580: Skeletal Mesh Anim Class button order (Contributed by [REDACTED_PERSON])
* Swapped ordering of 'Use Selected' and 'Browse to Asset' buttons to match the rest of the editor

diff --git a/Engine/Source/Editor/DetailCustomizations/Private/SkeletalMeshComponentDetails.cpp b/Engine/Source/Editor/DetailCustomizations/Private/SkeletalMeshComponentDetails.cpp
--- a/Engine/Source/Editor/DetailCustomizations/Private/SkeletalMeshComponentDetails.cpp
+++ b/Engine/Source/Editor/DetailCustomizations/Private/SkeletalMeshComponentDetails.cpp
@@ -173,7 +173,7 @@ void FSkeletalMeshComponentDetails::UpdateAnimationCategory(IDetailLayoutBuilder
 			.VAlign(VAlign_Center)
 			.Padding(2.0f, 1.0f)
 			[
-				PropertyCustomizationHelpers::MakeBrowseButton(FSimpleDelegate::CreateSP(this, &FSkeletalMeshComponentDetails::OnBrowseToAnimBlueprint))
+				PropertyCustomizationHelpers::MakeUseSelectedButton(FSimpleDelegate::CreateSP(this, &FSkeletalMeshComponentDetails::UseSelectedAnimBlueprint))
 			]
 			+ SHorizontalBox::Slot()
 			.AutoWidth()
@@ -181,7 +181,7 @@ void FSkeletalMeshComponentDetails::UpdateAnimationCategory(IDetailLayoutBuilder
 			.VAlign(VAlign_Center)
 			.Padding(2.0f, 1.0f)
 			[
-				PropertyCustomizationHelpers::MakeUseSelectedButton(FSimpleDelegate::CreateSP(this, &FSkeletalMeshComponentDetails::UseSelectedAnimBlueprint))
+				PropertyCustomizationHelpers::MakeBrowseButton(FSimpleDelegate::CreateSP(this, &FSkeletalMeshComponentDetails::OnBrowseToAnimBlueprint))
 			]
 		];
 
