Fixed random crash when closing a map with landscape proxy selected : the landscape actor can get unregistered by the destruction (ULandscapeInfo unregistering the actor) before the actor gets unselected, which can lead to FLandscapeProxyUIDetails updating a landscape proxy with no parent

#rnx

diff --git a/Engine/Source/Editor/DetailCustomizations/Private/LandscapeProxyUIDetails.cpp b/Engine/Source/Editor/DetailCustomizations/Private/LandscapeProxyUIDetails.cpp
--- a/Engine/Source/Editor/DetailCustomizations/Private/LandscapeProxyUIDetails.cpp
+++ b/Engine/Source/Editor/DetailCustomizations/Private/LandscapeProxyUIDetails.cpp
@@ -78,9 +78,10 @@ void FLandscapeProxyUIDetails::CustomizeDetails( IDetailLayoutBuilder& DetailBui
 	{
 		if (EditingProxy.IsValid())
 		{
-			ALandscape* LandscapeActor = EditingProxy->GetLandscapeActor();
-			check(LandscapeActor);
-			LandscapeActors.AddUnique(LandscapeActor);
+			if (ALandscape* LandscapeActor = EditingProxy->GetLandscapeActor())
+			{
+				LandscapeActors.AddUnique(LandscapeActor);
+			}
 		}
 	}
 
