Fix for SoundNodeParamCrossfade inheriting tooltip descriptions from SoundNodeDistanceCrossfade

diff --git a/Engine/Source/Editor/DetailCustomizations/Private/DetailCustomizations.cpp b/Engine/Source/Editor/DetailCustomizations/Private/DetailCustomizations.cpp
--- a/Engine/Source/Editor/DetailCustomizations/Private/DetailCustomizations.cpp
+++ b/Engine/Source/Editor/DetailCustomizations/Private/DetailCustomizations.cpp
@@ -142,6 +142,7 @@
 #include "SoftClassPathCustomization.h"
 #include "SoftObjectPathCustomization.h"
 #include "SoundBaseDetails.h"
+#include "Sound/SoundNodeDistanceCrossFade.h"
 #include "SoundSourceBusDetails.h"
 #include "SoundWaveDetails.h"
 #include "SourceCodeAccessSettingsDetails.h"
@@ -476,6 +477,8 @@ void FDetailCustomizationsModule::RegisterObjectCustomizations()
 	RegisterCustomClassLayout("LandscapeGrassType", FOnGetDetailCustomizationInstance::CreateStatic(&FLandscapeGrassTypeDetails::MakeInstance));
 
 	RegisterCustomClassLayout("BoundsCopyComponent", FOnGetDetailCustomizationInstance::CreateStatic(&FBoundsCopyComponentDetailsCustomization::MakeInstance));
+
+	RegisterCustomClassLayout("SoundNodeDistanceCrossFade", FOnGetDetailCustomizationInstance::CreateStatic(&FCrossFadeCustomization::MakeInstance));
 }
 
 #define LOCTEXT_NAMESPACE "DetailsSections"
diff --git a/Engine/Source/Editor/DetailCustomizations/Private/DistanceDatumStructCustomization.cpp b/Engine/Source/Editor/DetailCustomizations/Private/DistanceDatumStructCustomization.cpp
--- a/Engine/Source/Editor/DetailCustomizations/Private/DistanceDatumStructCustomization.cpp
+++ b/Engine/Source/Editor/DetailCustomizations/Private/DistanceDatumStructCustomization.cpp
@@ -3,6 +3,7 @@
 #include "DistanceDatumStructCustomization.h"
 
 #include "Containers/Array.h"
+#include "DetailLayoutBuilder.h"
 #include "DetailWidgetRow.h"
 #include "HAL/Platform.h"
 #include "IDetailChildrenBuilder.h"
@@ -23,14 +24,14 @@
 
 TSharedRef<IPropertyTypeCustomization> FDistanceDatumStructCustomization::MakeInstance()
 {
-	return MakeShareable(new FDistanceDatumStructCustomization);
+	return MakeShared<FDistanceDatumStructCustomization>(FDistanceDatumStructCustomization::FPrivateToken());
 }
 
 FDistanceDatumStructCustomization::~FDistanceDatumStructCustomization()
 {
 }
 
-FDistanceDatumStructCustomization::FDistanceDatumStructCustomization()
+FDistanceDatumStructCustomization::FDistanceDatumStructCustomization(FPrivateToken)
 {
 }
 
@@ -88,22 +89,55 @@ void FDistanceDatumStructCustomization::CustomizeChildren(TSharedRef<IPropertyHa
 		{
 			if (Child->GetProperty()->GetFName() == GET_MEMBER_NAME_CHECKED(FDistanceDatum, FadeInDistanceStart))
 			{
-				ChildRow.DisplayName(FText::Format(LOCTEXT("FadeInStart", "Fade In {0} Value Start"), ParamDesc));
+				ChildRow.DisplayName(FText::Format(LOCTEXT("FadeInStart", "Fade In {0} Start"), ParamDesc));
+				ChildRow.ToolTip(LOCTEXT("FadeInStartText", "The param value at which to start hearing this sound."));
 			}
 			else if (Child->GetProperty()->GetFName() == GET_MEMBER_NAME_CHECKED(FDistanceDatum, FadeInDistanceEnd))
 			{
-				ChildRow.DisplayName(FText::Format(LOCTEXT("FadeInEnd", "Fade In {0} Value End"), ParamDesc));
+				ChildRow.DisplayName(FText::Format(LOCTEXT("FadeInEnd", "Fade In {0} End"), ParamDesc));
+				ChildRow.ToolTip(LOCTEXT("FadeInEndText", "The param value at which this sound has faded in completely."));
 			}
 			else if (Child->GetProperty()->GetFName() == GET_MEMBER_NAME_CHECKED(FDistanceDatum, FadeOutDistanceStart))
 			{
-				ChildRow.DisplayName(FText::Format(LOCTEXT("FadeOutStart", "Fade Out {0} Value Start"), ParamDesc));
+				ChildRow.DisplayName(FText::Format(LOCTEXT("FadeOutStart", "Fade Out {0} Start"), ParamDesc));
+				ChildRow.ToolTip(LOCTEXT("FadeOutStartText", "The param value at which this sound starts fading out"));
 			}
 			else if (Child->GetProperty()->GetFName() == GET_MEMBER_NAME_CHECKED(FDistanceDatum, FadeOutDistanceEnd))
 			{
-				ChildRow.DisplayName(FText::Format(LOCTEXT("FadeOutEnd", "Fade Out {0} Value End"), ParamDesc));
+				ChildRow.DisplayName(FText::Format(LOCTEXT("FadeOutEnd", "Fade Out {0} End"), ParamDesc));
+				ChildRow.ToolTip(LOCTEXT("FadeOutEndText", "The param value at which this sound is no longer audible."));
 			}
 		}
 	}
 }
 
+void FCrossFadeCustomization::CustomizeDetails(IDetailLayoutBuilder& DetailBuilder)
+{
+	
+	TArray< TWeakObjectPtr< UObject > > Objects;
+	DetailBuilder.GetObjectsBeingCustomized(Objects);
+
+	//Restrict this to one object at a time
+	if (Objects.Num() != 1)
+	{
+		return;
+	}
+
+	//We only need to change the tooltips if we're working with a SoundNodeParamCrossFade
+	if (Objects[0]->IsA(USoundNodeParamCrossFade::StaticClass()))
+	{
+		TSharedRef< IPropertyHandle > CrossFadeHandle = DetailBuilder.GetProperty(GET_MEMBER_NAME_CHECKED(USoundNodeDistanceCrossFade, CrossFadeInput));
+		ensure(CrossFadeHandle->IsValidHandle());
+
+		CrossFadeHandle->SetToolTipText(LOCTEXT("CrossFadeInputText", "Each input needs to have the correct data filled in so the SoundNodeParamCrossFade is able to determine which sounds to play"));
+	}
+
+}
+
+TSharedRef< IDetailCustomization > FCrossFadeCustomization::MakeInstance()
+{
+	return MakeShared<FCrossFadeCustomization>();
+}
+
+
 #undef LOCTEXT_NAMESPACE
diff --git a/Engine/Source/Editor/DetailCustomizations/Private/DistanceDatumStructCustomization.h b/Engine/Source/Editor/DetailCustomizations/Private/DistanceDatumStructCustomization.h
--- a/Engine/Source/Editor/DetailCustomizations/Private/DistanceDatumStructCustomization.h
+++ b/Engine/Source/Editor/DetailCustomizations/Private/DistanceDatumStructCustomization.h
@@ -4,6 +4,7 @@
 
 #include "IPropertyTypeCustomization.h"
 #include "Templates/SharedPointer.h"
+#include "IDetailCustomization.h"
 
 class FDetailWidgetRow;
 class IDetailChildrenBuilder;
@@ -14,6 +15,9 @@ class IPropertyHandle;
  */
 class FDistanceDatumStructCustomization : public IPropertyTypeCustomization
 {
+private:
+	struct FPrivateToken { explicit FPrivateToken() = default; };
+
 public:
 	static TSharedRef<IPropertyTypeCustomization> MakeInstance();
 
@@ -26,9 +30,18 @@ public:
 	virtual void CustomizeHeader( TSharedRef<IPropertyHandle> InStructPropertyHandle, FDetailWidgetRow& HeaderRow, IPropertyTypeCustomizationUtils& StructCustomizationUtils ) override;
 	virtual void CustomizeChildren( TSharedRef<IPropertyHandle> InStructPropertyHandle, IDetailChildrenBuilder& StructBuilder, IPropertyTypeCustomizationUtils& StructCustomizationUtils ) override;
 
-private:
 	/**
 	 * Constructor
 	 */
-	FDistanceDatumStructCustomization();
+	explicit FDistanceDatumStructCustomization(FPrivateToken);
 };
+
+class FCrossFadeCustomization : public IDetailCustomization
+{
+public:
+	// IDetailCustomization interface
+	virtual void CustomizeDetails(IDetailLayoutBuilder& DetailBuilder) override;
+	//
+
+	static TSharedRef< IDetailCustomization > MakeInstance();
+};
\ No newline at end of file
