Added meta-data to force text keys to be unstable (when edited) on a per-property basis

"UniqueIdWhenEditedOnDefault" to make the key unstable on a CDO/archetype
"UniqueIdWhenEditedOnInstance" to make the key unstable on an instance

#rnx

diff --git a/Engine/Source/Editor/DetailCustomizations/Private/TextCustomization.cpp b/Engine/Source/Editor/DetailCustomizations/Private/TextCustomization.cpp
--- a/Engine/Source/Editor/DetailCustomizations/Private/TextCustomization.cpp
+++ b/Engine/Source/Editor/DetailCustomizations/Private/TextCustomization.cpp
@@ -102,12 +102,35 @@ namespace
 		{
 			if (PropertyHandle->IsValidHandle())
 			{
-				TArray<UPackage*> PropertyPackages;
-				PropertyHandle->GetOuterPackages(PropertyPackages);
+				static const FName NAME_UniqueIdWhenEditedOnDefault  = "UniqueIdWhenEditedOnDefault";
+				static const FName NAME_UniqueIdWhenEditedOnInstance = "UniqueIdWhenEditedOnInstance";
 
-				check(PropertyPackages.IsValidIndex(InIndex));
+				UPackage* PropertyPackage = nullptr;
+				bool bForceUniqueId = false;
 
-				StaticStableTextId(PropertyPackages[InIndex], InEditAction, InTextSource, InProposedNamespace, InProposedKey, OutStableNamespace, OutStableKey);
+				// We can't query if this property handle is under an object or an external struct, so try and get the objects first before falling back to using the packages
+				// Note: We want to use the object where possible so that we can tell if we're editing a CDO/archetype or an instance
+				{
+					TArray<UObject*> PropertyObjects;
+					PropertyHandle->GetOuterObjects(PropertyObjects);
+
+					if (PropertyObjects.IsValidIndex(InIndex))
+					{
+						PropertyPackage = PropertyObjects[InIndex]->GetPackage();
+						bForceUniqueId = PropertyHandle->HasMetaData(PropertyObjects[InIndex]->HasAnyFlags(RF_ClassDefaultObject | RF_ArchetypeObject) ? NAME_UniqueIdWhenEditedOnDefault : NAME_UniqueIdWhenEditedOnInstance);
+					}
+				}
+				if (!PropertyPackage)
+				{
+					TArray<UPackage*> PropertyPackages;
+					PropertyHandle->GetOuterPackages(PropertyPackages);
+
+					check(PropertyPackages.IsValidIndex(InIndex));
+					PropertyPackage = PropertyPackages[InIndex];
+				}
+
+				check(PropertyPackage);
+				StaticStableTextId(PropertyPackage, InEditAction, InTextSource, InProposedNamespace, bForceUniqueId ? FString() : InProposedKey, OutStableNamespace, OutStableKey);
 			}
 		}
 #endif // USE_STABLE_LOCALIZATION_KEYS
