Fix detail panel not showing properly until scrolled if it uses a lot of call-in-editor methods.

diff --git a/Engine/Source/Editor/DetailCustomizations/Private/ObjectDetails.cpp b/Engine/Source/Editor/DetailCustomizations/Private/ObjectDetails.cpp
--- a/Engine/Source/Editor/DetailCustomizations/Private/ObjectDetails.cpp
+++ b/Engine/Source/Editor/DetailCustomizations/Private/ObjectDetails.cpp
@@ -162,7 +162,16 @@ void FObjectDetails::AddCallInEditorMethods(IDetailLayoutBuilder& DetailBuilder)
 			FCategoryEntry(FName InCategoryName)
 				: CategoryName(InCategoryName)
 			{
-				WrapBox = SNew(SWrapBox).UseAllottedSize(true);
+				WrapBox = SNew(SWrapBox)
+					// Setting the preferred size here (despite using UseAllottedSize) is a workaround for an issue
+					// when contained in a scroll box: prior to the first tick, the wrap box will use preferred size
+					// instead of allotted, and if preferred size is set small, it will cause the box to wrap a lot and
+					// request too much space from the scroll box. On next tick, SWrapBox is updated but the scroll box
+					// does not realize that it needs to show more elements, until it is scrolled.
+					// Setting a large value here means that the SWrapBox will request too little space prior to tick,
+					// which will cause the scroll box to virtualize more elements at the start, but this is less broken.
+					.PreferredSize(2000)
+					.UseAllottedSize(true);
 			}
 		};
 
