Downgraded an assert to an ensure

#rnx

diff --git a/Engine/Source/Editor/DetailCustomizations/Private/TextCustomization.cpp b/Engine/Source/Editor/DetailCustomizations/Private/TextCustomization.cpp
--- a/Engine/Source/Editor/DetailCustomizations/Private/TextCustomization.cpp
+++ b/Engine/Source/Editor/DetailCustomizations/Private/TextCustomization.cpp
@@ -129,7 +129,7 @@ namespace
 					PropertyPackage = PropertyPackages[InIndex];
 				}
 
-				check(PropertyPackage);
+				ensureAlwaysMsgf(PropertyPackage, TEXT("A package must be available for key stabilization to work correctly. Did you forget to set the package on a FStructOnScope?"));
 				StaticStableTextId(PropertyPackage, InEditAction, InTextSource, InProposedNamespace, bForceUniqueId ? FString() : InProposedKey, OutStableNamespace, OutStableKey);
 			}
 		}
