Remove unused code.

diff --git a/Engine/Source/Editor/DetailCustomizations/Private/PrimitiveComponentDetails.h b/Engine/Source/Editor/DetailCustomizations/Private/PrimitiveComponentDetails.h
--- a/Engine/Source/Editor/DetailCustomizations/Private/PrimitiveComponentDetails.h
+++ b/Engine/Source/Editor/DetailCustomizations/Private/PrimitiveComponentDetails.h
@@ -33,10 +33,6 @@ private:
 	void AddPhysicsCategory(IDetailLayoutBuilder& DetailBuilder);
 	void AddCollisionCategory(IDetailLayoutBuilder& DetailBuilder);
 
-	ECheckBoxState IsMobilityActive(TSharedRef<IPropertyHandle> MobilityHandle, EComponentMobility::Type InMobility) const;
-
-	void OnMobilityChanged(ECheckBoxState InCheckedState, TSharedRef<IPropertyHandle> MobilityHandle, EComponentMobility::Type InMobility);
-
 	void AddAdvancedSubCategory(IDetailLayoutBuilder& DetailBuilder, FName MainCategory, FName SubCategory);
 
 	FReply OnMobilityResetClicked(TSharedRef<IPropertyHandle> MobilityHandle);
