Clean up the General Section Tab for Post Process Volumes in the details panel. Rendering Features (Post process material etc.) and PostProcessVolumeSettings (Unbound Extents, Enabled etc.) for now as a start. Removed Brush Settings.

diff --git a/Engine/Source/Editor/DetailCustomizations/Private/DetailCustomizations.cpp b/Engine/Source/Editor/DetailCustomizations/Private/DetailCustomizations.cpp
--- a/Engine/Source/Editor/DetailCustomizations/Private/DetailCustomizations.cpp
+++ b/Engine/Source/Editor/DetailCustomizations/Private/DetailCustomizations.cpp
@@ -814,6 +814,16 @@ void FDetailCustomizationsModule::RegisterSectionMappings()
 			Section->AddCategory("ChaosPhysics");
 		}
 	}
+
+	// Post Process Volume
+	{
+		{
+			TSharedRef<FPropertySection> Section = PropertyModule.FindOrCreateSection("PostProcessVolume", "General", LOCTEXT("General", "General"));
+			Section->AddCategory("PostProcessVolumeSettings");
+			Section->AddCategory("Rendering Features");
+			Section->RemoveCategory("Brush Settings");
+		}
+	}
 }
 
 #undef LOCTEXT_NAMESPACE
