Option to hide Static mobility on lights when project has static lighting disabled.
- controlled via cvar (defaults to existing behaviour)

diff --git a/Engine/Source/Editor/DetailCustomizations/Private/MobilityCustomization.cpp b/Engine/Source/Editor/DetailCustomizations/Private/MobilityCustomization.cpp
--- a/Engine/Source/Editor/DetailCustomizations/Private/MobilityCustomization.cpp
+++ b/Engine/Source/Editor/DetailCustomizations/Private/MobilityCustomization.cpp
@@ -22,6 +22,11 @@
 #include "Widgets/Input/SSegmentedControl.h"
 #include "Widgets/Text/STextBlock.h"
 
+TAutoConsoleVariable<bool> CVarHideLightStaticMobilityWhenStaticLightingDisabled(
+	TEXT("r.Editor.HideLightStaticMobilityWhenStaticLightingDisabled"),
+	false,
+	TEXT("Hide Static mobility on Light components when project has Static Lighting disabled."),
+	ECVF_RenderThreadSafe);
 
 #define LOCTEXT_NAMESPACE "MobilityCustomization" 
 
@@ -67,6 +72,14 @@ void FMobilityCustomization::GenerateHeaderRowContent(FDetailWidgetRow& WidgetRo
 	bool bShowStatic = !( RestrictedMobilityBits & StaticMobilityBitMask );
 	bool bShowStationary = !( RestrictedMobilityBits & StationaryMobilityBitMask );
 
+	static const auto AllowStaticLightingVar = IConsoleManager::Get().FindTConsoleVariableDataInt(TEXT("r.AllowStaticLighting"));
+	const bool bAllowStaticLighting = (!AllowStaticLightingVar || AllowStaticLightingVar->GetValueOnGameThread() != 0);
+
+	if (bForLight && CVarHideLightStaticMobilityWhenStaticLightingDisabled.GetValueOnGameThread())
+	{
+		bShowStatic &= bAllowStaticLighting;
+	}
+
 	int32 ColumnIndex = 0;
 
 	if ( bShowStatic )
