Change local light unit to be displayed next to the intensity.

diff --git a/Engine/Source/Editor/DetailCustomizations/Private/LightComponentDetails.cpp b/Engine/Source/Editor/DetailCustomizations/Private/LightComponentDetails.cpp
--- a/Engine/Source/Editor/DetailCustomizations/Private/LightComponentDetails.cpp
+++ b/Engine/Source/Editor/DetailCustomizations/Private/LightComponentDetails.cpp
@@ -10,6 +10,7 @@
 #include "Delegates/Delegate.h"
 #include "DetailCategoryBuilder.h"
 #include "DetailLayoutBuilder.h"
+#include "Engine/Scene.h"
 #include "IDetailPropertyRow.h"
 #include "Internationalization/Internationalization.h"
 #include "Internationalization/Text.h"
@@ -22,6 +23,7 @@
 #include "UObject/Object.h"
 #include "UObject/UnrealType.h"
 #include "UObject/WeakObjectPtr.h"
+#include "DetailWidgetRow.h"
 
 #define LOCTEXT_NAMESPACE "LightComponentDetails"
 
@@ -30,6 +32,106 @@ TSharedRef<IDetailCustomization> FLightComponentDetails::MakeInstance()
 	return MakeShareable( new FLightComponentDetails );
 }
 
+void FLightComponentDetails::AddLocalLightIntensityWithUnit(IDetailLayoutBuilder& DetailBuilder, ULocalLightComponent* Component)
+{
+	float ConversionFactor = 1.f;
+	uint8 Value = 0; // Unitless
+	if (IntensityUnitsProperty->GetValue(Value) == FPropertyAccess::Success)
+	{
+		ConversionFactor = ULocalLightComponent::GetUnitsConversionFactor((ELightUnits)0, (ELightUnits)Value);
+	}
+	IntensityUnitsProperty->SetOnPropertyValuePreChange(FSimpleDelegate::CreateSP(this, &FLightComponentDetails::OnIntensityUnitsPreChange, Component));
+	IntensityUnitsProperty->SetOnPropertyValueChanged(FSimpleDelegate::CreateSP(this, &FLightComponentDetails::OnIntensityUnitsChanged, Component));
+
+	// Inverse squared falloff point lights (the default) are in units of lumens, instead of just being a brightness scale
+	if (Component->IntensityUnits == ELightUnits::EV)
+	{
+		// +/-32 stops give a large enough range
+		LightIntensityProperty->SetInstanceMetaData("UIMin",TEXT("-32.0f"));
+		LightIntensityProperty->SetInstanceMetaData("UIMax",TEXT("32.0f"));
+	}
+	else
+	{
+		LightIntensityProperty->SetInstanceMetaData("UIMin",TEXT("0.0f"));
+		LightIntensityProperty->SetInstanceMetaData("UIMax",  *FString::SanitizeFloat(100000.0f * ConversionFactor));
+		LightIntensityProperty->SetInstanceMetaData("SliderExponent", TEXT("2.0f"));
+	}
+	if (Component->IntensityUnits == ELightUnits::Lumens)
+	{
+		LightIntensityProperty->SetInstanceMetaData("Units", TEXT("lm"));
+		LightIntensityProperty->SetToolTipText(LOCTEXT("LightIntensityInLumensToolTipText", "Luminous power or flux in lumens"));
+	}
+	else if (Component->IntensityUnits == ELightUnits::Candelas)
+	{
+		LightIntensityProperty->SetInstanceMetaData("Units", TEXT("cd"));
+		LightIntensityProperty->SetToolTipText(LOCTEXT("LightIntensityInCandelasToolTipText", "Luminous intensity in candelas"));
+	}
+	else if (Component->IntensityUnits == ELightUnits::EV)
+	{
+		LightIntensityProperty->SetInstanceMetaData("Units", TEXT("ev"));
+		LightIntensityProperty->SetToolTipText(LOCTEXT("LightIntensityInCandelasToolTipText", "Luminous intensity in EV100"));
+	}
+
+	// Make these come first
+	IDetailCategoryBuilder& LightCategory = DetailBuilder.EditCategory( "Light", FText::GetEmpty(), ECategoryPriority::TypeSpecific );
+
+	// Target version have intensity value and unity sitting side by side. Currently there is a 
+	// bug causing the reset button to not reset the intensity value, which is why this version 
+	// is not enabled by default for now
+#if 0
+	// Add light intensity and unit on the same line
+	FDetailWidgetRow& IntensityAndUnitRow = DetailBuilder.AddCustomRowToCategory(LightIntensityProperty, LightIntensityProperty->GetPropertyDisplayName())
+	.NameContent()
+	[
+		LightIntensityProperty->CreatePropertyNameWidget()
+	]
+	.ValueContent()
+	[
+		SNew(SHorizontalBox)
+		+ SHorizontalBox::Slot()
+		.AutoWidth()
+		.VAlign(VAlign_Center)
+		.Padding(2.f, 2.f, 10.f, 2.f)
+		[
+			LightIntensityProperty->CreatePropertyValueWidget()
+		]
+		+ SHorizontalBox::Slot()
+		.AutoWidth()
+		[
+			IntensityUnitsProperty->CreatePropertyValueWidget()
+		]
+	];
+
+	if (!IESBrightnessEnabledProperty->IsValidHandle())
+	{
+		TWeakObjectPtr<ULightComponent> BaseComponent = Component;
+		IntensityAndUnitRow
+			.OverrideResetToDefault(FResetToDefaultOverride::Create(FIsResetToDefaultVisible::CreateSP(this, &FLightComponentDetails::IsIntensityResetToDefaultVisible, BaseComponent), FResetToDefaultHandler::CreateSP(this, &FLightComponentDetails::ResetIntensityToDefault, BaseComponent)));
+	}
+	else
+	{
+		TWeakObjectPtr<ULightComponent> BaseComponent = Component;
+		IntensityAndUnitRow
+			.IsEnabled(TAttribute<bool>(this, &FLightComponentDetails::IsLightBrightnessEnabled))
+			.OverrideResetToDefault(FResetToDefaultOverride::Create(FIsResetToDefaultVisible::CreateSP(this, &FLightComponentDetails::IsIntensityResetToDefaultVisible, BaseComponent), FResetToDefaultHandler::CreateSP(this, &FLightComponentDetails::ResetIntensityToDefault, BaseComponent)));
+	}
+#else
+	// Add light intensity and light unit onto two different lines
+	if( !IESBrightnessEnabledProperty->IsValidHandle() )
+	{
+		LightCategory.AddProperty( LightIntensityProperty );
+	}
+	else
+	{
+		TWeakObjectPtr<ULightComponent> BaseComponent = Component;
+		LightCategory.AddProperty( LightIntensityProperty )
+			.IsEnabled( TAttribute<bool>( this, &FLightComponentDetails::IsLightBrightnessEnabled ) )
+			.OverrideResetToDefault(FResetToDefaultOverride::Create(FIsResetToDefaultVisible::CreateSP(this, &FLightComponentDetails::IsIntensityResetToDefaultVisible, BaseComponent), FResetToDefaultHandler::CreateSP(this, &FLightComponentDetails::ResetIntensityToDefault, BaseComponent)));
+	}
+	LightCategory.AddProperty(IntensityUnitsProperty).OverrideResetToDefault(FResetToDefaultOverride::Create(FIsResetToDefaultVisible::CreateSP(this, &FLightComponentDetails::IsIntensityUnitsResetToDefaultVisible, Component), FResetToDefaultHandler::CreateSP(this, &FLightComponentDetails::ResetIntensityUnitsToDefault, Component)));
+#endif
+}
+
 void FLightComponentDetails::CustomizeDetails( IDetailLayoutBuilder& DetailBuilder )
 {
 	TArray<TWeakObjectPtr<UObject>> Objects;
@@ -48,6 +150,7 @@ void FLightComponentDetails::CustomizeDetails( IDetailLayoutBuilder& DetailBuild
 	DetailBuilder.EditCategory("Rendering", FText::GetEmpty(), ECategoryPriority::TypeSpecific);
 
 	LightIntensityProperty = DetailBuilder.GetProperty(GET_MEMBER_NAME_CHECKED(ULightComponent, Intensity), ULightComponentBase::StaticClass());
+	IntensityUnitsProperty = DetailBuilder.GetProperty(GET_MEMBER_NAME_CHECKED(ULocalLightComponent, IntensityUnits), ULocalLightComponent::StaticClass());
 	IESBrightnessTextureProperty = DetailBuilder.GetProperty(GET_MEMBER_NAME_CHECKED(ULightComponent, IESTexture));
 	IESBrightnessEnabledProperty = DetailBuilder.GetProperty(GET_MEMBER_NAME_CHECKED(ULightComponent, bUseIESBrightness));
 	IESBrightnessScaleProperty = DetailBuilder.GetProperty(GET_MEMBER_NAME_CHECKED(ULightComponent, IESBrightnessScale));
@@ -55,19 +158,32 @@ void FLightComponentDetails::CustomizeDetails( IDetailLayoutBuilder& DetailBuild
 	if( !IESBrightnessEnabledProperty->IsValidHandle() )
 	{
 		// Brightness and color should be listed first
-		LightCategory.AddProperty( LightIntensityProperty );
+		if (ULocalLightComponent* LocalComponent = Cast<ULocalLightComponent>(Objects[0].Get()))
+		{
+			AddLocalLightIntensityWithUnit(DetailBuilder, LocalComponent);
+		}
+		else
+		{
+			LightCategory.AddProperty( LightIntensityProperty );
+		}
 		LightCategory.AddProperty( DetailBuilder.GetProperty(GET_MEMBER_NAME_CHECKED(ULightComponent, LightColor), ULightComponentBase::StaticClass() ) );
 	}
 	else
 	{
-		IDetailCategoryBuilder& LightProfilesCategory = DetailBuilder.EditCategory( "Light Profiles", FText::GetEmpty(), ECategoryPriority::Default );
-
-		LightCategory.AddProperty( LightIntensityProperty )
-			.IsEnabled( TAttribute<bool>( this, &FLightComponentDetails::IsLightBrightnessEnabled ) )
-			.OverrideResetToDefault(FResetToDefaultOverride::Create(FIsResetToDefaultVisible::CreateSP(this, &FLightComponentDetails::IsIntensityResetToDefaultVisible, Component), FResetToDefaultHandler::CreateSP(this, &FLightComponentDetails::ResetIntensityToDefault, Component)));
+		if (ULocalLightComponent* LocalComponent = Cast<ULocalLightComponent>(Objects[0].Get()))
+		{
+			AddLocalLightIntensityWithUnit(DetailBuilder, LocalComponent);
+		}
+		else
+		{
+			LightCategory.AddProperty( LightIntensityProperty )
+				.IsEnabled( TAttribute<bool>( this, &FLightComponentDetails::IsLightBrightnessEnabled ) )
+				.OverrideResetToDefault(FResetToDefaultOverride::Create(FIsResetToDefaultVisible::CreateSP(this, &FLightComponentDetails::IsIntensityResetToDefaultVisible, Component), FResetToDefaultHandler::CreateSP(this, &FLightComponentDetails::ResetIntensityToDefault, Component)));
+		}
 
 		LightCategory.AddProperty( DetailBuilder.GetProperty("LightColor", ULightComponentBase::StaticClass() ) );
 
+		IDetailCategoryBuilder& LightProfilesCategory = DetailBuilder.EditCategory( "Light Profiles", FText::GetEmpty(), ECategoryPriority::Default );
 		LightProfilesCategory.AddProperty( IESBrightnessTextureProperty );
 
 		LightProfilesCategory.AddProperty( IESBrightnessEnabledProperty )
@@ -152,4 +268,100 @@ bool FLightComponentDetails::IsIntensityResetToDefaultVisible(TSharedPtr<IProper
 	}
 }
 
+void FLightComponentDetails::CustomizeDetails(const TSharedPtr<IDetailLayoutBuilder>& DetailBuilder)
+{
+	CachedDetailBuilder = DetailBuilder;
+	CustomizeDetails(*DetailBuilder);
+}
+
+void FLightComponentDetails::OnIntensityUnitsPreChange(ULocalLightComponent* Component)
+{
+	if (Component)
+	{
+		LastLightBrigtness = Component->ComputeLightBrightness();
+	}
+}
+
+void FLightComponentDetails::OnIntensityUnitsChanged(ULocalLightComponent* Component)
+{
+	// Convert the brightness using the new units.
+	if (Component)
+	{
+		FLightComponentDetails::SetComponentIntensity(Component, LastLightBrigtness);
+	}
+
+	// Here we can only take the ptr as ForceRefreshDetails() checks that the reference is unique.
+	IDetailLayoutBuilder* DetailBuilder = CachedDetailBuilder.Pin().Get();
+	if (DetailBuilder)
+	{
+		DetailBuilder->ForceRefreshDetails();
+	}
+}
+
+namespace
+{
+	void SetComponentIntensityUnits(ULocalLightComponent* Component, ELightUnits InUnits)
+	{
+		check(Component);
+
+		FProperty* IntensityUnitsProperty = FindFieldChecked<FProperty>(ULocalLightComponent::StaticClass(), GET_MEMBER_NAME_CHECKED(ULocalLightComponent, IntensityUnits));
+		FPropertyChangedEvent PropertyChangedEvent(IntensityUnitsProperty);
+
+		const ELightUnits PreviousUnits = Component->IntensityUnits;
+		Component->IntensityUnits = InUnits;
+		Component->PostEditChangeProperty(PropertyChangedEvent);
+		Component->MarkRenderStateDirty();
+
+		// Propagate changes to instances.
+		TArray<UObject*> Instances;
+		Component->GetArchetypeInstances(Instances);
+		for (UObject* Instance : Instances)
+		{
+			ULocalLightComponent* InstanceComponent = Cast<ULocalLightComponent>(Instance);
+			if (InstanceComponent && InstanceComponent->IntensityUnits == PreviousUnits)
+			{
+				InstanceComponent->IntensityUnits = Component->IntensityUnits;
+				InstanceComponent->PostEditChangeProperty(PropertyChangedEvent);
+				InstanceComponent->MarkRenderStateDirty();
+			}
+		}
+	}
+}
+
+void FLightComponentDetails::ResetIntensityUnitsToDefault(TSharedPtr<IPropertyHandle> PropertyHandle, ULocalLightComponent* Component)
+{
+	// Actors (and blueprints) spawned from the actor factory inherit the intensity units from the project settings.
+	if (Component && Component->GetArchetype() && !Component->GetArchetype()->IsInBlueprint())
+	{
+		static const auto CVarDefaultLightUnits = IConsoleManager::Get().FindTConsoleVariableDataInt(TEXT("r.DefaultFeature.LightUnits"));
+		const ELightUnits DefaultUnits = (ELightUnits)CVarDefaultLightUnits->GetValueOnGameThread();
+
+		if (DefaultUnits != Component->IntensityUnits)
+		{
+			SetComponentIntensityUnits(Component, DefaultUnits);
+		}
+	}
+	else
+	{
+		// Fall back to default handler. 
+		PropertyHandle->ResetToDefault();
+	}
+}
+
+bool FLightComponentDetails::IsIntensityUnitsResetToDefaultVisible(TSharedPtr<IPropertyHandle> PropertyHandle, ULocalLightComponent* Component) const
+{
+	// Actors (and blueprints) spawned from the actor factory inherit the project settings.
+	if (Component && Component->GetArchetype() && !Component->GetArchetype()->IsInBlueprint())
+	{
+		static const auto CVarDefaultLightUnits = IConsoleManager::Get().FindTConsoleVariableDataInt(TEXT("r.DefaultFeature.LightUnits"));
+		const ELightUnits DefaultUnits = (ELightUnits)CVarDefaultLightUnits->GetValueOnGameThread();
+		return DefaultUnits != Component->IntensityUnits;
+	}
+	else
+	{
+		// Fall back to default handler
+		return PropertyHandle->DiffersFromDefault();
+	}
+}
+
 #undef LOCTEXT_NAMESPACE
diff --git a/Engine/Source/Editor/DetailCustomizations/Private/LightComponentDetails.h b/Engine/Source/Editor/DetailCustomizations/Private/LightComponentDetails.h
--- a/Engine/Source/Editor/DetailCustomizations/Private/LightComponentDetails.h
+++ b/Engine/Source/Editor/DetailCustomizations/Private/LightComponentDetails.h
@@ -9,6 +9,7 @@
 class IDetailLayoutBuilder;
 class IPropertyHandle;
 class ULightComponent;
+class ULocalLightComponent;
 
 class FLightComponentDetails : public IDetailCustomization
 {
@@ -18,6 +19,7 @@ public:
 
 	/** IDetailCustomization interface */
 	virtual void CustomizeDetails( IDetailLayoutBuilder& DetailBuilder ) override;
+	virtual void CustomizeDetails(const TSharedPtr<IDetailLayoutBuilder>& DetailBuilder) override;
 
 	static void SetComponentIntensity(ULightComponent* Component, float InIntensity);
 
@@ -36,4 +38,18 @@ private:
 	TSharedPtr<IPropertyHandle> IESBrightnessEnabledProperty;
 	TSharedPtr<IPropertyHandle> IESBrightnessScaleProperty;
 	TSharedPtr<IPropertyHandle> LightIntensityProperty;
+	TSharedPtr<IPropertyHandle> IntensityUnitsProperty;
+
+	TWeakPtr<IDetailLayoutBuilder> CachedDetailBuilder;
+	void ResetIntensityUnitsToDefault(TSharedPtr<IPropertyHandle> PropertyHandle, ULocalLightComponent* Component);
+	bool IsIntensityUnitsResetToDefaultVisible(TSharedPtr<IPropertyHandle> PropertyHandle, ULocalLightComponent* Component) const;
+
+	/** Called when the intensity units are changed */
+	void OnIntensityUnitsPreChange(ULocalLightComponent* Component);
+	void OnIntensityUnitsChanged(ULocalLightComponent* Component);
+
+	/* Add local light intensity with unit */
+	void AddLocalLightIntensityWithUnit(IDetailLayoutBuilder& DetailBuilder, ULocalLightComponent* Component);
+
+	float LastLightBrigtness = 0;
 };
diff --git a/Engine/Source/Editor/DetailCustomizations/Private/LocalLightComponentDetails.cpp b/Engine/Source/Editor/DetailCustomizations/Private/LocalLightComponentDetails.cpp
--- a/Engine/Source/Editor/DetailCustomizations/Private/LocalLightComponentDetails.cpp
+++ b/Engine/Source/Editor/DetailCustomizations/Private/LocalLightComponentDetails.cpp
@@ -38,152 +38,9 @@ void FLocalLightComponentDetails::CustomizeDetails( IDetailLayoutBuilder& Detail
 	DetailBuilder.GetObjectsBeingCustomized(Objects);
 	ULocalLightComponent* Component = Cast<ULocalLightComponent>(Objects[0].Get());
 
-	TSharedPtr<IPropertyHandle> LightIntensityProperty = DetailBuilder.GetProperty(GET_MEMBER_NAME_CHECKED(ULightComponentBase, Intensity), ULightComponentBase::StaticClass());
-	TSharedPtr<IPropertyHandle> IntensityUnitsProperty = DetailBuilder.GetProperty(GET_MEMBER_NAME_CHECKED(ULocalLightComponent, IntensityUnits), ULocalLightComponent::StaticClass());
-	// DetailBuilder.HideProperty(LightIntensityProperty);
-
-	float ConversionFactor = 1.f;
-	uint8 Value = 0; // Unitless
-	if (IntensityUnitsProperty->GetValue(Value) == FPropertyAccess::Success)
-	{
-		ConversionFactor = ULocalLightComponent::GetUnitsConversionFactor((ELightUnits)0, (ELightUnits)Value);
-	}
-	IntensityUnitsProperty->SetOnPropertyValuePreChange(FSimpleDelegate::CreateSP(this, &FLocalLightComponentDetails::OnIntensityUnitsPreChange, Component));
-	IntensityUnitsProperty->SetOnPropertyValueChanged(FSimpleDelegate::CreateSP(this, &FLocalLightComponentDetails::OnIntensityUnitsChanged, Component));
-
-	// Inverse squared falloff point lights (the default) are in units of lumens, instead of just being a brightness scale
-	if (Component->IntensityUnits == ELightUnits::EV)
-	{
-		// +/-32 stops give a large enough range
-		LightIntensityProperty->SetInstanceMetaData("UIMin",TEXT("-32.0f"));
-		LightIntensityProperty->SetInstanceMetaData("UIMax",TEXT("32.0f"));
-	}
-	else
-	{
-		LightIntensityProperty->SetInstanceMetaData("UIMin",TEXT("0.0f"));
-		LightIntensityProperty->SetInstanceMetaData("UIMax",  *FString::SanitizeFloat(100000.0f * ConversionFactor));
-		LightIntensityProperty->SetInstanceMetaData("SliderExponent", TEXT("2.0f"));
-	}
-	if (Component->IntensityUnits == ELightUnits::Lumens)
-	{
-		LightIntensityProperty->SetInstanceMetaData("Units", TEXT("lm"));
-		LightIntensityProperty->SetToolTipText(LOCTEXT("LightIntensityInLumensToolTipText", "Luminous power or flux in lumens"));
-	}
-	else if (Component->IntensityUnits == ELightUnits::Candelas)
-	{
-		LightIntensityProperty->SetInstanceMetaData("Units", TEXT("cd"));
-		LightIntensityProperty->SetToolTipText(LOCTEXT("LightIntensityInCandelasToolTipText", "Luminous intensity in candelas"));
-	}
-	else if (Component->IntensityUnits == ELightUnits::EV)
-	{
-		LightIntensityProperty->SetInstanceMetaData("Units", TEXT("ev"));
-		LightIntensityProperty->SetToolTipText(LOCTEXT("LightIntensityInCandelasToolTipText", "Luminous intensity in EV100"));
-	}
-
-	// Make these come first
-	IDetailCategoryBuilder& LightCategory = DetailBuilder.EditCategory( "Light", FText::GetEmpty(), ECategoryPriority::TypeSpecific );
-
-	LightCategory.AddProperty(IntensityUnitsProperty)
-		.OverrideResetToDefault(FResetToDefaultOverride::Create(FIsResetToDefaultVisible::CreateSP(this, &FLocalLightComponentDetails::IsIntensityUnitsResetToDefaultVisible, Component), FResetToDefaultHandler::CreateSP(this, &FLocalLightComponentDetails::ResetIntensityUnitsToDefault, Component)));
-
+	IDetailCategoryBuilder& LightCategory = DetailBuilder.EditCategory("Light", FText::GetEmpty(), ECategoryPriority::TypeSpecific);
 	LightCategory.AddProperty( DetailBuilder.GetProperty( GET_MEMBER_NAME_CHECKED(ULocalLightComponent, AttenuationRadius), ULocalLightComponent::StaticClass() ) );
 
 }
 
-void FLocalLightComponentDetails::CustomizeDetails(const TSharedPtr<IDetailLayoutBuilder>& DetailBuilder)
-{
-	CachedDetailBuilder = DetailBuilder;
-	CustomizeDetails(*DetailBuilder);
-}
-
-void FLocalLightComponentDetails::OnIntensityUnitsPreChange(ULocalLightComponent* Component)
-{
-	if (Component)
-	{
-		LastLightBrigtness = Component->ComputeLightBrightness();
-	}
-}
-
-void FLocalLightComponentDetails::OnIntensityUnitsChanged(ULocalLightComponent* Component)
-{
-	// Convert the brightness using the new units.
-	if (Component)
-	{
-		FLightComponentDetails::SetComponentIntensity(Component, LastLightBrigtness);
-	}
-
-	// Here we can only take the ptr as ForceRefreshDetails() checks that the reference is unique.
-	IDetailLayoutBuilder* DetailBuilder = CachedDetailBuilder.Pin().Get();
-	if (DetailBuilder)
-	{
-		DetailBuilder->ForceRefreshDetails();
-	}
-}
-
-namespace
-{
-	void SetComponentIntensityUnits(ULocalLightComponent* Component, ELightUnits InUnits)
-	{
-		check(Component);
-
-		FProperty* IntensityUnitsProperty = FindFieldChecked<FProperty>(ULocalLightComponent::StaticClass(), GET_MEMBER_NAME_CHECKED(ULocalLightComponent, IntensityUnits));
-		FPropertyChangedEvent PropertyChangedEvent(IntensityUnitsProperty);
-
-		const ELightUnits PreviousUnits = Component->IntensityUnits;
-		Component->IntensityUnits = InUnits;
-		Component->PostEditChangeProperty(PropertyChangedEvent);
-		Component->MarkRenderStateDirty();
-
-		// Propagate changes to instances.
-		TArray<UObject*> Instances;
-		Component->GetArchetypeInstances(Instances);
-		for (UObject* Instance : Instances)
-		{
-			ULocalLightComponent* InstanceComponent = Cast<ULocalLightComponent>(Instance);
-			if (InstanceComponent && InstanceComponent->IntensityUnits == PreviousUnits)
-			{
-				InstanceComponent->IntensityUnits = Component->IntensityUnits;
-				InstanceComponent->PostEditChangeProperty(PropertyChangedEvent);
-				InstanceComponent->MarkRenderStateDirty();
-			}
-		}
-	}
-}
-
-void FLocalLightComponentDetails::ResetIntensityUnitsToDefault(TSharedPtr<IPropertyHandle> PropertyHandle, ULocalLightComponent* Component)
-{
-	// Actors (and blueprints) spawned from the actor factory inherit the intensity units from the project settings.
-	if (Component && Component->GetArchetype() && !Component->GetArchetype()->IsInBlueprint())
-	{
-		static const auto CVarDefaultLightUnits = IConsoleManager::Get().FindTConsoleVariableDataInt(TEXT("r.DefaultFeature.LightUnits"));
-		const ELightUnits DefaultUnits = (ELightUnits)CVarDefaultLightUnits->GetValueOnGameThread();
-
-		if (DefaultUnits != Component->IntensityUnits)
-		{
-			SetComponentIntensityUnits(Component, DefaultUnits);
-		}
-	}
-	else
-	{
-		// Fall back to default handler. 
-		PropertyHandle->ResetToDefault();
-	}
-}
-
-bool FLocalLightComponentDetails::IsIntensityUnitsResetToDefaultVisible(TSharedPtr<IPropertyHandle> PropertyHandle, ULocalLightComponent* Component) const
-{
-	// Actors (and blueprints) spawned from the actor factory inherit the project settings.
-	if (Component && Component->GetArchetype() && !Component->GetArchetype()->IsInBlueprint())
-	{
-		static const auto CVarDefaultLightUnits = IConsoleManager::Get().FindTConsoleVariableDataInt(TEXT("r.DefaultFeature.LightUnits"));
-		const ELightUnits DefaultUnits = (ELightUnits)CVarDefaultLightUnits->GetValueOnGameThread();
-		return DefaultUnits != Component->IntensityUnits;
-	}
-	else
-	{
-		// Fall back to default handler
-		return PropertyHandle->DiffersFromDefault();
-	}
-}
-
 #undef LOCTEXT_NAMESPACE
diff --git a/Engine/Source/Editor/DetailCustomizations/Private/LocalLightComponentDetails.h b/Engine/Source/Editor/DetailCustomizations/Private/LocalLightComponentDetails.h
--- a/Engine/Source/Editor/DetailCustomizations/Private/LocalLightComponentDetails.h
+++ b/Engine/Source/Editor/DetailCustomizations/Private/LocalLightComponentDetails.h
@@ -12,27 +12,9 @@ class ULocalLightComponent;
 class FLocalLightComponentDetails : public IDetailCustomization
 {
 public:
-	FLocalLightComponentDetails() : LastLightBrigtness(0) {}
-
 	/** Makes a new instance of this detail layout class for a specific detail view requesting it */
 	static TSharedRef<IDetailCustomization> MakeInstance();
 
 	/** IDetailCustomization interface */
 	virtual void CustomizeDetails( IDetailLayoutBuilder& DetailBuilder ) override;
-	virtual void CustomizeDetails( const TSharedPtr<IDetailLayoutBuilder>& DetailBuilder ) override;
-
-protected:
-
-	// The detail builder for this cusomtomisation
-	TWeakPtr<IDetailLayoutBuilder> CachedDetailBuilder;
-
-	void ResetIntensityUnitsToDefault(TSharedPtr<IPropertyHandle> PropertyHandle, ULocalLightComponent* Component);
-	bool IsIntensityUnitsResetToDefaultVisible(TSharedPtr<IPropertyHandle> PropertyHandle, ULocalLightComponent* Component) const;
-
-
-	/** Called when the intensity units are changed */
-	void OnIntensityUnitsPreChange(ULocalLightComponent* Component);
-	void OnIntensityUnitsChanged(ULocalLightComponent* Component);
-
-	float LastLightBrigtness;
 };
diff --git a/Engine/Source/Runtime/Engine/Classes/Components/LocalLightComponent.h b/Engine/Source/Runtime/Engine/Classes/Components/LocalLightComponent.h
--- a/Engine/Source/Runtime/Engine/Classes/Components/LocalLightComponent.h
+++ b/Engine/Source/Runtime/Engine/Classes/Components/LocalLightComponent.h
@@ -23,7 +23,7 @@ class ENGINE_API ULocalLightComponent : public ULightComponent
 	 * The peak luminous intensity is measured in candelas,
 	 * while the luminous power is measured in lumens.
 	 */
-	UPROPERTY(EditAnywhere, BlueprintReadOnly, Category=Light, AdvancedDisplay, meta=(DisplayName="Intensity Units", EditCondition="!bUseIESBrightness"))
+	UPROPERTY(EditAnywhere, BlueprintReadOnly, Category=Light, meta=(DisplayName="Intensity Units", EditCondition="!bUseIESBrightness"))
 	ELightUnits IntensityUnits;
 
 	/**
