Skeletal Mesh - Fixed a crash in FSkeletalMeshReductionSettingsDetails::CustomizeChildren() when the object is a Cloth Asset.

#rnx

diff --git a/Engine/Source/Editor/DetailCustomizations/Private/SkeletalMeshReductionSettingsDetails.cpp b/Engine/Source/Editor/DetailCustomizations/Private/SkeletalMeshReductionSettingsDetails.cpp
--- a/Engine/Source/Editor/DetailCustomizations/Private/SkeletalMeshReductionSettingsDetails.cpp
+++ b/Engine/Source/Editor/DetailCustomizations/Private/SkeletalMeshReductionSettingsDetails.cpp
@@ -107,7 +107,7 @@ void FSkeletalMeshReductionSettingsDetails::CustomizeChildren(TSharedRef<IProper
 		// Only able to do this for LOD1 and above, so only show the property if this is the case
 		if (LODIndex > 0)
 		{
-			bool AllowInline = !SkeletalMesh->IsLODImportedDataEmpty(LODIndex);
+			bool AllowInline = SkeletalMesh && !SkeletalMesh->IsLODImportedDataEmpty(LODIndex);
 			// Add and retrieve the default widgets
 			IDetailPropertyRow& Row = StructBuilder.AddProperty(BaseLODPropertyHandle->AsShared());
 
