Fix CIS issue.

diff --git a/Engine/Source/Editor/DetailCustomizations/Private/LightComponentDetails.cpp b/Engine/Source/Editor/DetailCustomizations/Private/LightComponentDetails.cpp
--- a/Engine/Source/Editor/DetailCustomizations/Private/LightComponentDetails.cpp
+++ b/Engine/Source/Editor/DetailCustomizations/Private/LightComponentDetails.cpp
@@ -69,7 +69,7 @@ void FLightComponentDetails::AddLocalLightIntensityWithUnit(IDetailLayoutBuilder
 	else if (Component->IntensityUnits == ELightUnits::EV)
 	{
 		LightIntensityProperty->SetInstanceMetaData("Units", TEXT("ev"));
-		LightIntensityProperty->SetToolTipText(LOCTEXT("LightIntensityInCandelasToolTipText", "Luminous intensity in EV100"));
+		LightIntensityProperty->SetToolTipText(LOCTEXT("LightIntensityInEVToolTipText", "Luminous intensity in EV100"));
 	}
 
 	// Make these come first
