Support parsing drop frame timecode from FText

diff --git a/Engine/Source/Editor/DetailCustomizations/Private/TimecodeDetailsCustomization.cpp b/Engine/Source/Editor/DetailCustomizations/Private/TimecodeDetailsCustomization.cpp
--- a/Engine/Source/Editor/DetailCustomizations/Private/TimecodeDetailsCustomization.cpp
+++ b/Engine/Source/Editor/DetailCustomizations/Private/TimecodeDetailsCustomization.cpp
@@ -82,8 +82,28 @@ void FTimecodeDetailsCustomization::OnTimecodeTextCommitted(const FText& InText,
 			TimecodeProperty->NotifyPreChange();
 			((FTimecode*)RawData[0])->Hours = FCString::Atoi(*Splits[0]);
 			((FTimecode*)RawData[0])->Minutes = NumSplits > 1 ? FCString::Atoi(*Splits[1]) : 0;
-			((FTimecode*)RawData[0])->Seconds = NumSplits > 2 ? FCString::Atoi(*Splits[2]) : 0;
-			((FTimecode*)RawData[0])->Frames = NumSplits > 3 ? FCString::Atoi(*Splits[3]) : 0;
+
+			if (NumSplits > 2)
+			{
+				TArray<FString> SubSplits;
+				Splits[2].ParseIntoArray(SubSplits, TEXT(";"));
+
+				const int NumSubSplits = SubSplits.Num();
+
+				((FTimecode*) RawData[0])->bDropFrameFormat = NumSubSplits != 1 ? true : false;
+
+				if (((FTimecode*) RawData[0])->bDropFrameFormat)
+				{
+					((FTimecode*) RawData[0])->Seconds = FCString::Atoi(*SubSplits[0]);
+					((FTimecode*) RawData[0])->Frames = NumSubSplits > 1 ? FCString::Atoi(*SubSplits[1]) : 0;
+				}
+				else
+				{
+					((FTimecode*) RawData[0])->Seconds = FCString::Atoi(*Splits[2]);
+					((FTimecode*) RawData[0])->Frames = NumSplits > 3 ? FCString::Atoi(*Splits[3]) : 0;
+				}
+			}
+			
 			TimecodeProperty->NotifyPostChange(EPropertyChangeType::ValueSet);
 			TimecodeProperty->NotifyFinishedChangingProperties();
 
