Added RowTag to the custom rows created by the DataTableRowHandle customization, so the custom rows are associated with the property. This fixes an issue where the custom rows created by the DataTableRowHandle customization weren't appearing when a DataTableRowHandle property was being explicitly shown in a DetailsView with the PropertyToShow filtering, because the custom rows didn't match the underlying property name

#rnx

diff --git a/Engine/Source/Editor/DetailCustomizations/Private/DataTableCustomization.cpp b/Engine/Source/Editor/DetailCustomizations/Private/DataTableCustomization.cpp
--- a/Engine/Source/Editor/DetailCustomizations/Private/DataTableCustomization.cpp
+++ b/Engine/Source/Editor/DetailCustomizations/Private/DataTableCustomization.cpp
@@ -58,6 +58,7 @@ void FDataTableCustomizationLayout::CustomizeChildren(TSharedRef<class IProperty
 	/** Get all the existing property handles */
 	DataTablePropertyHandle = InStructPropertyHandle->GetChildHandle("DataTable");
 	RowNamePropertyHandle = InStructPropertyHandle->GetChildHandle("RowName");
+	const FName PropertyName = InStructPropertyHandle->GetProperty()->GetFName();
 
 	if (DataTablePropertyHandle->IsValidHandle() && RowNamePropertyHandle->IsValidHandle())
 	{
@@ -67,6 +68,7 @@ void FDataTableCustomizationLayout::CustomizeChildren(TSharedRef<class IProperty
 
 		/** Construct a asset picker widget with a custom filter */
 		StructBuilder.AddCustomRow(LOCTEXT("DataTable_TableName", "Data Table"))
+			.RowTag(PropertyName)
 			.NameContent()
 			[
 				SNew(STextBlock)
@@ -89,6 +91,7 @@ void FDataTableCustomizationLayout::CustomizeChildren(TSharedRef<class IProperty
 
 		/** Construct a combo box widget to select from a list of valid options */
 		StructBuilder.AddCustomRow(LOCTEXT("DataTable_RowName", "Row Name"))
+			.RowTag(PropertyName)
 			.NameContent()
 			[
 				SNew(STextBlock)
