Optimization for selecting large numbers of objects with BodyInstanceCustomizations for DetailView
* CollisionResponse values are quite expensive to validate each SDetailView::Tick - accounting for the majority of the cost of validating StaticMeshActor objects
* Can safely ignore validation of the CollisionResponse field because it is either set via the details view itself when using a Custom preset OR is set via the Project Settings when using a preset.  If using a preset, invalidation occurs which updates the DetailView.

diff --git a/Engine/Source/Editor/DetailCustomizations/Private/BodyInstanceCustomization.cpp b/Engine/Source/Editor/DetailCustomizations/Private/BodyInstanceCustomization.cpp
--- a/Engine/Source/Editor/DetailCustomizations/Private/BodyInstanceCustomization.cpp
+++ b/Engine/Source/Editor/DetailCustomizations/Private/BodyInstanceCustomization.cpp
@@ -133,6 +133,12 @@ void FBodyInstanceCustomization::AddCollisionCategory(TSharedRef<class IProperty
 	check (CollisionProfileNameHandle.IsValid());
 	check (CollisionEnabledHandle.IsValid());
 	check (ObjectTypeHandle.IsValid());
+	check (CollisionResponsesHandle.IsValid());
+	
+	// Expensive to validate
+    // Presets update the DetailView on event, cannot be changed when custom from scripts so is safe
+    // to skip validation each from for.
+    CollisionResponsesHandle->SetIgnoreValidation(true);
 
 	// need to find profile name
 	FName ProfileName;
