PropertyEditor no longer hides the "Materials" category if it contains other properties

diff --git a/Engine/Source/Editor/DetailCustomizations/Private/ComponentMaterialCategory.cpp b/Engine/Source/Editor/DetailCustomizations/Private/ComponentMaterialCategory.cpp
--- a/Engine/Source/Editor/DetailCustomizations/Private/ComponentMaterialCategory.cpp
+++ b/Engine/Source/Editor/DetailCustomizations/Private/ComponentMaterialCategory.cpp
@@ -205,13 +205,26 @@ void FComponentMaterialCategory::Create( IDetailLayoutBuilder& DetailBuilder )
 		}
 	}
 
-
 	// Make a category for the materials.
 	MaterialCategory = &DetailBuilder.EditCategory("Materials", FText::GetEmpty(), ECategoryPriority::TypeSpecific );
 
-	MaterialCategory->AddCustomBuilder( MaterialList );
-	MaterialCategory->SetCategoryVisibility( bAnyMaterialsToDisplay );
+	if (bAnyMaterialsToDisplay)
+	{
+		MaterialCategory->AddCustomBuilder( MaterialList );
+		MaterialCategory->SetCategoryVisibility( bAnyMaterialsToDisplay );
+		return;
+	}
+
+	// Check again - if there are any properties in the category, we don't want to hide it
+	TArray<TSharedRef<IPropertyHandle>> DefaultMaterialProperties;
+	MaterialCategory->GetDefaultProperties(DefaultMaterialProperties);
+	
+	if (!DefaultMaterialProperties.IsEmpty())
+	{
+		bAnyMaterialsToDisplay = true;
+	}
 	
+	MaterialCategory->SetCategoryVisibility( bAnyMaterialsToDisplay );
 }
 
 void FComponentMaterialCategory::OnGetMaterialsForView( IMaterialListBuilder& MaterialList )
@@ -247,6 +260,21 @@ void FComponentMaterialCategory::OnGetMaterialsForView( IMaterialListBuilder& Ma
 		}
 	}
 
+	if (bAnyMaterialsToDisplay)
+	{
+		MaterialCategory->SetCategoryVisibility(true);
+		return;
+	}
+
+	// Check again - if there are any properties in the category, we don't want to hide it
+	TArray<TSharedRef<IPropertyHandle>> DefaultMaterialProperties;
+	MaterialCategory->GetDefaultProperties(DefaultMaterialProperties);
+
+	if (!DefaultMaterialProperties.IsEmpty())
+	{
+		bAnyMaterialsToDisplay = true;
+	}
+
 	MaterialCategory->SetCategoryVisibility(bAnyMaterialsToDisplay);
 }
 
