Made FKeyStructCustomization non-copyable to prevent deprecation warnings in compiler-generated special member functions.

diff --git a/Engine/Source/Editor/DetailCustomizations/Private/KeyStructCustomization.cpp b/Engine/Source/Editor/DetailCustomizations/Private/KeyStructCustomization.cpp
--- a/Engine/Source/Editor/DetailCustomizations/Private/KeyStructCustomization.cpp
+++ b/Engine/Source/Editor/DetailCustomizations/Private/KeyStructCustomization.cpp
@@ -23,6 +23,8 @@ class SWidget;
 
 #define LOCTEXT_NAMESPACE "FKeyStructCustomization"
 
+FKeyStructCustomization::FKeyStructCustomization() = default;
+
 /* FKeyStructCustomization static interface
  *****************************************************************************/
 
diff --git a/Engine/Source/Editor/DetailCustomizations/Public/KeyStructCustomization.h b/Engine/Source/Editor/DetailCustomizations/Public/KeyStructCustomization.h
--- a/Engine/Source/Editor/DetailCustomizations/Public/KeyStructCustomization.h
+++ b/Engine/Source/Editor/DetailCustomizations/Public/KeyStructCustomization.h
@@ -22,6 +22,10 @@ class DETAILCUSTOMIZATIONS_API FKeyStructCustomization
 	: public IPropertyTypeCustomization
 {
 public:
+	FKeyStructCustomization();
+
+	UE_NONCOPYABLE(FKeyStructCustomization)
+
 	// IPropertyTypeCustomization interface
 
 	virtual void CustomizeHeader(TSharedRef<class IPropertyHandle> StructPropertyHandle, class FDetailWidgetRow& HeaderRow, IPropertyTypeCustomizationUtils& StructCustomizationUtils) override;
