BodyInstance : Updated BodyInstanceCustomization to allow derived components to hide unused attributes.

diff --git a/Engine/Source/Editor/DetailCustomizations/Private/BodyInstanceCustomization.cpp b/Engine/Source/Editor/DetailCustomizations/Private/BodyInstanceCustomization.cpp
--- a/Engine/Source/Editor/DetailCustomizations/Private/BodyInstanceCustomization.cpp
+++ b/Engine/Source/Editor/DetailCustomizations/Private/BodyInstanceCustomization.cpp
@@ -7,6 +7,7 @@
 #include "Components/SkeletalMeshComponent.h"
 #include "Components/StaticMeshComponent.h"
 #include "Containers/EnumAsByte.h"
+#include "DeformableInterface.h"
 #include "DestructibleInterface.h"
 #include "DetailCategoryBuilder.h"
 #include "DetailLayoutBuilder.h"
@@ -1260,6 +1261,16 @@ void FBodyInstanceCustomizationHelper::UpdateFilters()
 	bDisplayConstraints = true;
 	bDisplayEnablePhysics = true;
 	bDisplayAsyncScene = true;
+	bDisplayLinearDamping = true;
+	bDisplayAngularDamping = true;
+	bDisplayEnableGravity = true;
+	bDisplayInertiaConditioning = true;
+	bDisplayWalkableSlopeOverride = true;
+	bDisplayAutoWeld = true;
+	bDisplayStartAwake = true;
+	bDisplayCOMNudge = true;
+	bDisplayMassScale = true;
+	bDisplayMaxAngularVelocity = true;
 
 	for (int32 i = 0; i < ObjectsCustomized.Num(); ++i)
 	{
@@ -1270,6 +1281,23 @@ void FBodyInstanceCustomizationHelper::UpdateFilters()
 				bDisplayMass = false;
 				bDisplayConstraints = false;
 			}
+			else if (Cast<IDeformableInterface>(ObjectsCustomized[i].Get()))
+			{
+				bDisplayMass = false;
+				bDisplayConstraints = false;
+				bDisplayEnablePhysics = false;
+				bDisplayAsyncScene = false;
+				bDisplayLinearDamping = false;
+				bDisplayAngularDamping = false;
+				bDisplayEnableGravity = false;
+				bDisplayInertiaConditioning = false;
+				bDisplayWalkableSlopeOverride = false;
+				bDisplayAutoWeld = false;
+				bDisplayStartAwake = false;
+				bDisplayCOMNudge = false;
+				bDisplayMassScale = false;
+				bDisplayMaxAngularVelocity = false;
+			}
 			else
 			{
 				if(ObjectsCustomized[i]->IsA(UBodySetup::StaticClass()))
@@ -1308,9 +1336,25 @@ void FBodyInstanceCustomizationHelper::CustomizeDetails( IDetailLayoutBuilder& D
 
 		AddMassInKg(PhysicsCategory, BodyInstanceHandler);
 
-		PhysicsCategory.AddProperty(BodyInstanceHandler->GetChildHandle(GET_MEMBER_NAME_CHECKED(FBodyInstance, LinearDamping)));
-		PhysicsCategory.AddProperty(BodyInstanceHandler->GetChildHandle(GET_MEMBER_NAME_CHECKED(FBodyInstance, AngularDamping)));
-		PhysicsCategory.AddProperty(BodyInstanceHandler->GetChildHandle(GET_MEMBER_NAME_CHECKED(FBodyInstance, bEnableGravity)));
+
+		auto EnablePhysicsProperty = [BodyInstanceHandler, &PhysicsCategory](FName PropertyName, bool bEnable)
+		{
+			TSharedRef<IPropertyHandle> Property = BodyInstanceHandler->GetChildHandle(PropertyName).ToSharedRef();
+			if (bEnable)
+			{
+				PhysicsCategory.AddProperty(Property);
+			}
+			else
+			{
+				Property->MarkHiddenByCustomization();
+			}
+			return Property;
+		};
+		EnablePhysicsProperty(GET_MEMBER_NAME_CHECKED(FBodyInstance, LinearDamping), bDisplayLinearDamping);
+		EnablePhysicsProperty(GET_MEMBER_NAME_CHECKED(FBodyInstance, AngularDamping), bDisplayAngularDamping);
+		EnablePhysicsProperty(GET_MEMBER_NAME_CHECKED(FBodyInstance, bEnableGravity), bDisplayEnableGravity);
+		EnablePhysicsProperty(GET_MEMBER_NAME_CHECKED(FBodyInstance, bInertiaConditioning), bDisplayInertiaConditioning);
+		EnablePhysicsProperty(GET_MEMBER_NAME_CHECKED(FBodyInstance, WalkableSlopeOverride), bDisplayWalkableSlopeOverride);
 
 		AddBodyConstraint(PhysicsCategory, BodyInstanceHandler);
 
@@ -1318,12 +1362,14 @@ void FBodyInstanceCustomizationHelper::CustomizeDetails( IDetailLayoutBuilder& D
 		PhysicsCategory.AddProperty(BodyInstanceHandler->GetChildHandle(GET_MEMBER_NAME_CHECKED(FBodyInstance, bAutoWeld)))
 			.Visibility(TAttribute<EVisibility>(this, &FBodyInstanceCustomizationHelper::IsAutoWeldVisible));
 
-		PhysicsCategory.AddProperty(BodyInstanceHandler->GetChildHandle(GET_MEMBER_NAME_CHECKED(FBodyInstance, bStartAwake)));
+		EnablePhysicsProperty(GET_MEMBER_NAME_CHECKED(FBodyInstance, bStartAwake), bDisplayStartAwake);
+		EnablePhysicsProperty(GET_MEMBER_NAME_CHECKED(FBodyInstance, COMNudge), bDisplayCOMNudge);
+		EnablePhysicsProperty(GET_MEMBER_NAME_CHECKED(FBodyInstance, MassScale), bDisplayMassScale);
 
-		PhysicsCategory.AddProperty(BodyInstanceHandler->GetChildHandle(GET_MEMBER_NAME_CHECKED(FBodyInstance, COMNudge)));
-		PhysicsCategory.AddProperty(BodyInstanceHandler->GetChildHandle(GET_MEMBER_NAME_CHECKED(FBodyInstance, MassScale)));
-
-		AddMaxAngularVelocity(PhysicsCategory, BodyInstanceHandler);
+		if (bDisplayMaxAngularVelocity)
+		{
+			AddMaxAngularVelocity(PhysicsCategory, BodyInstanceHandler);
+		}
 
 		// Hide legacy settings in Chaos
 		BodyInstanceHandler->GetChildHandle(GET_MEMBER_NAME_CHECKED(FBodyInstance, DOFMode))->MarkHiddenByCustomization();
diff --git a/Engine/Source/Editor/DetailCustomizations/Private/BodyInstanceCustomization.h b/Engine/Source/Editor/DetailCustomizations/Private/BodyInstanceCustomization.h
--- a/Engine/Source/Editor/DetailCustomizations/Private/BodyInstanceCustomization.h
+++ b/Engine/Source/Editor/DetailCustomizations/Private/BodyInstanceCustomization.h
@@ -189,6 +189,16 @@ private:
 	bool bDisplayConstraints;
 	bool bDisplayEnablePhysics;
 	bool bDisplayAsyncScene;
+	bool bDisplayLinearDamping;
+	bool bDisplayAngularDamping;
+	bool bDisplayEnableGravity;
+	bool bDisplayInertiaConditioning;
+	bool bDisplayWalkableSlopeOverride;
+	bool bDisplayAutoWeld;
+	bool bDisplayStartAwake;
+	bool bDisplayCOMNudge;
+	bool bDisplayMassScale;
+	bool bDisplayMaxAngularVelocity;
 
 	TSharedPtr<IPropertyHandle> MassInKgOverrideHandle;
 	TSharedPtr<IPropertyHandle> DOFModeProperty;
diff --git a/Engine/Source/Runtime/Engine/Private/DeformableInterface.cpp b/Engine/Source/Runtime/Engine/Private/DeformableInterface.cpp
new file mode 100644
--- /dev/null
+++ b/Engine/Source/Runtime/Engine/Private/DeformableInterface.cpp
@@ -0,0 +1,12 @@
+// Copyright Epic Games, Inc. All Rights Reserved.
+
+#include "DeformableInterface.h"
+
+#include UE_INLINE_GENERATED_CPP_BY_NAME(DeformableInterface)
+
+UDeformableInterface::UDeformableInterface(const FObjectInitializer& ObjectInitializer)
+	: Super(ObjectInitializer)
+{
+
+}
+
diff --git a/Engine/Source/Runtime/Engine/Public/DeformableInterface.h b/Engine/Source/Runtime/Engine/Public/DeformableInterface.h
new file mode 100644
--- /dev/null
+++ b/Engine/Source/Runtime/Engine/Public/DeformableInterface.h
@@ -0,0 +1,16 @@
+// Copyright Epic Games, Inc. All Rights Reserved.
+
+#pragma once
+#include "UObject/Interface.h"
+#include "DeformableInterface.generated.h"
+
+UINTERFACE()
+class ENGINE_API UDeformableInterface : public UInterface
+{
+	GENERATED_UINTERFACE_BODY()
+};
+
+class ENGINE_API IDeformableInterface
+{
+	GENERATED_IINTERFACE_BODY()
+};
