UX: Remove user-visible ES31 editor references and replace with Mobile

diff --git a/Engine/Config/Android/DataDrivenPlatformInfo.ini b/Engine/Config/Android/DataDrivenPlatformInfo.ini
--- a/Engine/Config/Android/DataDrivenPlatformInfo.ini
+++ b/Engine/Config/Android/DataDrivenPlatformInfo.ini
@@ -27,8 +27,8 @@ ShaderPlatform=OPENGL_ES3_1_ANDROID
 ShaderFormat=GLSL_ES3_1_ANDROID
 ActiveIconName=LevelEditor.PreviewMode.AndroidES31.Enabled
 InactiveIconName=LevelEditor.PreviewMode.AndroidES31.Disabled
-MenuTooltip=NSLOCTEXT("PreviewPlatform", "PreviewMenuTooltip_ES31", "Mobile preview using Android ES3.1 quality settings.")
-IconText=NSLOCTEXT("PreviewPlatform", "PreviewIconText_ES31", "Android ES3.1")
+MenuTooltip=NSLOCTEXT("PreviewPlatform", "PreviewMenuTooltip_AndroidOpenGL ", "Mobile preview using Android OpenGL quality settings.")
+IconText=NSLOCTEXT("PreviewPlatform", "PreviewIconText_AndroidOpenGL", "Android OpenGL")
 DeviceProfileName=Android_Preview_OpenGL
 
 [PreviewPlatform AndroidVulkan]
@@ -64,7 +64,7 @@ bSupportsMobileMultiView=true
 bNeedsOfflineCompiler=true
 bSupportsManualVertexFetch=false
 bSupportsClipDistance = false
-FriendlyName=LOCTEXT("FriendlyShaderPlatformName_AndroidES31", "Android Opengl ES31")
+FriendlyName=LOCTEXT("FriendlyShaderPlatformName_Android", "Android OpenGL Mobile")
 
 [ShaderPlatform VULKAN_ES3_1_ANDROID]
 Language=Vulkan
@@ -77,7 +77,7 @@ bNeedsOfflineCompiler=true
 bSupportsIndexBufferUAVs=true
 bSupportsManualVertexFetch=false
 bSupportsClipDistance = false
-FriendlyName=LOCTEXT("FriendlyShaderPlatformName_AndroidVulkan", "Android Vulkan ES31")
+FriendlyName=LOCTEXT("FriendlyShaderPlatformName_AndroidVulkan", "Android Vulkan Mobile")
 
 [ShaderPlatform VULKAN_SM5_ANDROID]
 Language=Vulkan
diff --git a/Engine/Config/HoloLens/DataDrivenPlatformInfo.ini b/Engine/Config/HoloLens/DataDrivenPlatformInfo.ini
--- a/Engine/Config/HoloLens/DataDrivenPlatformInfo.ini
+++ b/Engine/Config/HoloLens/DataDrivenPlatformInfo.ini
@@ -34,4 +34,4 @@ bSupportsInstancedStereo=true
 bSupportsManualVertexFetch=false
 bSupportsVertexShaderLayer=true
 bSupportsClipDistance = false
-FriendlyName=LOCTEXT("FriendlyShaderPlatformName_HololensES31", "Hololens ES31")
\ No newline at end of file
+FriendlyName=LOCTEXT("FriendlyShaderPlatformName_HololensMobile", "Hololens Mobile")
\ No newline at end of file
diff --git a/Engine/Config/IOS/DataDrivenPlatformInfo.ini b/Engine/Config/IOS/DataDrivenPlatformInfo.ini
--- a/Engine/Config/IOS/DataDrivenPlatformInfo.ini
+++ b/Engine/Config/IOS/DataDrivenPlatformInfo.ini
@@ -54,7 +54,7 @@ bTargetsTiledGPU=true
 bSupports4ComponentUAVReadWrite=true
 bSupportsManualVertexFetch=false
 bSupportsVolumeTextureAtomics=false
-FriendlyName=LOCTEXT("FriendlyShaderPlatformName_IOSES31", "IOS Metal ES31")
+FriendlyName=LOCTEXT("FriendlyShaderPlatformName_IOSMobile", "IOS Metal Mobile")
 
 [ShaderPlatform METAL_MRT]
 Language=Metal
diff --git a/Engine/Config/Mac/DataDrivenPlatformInfo.ini b/Engine/Config/Mac/DataDrivenPlatformInfo.ini
--- a/Engine/Config/Mac/DataDrivenPlatformInfo.ini
+++ b/Engine/Config/Mac/DataDrivenPlatformInfo.ini
@@ -84,4 +84,4 @@ bIsMobile=true
 bSupports4ComponentUAVReadWrite=true
 bSupportsManualVertexFetch=false
 bSupportsVolumeTextureAtomics=false
-FriendlyName=LOCTEXT("FriendlyShaderPlatformName_MacES31", "Mac Metal ES31")
+FriendlyName=LOCTEXT("FriendlyShaderPlatformName_MacMobile", "Mac Metal Mobile")
diff --git a/Engine/Config/TVOS/DataDrivenPlatformInfo.ini b/Engine/Config/TVOS/DataDrivenPlatformInfo.ini
--- a/Engine/Config/TVOS/DataDrivenPlatformInfo.ini
+++ b/Engine/Config/TVOS/DataDrivenPlatformInfo.ini
@@ -36,4 +36,4 @@ bIsMobile=true
 bTargetsTiledGPU=true
 bSupports4ComponentUAVReadWrite=true
 bSupportsManualVertexFetch=false
-FriendlyName=LOCTEXT("FriendlyShaderPlatformName_TVOS_ES31", "TVOS Metal ES31")
\ No newline at end of file
+FriendlyName=LOCTEXT("FriendlyShaderPlatformName_TVOS_Mobile", "TVOS Metal Mobile")
\ No newline at end of file
diff --git a/Engine/Config/VulkanPC/DataDrivenPlatformInfo.ini b/Engine/Config/VulkanPC/DataDrivenPlatformInfo.ini
--- a/Engine/Config/VulkanPC/DataDrivenPlatformInfo.ini
+++ b/Engine/Config/VulkanPC/DataDrivenPlatformInfo.ini
@@ -260,7 +260,7 @@ bSupportsMobileMultiView=true
 bSupportsIndexBufferUAVs=true
 bSupportsManualVertexFetch=false
 bSupportsSceneDataCompressedTransforms=true
-FriendlyName=LOCTEXT("FriendlyShaderPlatformName_VulkanPC_ES31", "PC Vulkan ES31")
+FriendlyName=LOCTEXT("FriendlyShaderPlatformName_VulkanPC_Mobile", "PC Vulkan Mobile")
 
 
 [PreviewPlatform VULKAN_SM5]
diff --git a/Engine/Config/Windows/DataDrivenPlatformInfo.ini b/Engine/Config/Windows/DataDrivenPlatformInfo.ini
--- a/Engine/Config/Windows/DataDrivenPlatformInfo.ini
+++ b/Engine/Config/Windows/DataDrivenPlatformInfo.ini
@@ -155,14 +155,14 @@ bSupportsVertexShaderLayer=true
 bSupportsManualVertexFetch=false
 bSupportsSceneDataCompressedTransforms=true
 bSupportsDebugViewShaders=true
-FriendlyName=LOCTEXT("FriendlyShaderPlatformName_D3D_ES31", "PC D3D ES31")
+FriendlyName=LOCTEXT("FriendlyShaderPlatformName_D3D_Mobile", "PC D3D Mobile")
 SupportsMultiViewport=RuntimeDependent
 
 [PreviewPlatform PCD3D_ES3_1]
 PlatformName=PCD3D_ES3_1
 ShaderFormat=PCD3D_ES31
 ShaderPlatform=PCD3D_ES3_1
-MenuTooltip=LOCTEXT("PreviewMenuTooltip_PCD3D_ES3_1", "Windows using ES3_1 profile.")
+MenuTooltip=LOCTEXT("PreviewMenuTooltip_PCD3D_Mobile", "Windows using Mobile profile.")
 
 [ShaderPlatform OPENGL_PCES3_1]
 Language=OpenGL
@@ -173,4 +173,4 @@ ShaderFormat=GLSL_150_ES31
 bIsMobile=true
 bSupportsManualVertexFetch=false
 bSupportsClipDistance = false
-FriendlyName=LOCTEXT("FriendlyShaderPlatformName_PCGL_ES31", "PC Opengl ES31")
+FriendlyName=LOCTEXT("FriendlyShaderPlatformName_PCGL_Mobile", "PC OpenGL Mobile")
diff --git a/Engine/Source/Developer/Android/AndroidPlatformEditor/Private/AndroidPlatformEditorModule.cpp b/Engine/Source/Developer/Android/AndroidPlatformEditor/Private/AndroidPlatformEditorModule.cpp
--- a/Engine/Source/Developer/Android/AndroidPlatformEditor/Private/AndroidPlatformEditorModule.cpp
+++ b/Engine/Source/Developer/Android/AndroidPlatformEditor/Private/AndroidPlatformEditorModule.cpp
@@ -81,9 +81,9 @@ class FAndroidPlatformEditorModule
 			{
 				static FName NAME_GLSL_ES3_1_ANDROID(TEXT("GLSL_ES3_1_ANDROID"));
 				UShaderPlatformQualitySettings* AndroidMaterialQualitySettings = UMaterialShaderQualitySettings::Get()->GetShaderPlatformQualitySettings(NAME_GLSL_ES3_1_ANDROID);
-				SettingsModule->RegisterSettings("Project", "Platforms", "AndroidES31Quality",
-					LOCTEXT("AndroidES31QualitySettingsName", "Android Material Quality - ES31"),
-					LOCTEXT("AndroidES31QualitySettingsDescription", "Settings for Android ES3.1 material quality"),
+				SettingsModule->RegisterSettings("Project", "Platforms", "AndroidOpenGLMobileQuality",
+					LOCTEXT("AndroidOpenGLMobileQualitySettingsName", "Android Material Quality - OpenGL Mobile"),
+					LOCTEXT("AndroidOpenGLMobileQualitySettingsDescription", "Settings for Android OpenGL Mobile material quality"),
 					AndroidMaterialQualitySettings
 				);
 			}
diff --git a/Engine/Source/Editor/DetailCustomizations/Private/LinuxTargetSettingsDetails.cpp b/Engine/Source/Editor/DetailCustomizations/Private/LinuxTargetSettingsDetails.cpp
--- a/Engine/Source/Editor/DetailCustomizations/Private/LinuxTargetSettingsDetails.cpp
+++ b/Engine/Source/Editor/DetailCustomizations/Private/LinuxTargetSettingsDetails.cpp
@@ -56,7 +56,7 @@ static FText GetFriendlyNameFromLinuxShaderFormat(const FName InShaderFormat)
 
 	if (InShaderFormat == NAME_GLSL_150_ES31)
 	{
-		FriendlyRHIName = LOCTEXT("OpenGL3ES31", "OpenGL 3 (ES3.1, Experimental)");
+		FriendlyRHIName = LOCTEXT("OpenGL3ES31", "OpenGL 3 (Mobile, Experimental)");
 	}
 	else if (InShaderFormat == NAME_VULKAN_ES3_1_ANDROID || InShaderFormat == NAME_VULKAN_ES3_1)
 	{
diff --git a/Engine/Source/Editor/DetailCustomizations/Private/MacTargetSettingsDetails.cpp b/Engine/Source/Editor/DetailCustomizations/Private/MacTargetSettingsDetails.cpp
--- a/Engine/Source/Editor/DetailCustomizations/Private/MacTargetSettingsDetails.cpp
+++ b/Engine/Source/Editor/DetailCustomizations/Private/MacTargetSettingsDetails.cpp
@@ -114,22 +114,22 @@ static FText GetFriendlyNameFromRHINameMac(FName InRHIName)
 		FriendlyRHIName = LOCTEXT("D3DSM5", "Direct3D 11+ (SM5)");
 		break;
 	case SP_PCD3D_ES3_1:
-		FriendlyRHIName = LOCTEXT("D3DES31", "Direct3D (ES3.1, Mobile Preview)");
+		FriendlyRHIName = LOCTEXT("D3DESMobile", "Direct3D (Mobile Preview)");
 		break;
 	case SP_OPENGL_PCES3_1:
-		FriendlyRHIName = LOCTEXT("OpenGLES31PC", "OpenGL (ES3.1, Mobile Preview)");
+		FriendlyRHIName = LOCTEXT("OpenGLESMobilePC", "OpenGL (Mobile Preview)");
 		break;
 	case SP_OPENGL_ES3_1_ANDROID:
-		FriendlyRHIName = LOCTEXT("OpenGLES31", "OpenGLES 3.1 (Mobile)");
+		FriendlyRHIName = LOCTEXT("OpenGLESMobile", "OpenGLES (Mobile)");
 		break;
 	case SP_METAL:
-		FriendlyRHIName = LOCTEXT("Metal", "iOS Metal Mobile Renderer (ES3.1, Metal 2.4+, iOS 15.0 or later)");
+		FriendlyRHIName = LOCTEXT("Metal", "iOS Metal Mobile Renderer (Mobile, Metal 2.4+, iOS 15.0 or later)");
 		break;
 	case SP_METAL_MRT:
 		FriendlyRHIName = LOCTEXT("MetalMRT", "iOS Metal Desktop Renderer (SM5, Metal 2.4+, iOS 15.0 or later)");
 		break;
 	case SP_METAL_TVOS:
-		FriendlyRHIName = LOCTEXT("MetalTV", "tvOS Metal Mobile Renderer (ES3.1, Metal 2.4+, tvOS 15.0 or later)");
+		FriendlyRHIName = LOCTEXT("MetalTV", "tvOS Metal Mobile Renderer (Mobile, Metal 2.4+, tvOS 15.0 or later)");
 		break;
 	case SP_METAL_MRT_TVOS:
 		FriendlyRHIName = LOCTEXT("MetalMRTTV", "tvOS Metal Desktop Renderer (SM5, Metal 2.4+, tvOS 15.0 or later)");
@@ -138,7 +138,7 @@ static FText GetFriendlyNameFromRHINameMac(FName InRHIName)
 		FriendlyRHIName = LOCTEXT("MetalSM5", "Mac Metal Desktop Renderer (SM5, Metal 2.4+, macOS Monterey 12.0 or later)");
 		break;
 	case SP_METAL_MACES3_1:
-		FriendlyRHIName = LOCTEXT("MetalES3.1", "Mac Metal High-End Mobile Preview (ES3.1)");
+		FriendlyRHIName = LOCTEXT("MetalMobile", "Mac Metal High-End Mobile Preview (Mobile Preview)");
 		break;
 	case SP_METAL_MRT_MAC:
 		FriendlyRHIName = LOCTEXT("MetalMRTMac", "Mac Metal iOS/tvOS Desktop Renderer Preview (SM5)");
@@ -152,10 +152,10 @@ static FText GetFriendlyNameFromRHINameMac(FName InRHIName)
 		break;
 	case SP_VULKAN_PCES3_1:
 	case SP_VULKAN_ES3_1_ANDROID:
-		FriendlyRHIName = LOCTEXT("VulkanES31", "Vulkan (ES 3.1)");
+		FriendlyRHIName = LOCTEXT("VulkanMobile", "Vulkan (Mobile)");
 		break;
 	case SP_D3D_ES3_1_HOLOLENS:
-		FriendlyRHIName = LOCTEXT("D3DES31HL", "Direct3D (ES3.1, Hololens)");
+		FriendlyRHIName = LOCTEXT("D3DMobileHL", "Direct3D (Mobile, Hololens)");
 		break;
 	default:
 		FriendlyRHIName = FText::FromString(InRHIName.ToString());
diff --git a/Engine/Source/Editor/DetailCustomizations/Private/WindowsTargetSettingsDetails.cpp b/Engine/Source/Editor/DetailCustomizations/Private/WindowsTargetSettingsDetails.cpp
--- a/Engine/Source/Editor/DetailCustomizations/Private/WindowsTargetSettingsDetails.cpp
+++ b/Engine/Source/Editor/DetailCustomizations/Private/WindowsTargetSettingsDetails.cpp
@@ -65,10 +65,10 @@ static FText GetFriendlyNameFromWindowsShaderPlatform(FName InShaderPlatformName
 		FriendlyRHIName = LOCTEXT("DirectX11", "DirectX 11 & 12 (SM5)");
 		break;
 	case SP_PCD3D_ES3_1:
-		FriendlyRHIName = LOCTEXT("DirectXES31", "DirectX Mobile Emulation (ES3.1)");
+		FriendlyRHIName = LOCTEXT("DirectXMobile", "DirectX Mobile Emulation (Mobile)");
 		break;
 	case SP_D3D_ES3_1_HOLOLENS:
-		FriendlyRHIName = LOCTEXT("DirectXES31HL", "DirectX Hololens (ES3.1)");
+		FriendlyRHIName = LOCTEXT("DirectXMobileHL", "DirectX Hololens (Mobile)");
 		break;
 	case SP_VULKAN_SM5:
 		FriendlyRHIName = LOCTEXT("VulkanSM5", "Vulkan (SM5)");
@@ -109,7 +109,7 @@ static FText GetFriendlyNameForWindowsShaderPlatformCheckbox(FName InShaderPlatf
 		FriendlyName = LOCTEXT("SM5", "SM5");
 		break;
 	case SP_PCD3D_ES3_1:
-		FriendlyName = LOCTEXT("ES31", "ES3.1");
+		FriendlyName = LOCTEXT("Mobile", "Mobile");
 		break;
 	default:
 		break;
diff --git a/Engine/Source/Editor/MaterialEditor/Private/MaterialEditor.cpp b/Engine/Source/Editor/MaterialEditor/Private/MaterialEditor.cpp
--- a/Engine/Source/Editor/MaterialEditor/Private/MaterialEditor.cpp
+++ b/Engine/Source/Editor/MaterialEditor/Private/MaterialEditor.cpp
@@ -1965,7 +1965,7 @@ void FMaterialEditor::GeneratePreviewMenuContent(UToolMenu* Menu)
 	FToolMenuSection& FeatureLevelSection = Menu->AddSection("MaterialEditorFeaturePreview", LOCTEXT("MaterialFeatureHeading", "Feature Level"));
 	{
 		FeatureLevelSection.AddMenuEntry(FMaterialEditorCommands::Get().FeatureLevel_All);
-		FeatureLevelSection.AddMenuEntry(FMaterialEditorCommands::Get().FeatureLevel_ES31);
+		FeatureLevelSection.AddMenuEntry(FMaterialEditorCommands::Get().FeatureLevel_Mobile);
 		FeatureLevelSection.AddMenuEntry(FMaterialEditorCommands::Get().FeatureLevel_SM5);
 		FeatureLevelSection.AddMenuEntry(FMaterialEditorCommands::Get().FeatureLevel_SM6);
 	}
@@ -3469,7 +3469,7 @@ void FMaterialEditor::BindCommands()
 		FCanExecuteAction(),
 		FIsActionChecked::CreateSP(this, &FMaterialEditor::IsFeaturePreviewChecked, ERHIFeatureLevel::Num));
 	ToolkitCommands->MapAction(
-		Commands.FeatureLevel_ES31,
+		Commands.FeatureLevel_Mobile,
 		FExecuteAction::CreateSP(this, &FMaterialEditor::SetFeaturePreview, ERHIFeatureLevel::ES3_1),
 		FCanExecuteAction::CreateSP(this, &FMaterialEditor::IsFeaturePreviewAvailable, ERHIFeatureLevel::ES3_1),
 		FIsActionChecked::CreateSP(this, &FMaterialEditor::IsFeaturePreviewChecked, ERHIFeatureLevel::ES3_1));
diff --git a/Engine/Source/Editor/MaterialEditor/Private/MaterialEditorActions.cpp b/Engine/Source/Editor/MaterialEditor/Private/MaterialEditorActions.cpp
--- a/Engine/Source/Editor/MaterialEditor/Private/MaterialEditorActions.cpp
+++ b/Engine/Source/Editor/MaterialEditor/Private/MaterialEditorActions.cpp
@@ -89,7 +89,7 @@ void FMaterialEditorCommands::RegisterCommands()
 	UI_COMMAND(QualityLevel_Low, "Low", "Sets node preview to low quality.", EUserInterfaceActionType::RadioButton, FInputChord());
 
 	UI_COMMAND(FeatureLevel_All, "All", "Sets node preview to show all feature levels.", EUserInterfaceActionType::RadioButton, FInputChord());
-	UI_COMMAND(FeatureLevel_ES31, "ES3.1", "Sets node preview to show the ES3.1 feature level.", EUserInterfaceActionType::RadioButton, FInputChord());
+	UI_COMMAND(FeatureLevel_Mobile, "Mobile", "Sets node preview to show the Mobile feature level.", EUserInterfaceActionType::RadioButton, FInputChord());
 	UI_COMMAND(FeatureLevel_SM5, "SM5", "Sets node preview to show the SM5 feature level.", EUserInterfaceActionType::RadioButton, FInputChord());
 	UI_COMMAND(FeatureLevel_SM6, "SM6", "Sets node preview to show the SM6 feature level.", EUserInterfaceActionType::RadioButton, FInputChord());
 }
diff --git a/Engine/Source/Editor/MaterialEditor/Private/MaterialStats.cpp b/Engine/Source/Editor/MaterialEditor/Private/MaterialStats.cpp
--- a/Engine/Source/Editor/MaterialEditor/Private/MaterialStats.cpp
+++ b/Engine/Source/Editor/MaterialEditor/Private/MaterialStats.cpp
@@ -471,10 +471,10 @@ void FMaterialStats::BuildShaderPlatformDB()
 	// DirectX
 	AddShaderPlatform(EPlatformCategoryType::Desktop, SP_PCD3D_SM5, TEXT("DirectX SM5"), true, true, TEXT("Desktop, DirectX, Shader Model 5"));
 	AddShaderPlatform(EPlatformCategoryType::Desktop, SP_PCD3D_SM6, TEXT("DirectX SM6"), true, true, TEXT("Desktop, DirectX, Shader Model 6"));
-	AddShaderPlatform(EPlatformCategoryType::Desktop, SP_PCD3D_ES3_1, TEXT("DirectX ES 3.1"), true, true, TEXT("Desktop, DirectX, ES 3.1"));
+	AddShaderPlatform(EPlatformCategoryType::Desktop, SP_PCD3D_ES3_1, TEXT("DirectX  Mobile"), true, true, TEXT("Desktop, DirectX, Mobile"));
 
 	// Desktop is the closest fit for Hololens (for now)
-	AddShaderPlatform(EPlatformCategoryType::Desktop, SP_D3D_ES3_1_HOLOLENS, TEXT("DirectX ES 3.1 Hololens"), true, true, TEXT("Hololens, DirectX, ES 3.1"));
+	AddShaderPlatform(EPlatformCategoryType::Desktop, SP_D3D_ES3_1_HOLOLENS, TEXT("DirectX Mobile Hololens"), true, true, TEXT("Hololens, DirectX, Mobile"));
 #endif
 
 	// Vulkan
@@ -482,8 +482,8 @@ void FMaterialStats::BuildShaderPlatformDB()
 	AddShaderPlatform(EPlatformCategoryType::Desktop, SP_VULKAN_SM6, TEXT("Vulkan SM6"), true, true, TEXT("Desktop, Vulkan, Shader Model 6"));
 
 	// Android
-	AddShaderPlatform(EPlatformCategoryType::Android, SP_OPENGL_ES3_1_ANDROID, TEXT("Android GLES 3.1"), true, true, TEXT("Android, OpenGLES 3.1"));
-	AddShaderPlatform(EPlatformCategoryType::Android, SP_VULKAN_ES3_1_ANDROID, TEXT("Android Vulkan"), true, true, TEXT("Android, Vulkan"));
+	AddShaderPlatform(EPlatformCategoryType::Android, SP_OPENGL_ES3_1_ANDROID, TEXT("Android GLES Mobile"), true, true, TEXT("Android, OpenGLES Mobile"));
+	AddShaderPlatform(EPlatformCategoryType::Android, SP_VULKAN_ES3_1_ANDROID, TEXT("Android Vulkan Mobile"), true, true, TEXT("Android, Vulkan Mobile"));
 	AddShaderPlatform(EPlatformCategoryType::Android, SP_VULKAN_SM5_ANDROID, TEXT("Android Vulkan SM5"), true, true, TEXT("Android, Vulkan SM5"));
 
 	// Apple
diff --git a/Engine/Source/Editor/MaterialEditor/Public/MaterialEditorActions.h b/Engine/Source/Editor/MaterialEditor/Public/MaterialEditorActions.h
--- a/Engine/Source/Editor/MaterialEditor/Public/MaterialEditorActions.h
+++ b/Engine/Source/Editor/MaterialEditor/Public/MaterialEditorActions.h
@@ -206,7 +206,7 @@ public:
 	TSharedPtr< FUICommandInfo > QualityLevel_Low;
 
 	TSharedPtr< FUICommandInfo > FeatureLevel_All;
-	TSharedPtr< FUICommandInfo > FeatureLevel_ES31;
+	TSharedPtr< FUICommandInfo > FeatureLevel_Mobile;
 	TSharedPtr< FUICommandInfo > FeatureLevel_SM5;
 	TSharedPtr< FUICommandInfo > FeatureLevel_SM6;
 
diff --git a/Engine/Source/Editor/UnrealEd/Private/Kismet2/DebuggerCommands.cpp b/Engine/Source/Editor/UnrealEd/Private/Kismet2/DebuggerCommands.cpp
--- a/Engine/Source/Editor/UnrealEd/Private/Kismet2/DebuggerCommands.cpp
+++ b/Engine/Source/Editor/UnrealEd/Private/Kismet2/DebuggerCommands.cpp
@@ -339,7 +339,7 @@ void FPlayWorldCommands::RegisterCommands()
 	UI_COMMAND(PlayInViewport, "Selected Viewport", "Play this level in the active level editor viewport", EUserInterfaceActionType::Check, FInputChord());
 	UI_COMMAND(PlayInEditorFloating, "New Editor Window (PIE)", "Play this level in a new window", EUserInterfaceActionType::Check, FInputChord());
 	UI_COMMAND(PlayInVR, "VR Preview", "Play this level in VR", EUserInterfaceActionType::Check, FInputChord());
-	UI_COMMAND(PlayInMobilePreview, "Standalone Game Mobile", "Play this level as a mobile device preview in ES3.1 mode (runs in its own process)", EUserInterfaceActionType::Check, FInputChord());
+	UI_COMMAND(PlayInMobilePreview, "Standalone Game Mobile", "Play this level as a mobile device preview in Mobile mode (runs in its own process)", EUserInterfaceActionType::Check, FInputChord());
 	UI_COMMAND(PlayInVulkanPreview, "Standalone Game Vulkan Mobile", "Play this level using mobile Vulkan rendering (runs in its own process)", EUserInterfaceActionType::Check, FInputChord());
 	UI_COMMAND(PlayInNewProcess, "Standalone Game", "Play this level in a new window that runs in its own process", EUserInterfaceActionType::Check, FInputChord());
 	UI_COMMAND(PlayInCameraLocation, "Current Camera Location", "Spawn the player at the current camera location", EUserInterfaceActionType::RadioButton, FInputChord());
diff --git a/Engine/Source/Runtime/Android/AndroidRuntimeSettings/Classes/AndroidRuntimeSettings.h b/Engine/Source/Runtime/Android/AndroidRuntimeSettings/Classes/AndroidRuntimeSettings.h
--- a/Engine/Source/Runtime/Android/AndroidRuntimeSettings/Classes/AndroidRuntimeSettings.h
+++ b/Engine/Source/Runtime/Android/AndroidRuntimeSettings/Classes/AndroidRuntimeSettings.h
@@ -671,7 +671,7 @@ public:
 	UPROPERTY(GlobalConfig, EditAnywhere, Category = "Project SDK Override", Meta = (DisplayName = "Build-Tools Version (specific version or 'latest')"))
 	FString BuildToolsOverride;
 
-	/** Whether to enable LOD streaming for landscape visual meshes. Only supported on feature level ES3.1 or above. */
+	/** Whether to enable LOD streaming for landscape visual meshes. Only supported on feature level Mobile or above. */
 	UPROPERTY(GlobalConfig, EditAnywhere, Category = "Misc", Meta = (DisplayName = "Stream landscape visual mesh LODs"))
 	bool bStreamLandscapeMeshLODs;
 
diff --git a/Engine/Source/Runtime/Core/Private/HAL/ConsoleManager.cpp b/Engine/Source/Runtime/Core/Private/HAL/ConsoleManager.cpp
--- a/Engine/Source/Runtime/Core/Private/HAL/ConsoleManager.cpp
+++ b/Engine/Source/Runtime/Core/Private/HAL/ConsoleManager.cpp
@@ -2825,14 +2825,14 @@ static TAutoConsoleVariable<int32> CVarMobileAllowDitheredLODTransition(
 static TAutoConsoleVariable<int32> CVarMobileAllowPixelDepthOffset(
 	TEXT("r.Mobile.AllowPixelDepthOffset"),
 	1,
-	TEXT("Whether to allow 'Pixel Depth Offset' in materials for ES3.1 feature level. Depth modification in pixel shaders may reduce GPU performance"),
+	TEXT("Whether to allow 'Pixel Depth Offset' in materials for Mobile feature level. Depth modification in pixel shaders may reduce GPU performance"),
 	ECVF_ReadOnly | ECVF_RenderThreadSafe
 );
 
 static TAutoConsoleVariable<int32> CVarMobileAllowPerPixelShadingModels(
 	TEXT("r.Mobile.AllowPerPixelShadingModels"),
 	1,
-	TEXT("Whether to allow 'Per-Pixel Shader Models (From Material Expression)' in materials for ES3.1 feature level."),
+	TEXT("Whether to allow 'Per-Pixel Shader Models (From Material Expression)' in materials for Mobile feature level."),
 	ECVF_ReadOnly | ECVF_RenderThreadSafe
 );
 
@@ -2858,7 +2858,7 @@ static TAutoConsoleVariable<int32> CVarMobileForwardEnableClusteredReflections(
 static TAutoConsoleVariable<int32> CVarMobileSupportGPUScene(
 	TEXT("r.Mobile.SupportGPUScene"),
 	0,
-	TEXT("Whether to support GPU scene, required for auto-instancing (only ES3.1 feature level)"),
+	TEXT("Whether to support GPU scene, required for auto-instancing (only Mobile feature level)"),
 	ECVF_ReadOnly | ECVF_RenderThreadSafe
 );
 
@@ -3451,7 +3451,7 @@ static TAutoConsoleVariable<int32> CVarDisableOpenGLTextureStreamingSupport(
 static FAutoConsoleVariable CVarOpenGLUseEmulatedUBs(
 	TEXT("OpenGL.UseEmulatedUBs"),
 	1,
-	TEXT("If true, enable using emulated uniform buffers on OpenGL ES3.1 mode."),
+	TEXT("If true, enable using emulated uniform buffers on OpenGL Mobile mode."),
 	ECVF_ReadOnly
 	);
 
diff --git a/Engine/Source/Runtime/Engine/Private/GPUSkinVertexFactory.cpp b/Engine/Source/Runtime/Engine/Private/GPUSkinVertexFactory.cpp
--- a/Engine/Source/Runtime/Engine/Private/GPUSkinVertexFactory.cpp
+++ b/Engine/Source/Runtime/Engine/Private/GPUSkinVertexFactory.cpp
@@ -63,7 +63,7 @@ static FAutoConsoleVariableRef CVarUnlimitedBoneInfluencesThreshold(
 static TAutoConsoleVariable<bool> CVarMobileEnableCloth(
 	TEXT("r.Mobile.EnableCloth"),
 	true,
-	TEXT("If enabled, compile cloth shader permutations and render simulated cloth on mobile platforms and Windows ES3.1. Cannot be changed at runtime"),
+	TEXT("If enabled, compile cloth shader permutations and render simulated cloth on mobile platforms and Mobile mode on PC. Cannot be changed at runtime"),
 	ECVF_ReadOnly);
 
 IMPLEMENT_GLOBAL_SHADER_PARAMETER_STRUCT(FAPEXClothUniformShaderParameters,"APEXClothParam");
diff --git a/Engine/Source/Runtime/Engine/Private/Materials/HLSLMaterialTranslator.cpp b/Engine/Source/Runtime/Engine/Private/Materials/HLSLMaterialTranslator.cpp
--- a/Engine/Source/Runtime/Engine/Private/Materials/HLSLMaterialTranslator.cpp
+++ b/Engine/Source/Runtime/Engine/Private/Materials/HLSLMaterialTranslator.cpp
@@ -7123,7 +7123,7 @@ int32 FHLSLMaterialTranslator::SceneDepth(int32 Offset, int32 ViewportUV, bool b
 	if (ShaderFrequency == SF_Vertex && FeatureLevel <= ERHIFeatureLevel::ES3_1)
 	{
 		// mobile currently does not support this, we need to read a separate copy of the depth, we must disable framebuffer fetch and force scene texture reads.
-		return Errorf(TEXT("Cannot read scene depth from the vertex shader with feature level ES3.1 or below."));
+		return Errorf(TEXT("Cannot read scene depth from the vertex shader with the Mobile feature level"));
 	}
 
 	if (Material->IsTranslucencyWritingVelocity())
