Fix Component Reference Customization not recognizing the SoftComponentReference's ComponentProperty and PathToComponent. This was preventing selecting non-root components within nDisplay.

diff --git a/Engine/Source/Editor/DetailCustomizations/Private/ComponentReferenceCustomization.cpp b/Engine/Source/Editor/DetailCustomizations/Private/ComponentReferenceCustomization.cpp
--- a/Engine/Source/Editor/DetailCustomizations/Private/ComponentReferenceCustomization.cpp
+++ b/Engine/Source/Editor/DetailCustomizations/Private/ComponentReferenceCustomization.cpp
@@ -379,6 +379,8 @@ void FComponentReferenceCustomization::SetValue(const FComponentReference& Value
 			if (Value.OtherActor.IsValid())
 			{
 				SoftValue.OtherActor = Value.OtherActor.Get();
+				SoftValue.ComponentProperty = Value.ComponentProperty;
+				SoftValue.PathToComponent = Value.PathToComponent;
 			}
 			CastFieldChecked<const FStructProperty>(PropertyHandle->GetProperty())->Struct->ExportText(TextValue, &SoftValue, &SoftValue, nullptr, EPropertyPortFlags::PPF_None, nullptr);
 			ensure(PropertyHandle->SetValueFromFormattedString(TextValue) == FPropertyAccess::Result::Success);
@@ -418,6 +420,8 @@ FPropertyAccess::Result FComponentReferenceCustomization::GetValue(FComponentRef
 					if (SoftReference.OtherActor.IsValid())
 					{
 						ThisReference.OtherActor = SoftReference.OtherActor.Get();
+						ThisReference.ComponentProperty = SoftReference.ComponentProperty;
+						ThisReference.PathToComponent = SoftReference.PathToComponent;
 					}
 				}
 				else
