Fixed static analysis in ActorDetails.cpp: Added a nullptr check

diff --git a/Engine/Source/Editor/DetailCustomizations/Private/ActorDetails.cpp b/Engine/Source/Editor/DetailCustomizations/Private/ActorDetails.cpp
--- a/Engine/Source/Editor/DetailCustomizations/Private/ActorDetails.cpp
+++ b/Engine/Source/Editor/DetailCustomizations/Private/ActorDetails.cpp
@@ -565,7 +565,7 @@ FReply FActorDetails::HandleAddOrViewEventForVariable(UBlueprint* BP, FMulticast
 	const UFunction* SignatureFunction = Property ? Property->SignatureFunction : nullptr;
 	UEdGraph* EventGraph = BP ? FBlueprintEditorUtils::FindEventGraph(BP) : nullptr;
 
-	if (EventGraph && SignatureFunction)
+	if (EventGraph && SignatureFunction && Property)
 	{
 		const FVector2D SpawnPos = EventGraph->GetGoodPlaceForNewNode();
 
