Made landscape proxy count details more explicit

#rnx

diff --git a/Engine/Source/Editor/DetailCustomizations/Private/LandscapeProxyUIDetails.cpp b/Engine/Source/Editor/DetailCustomizations/Private/LandscapeProxyUIDetails.cpp
--- a/Engine/Source/Editor/DetailCustomizations/Private/LandscapeProxyUIDetails.cpp
+++ b/Engine/Source/Editor/DetailCustomizations/Private/LandscapeProxyUIDetails.cpp
@@ -169,7 +169,7 @@ void FLandscapeProxyUIDetails::CustomizeDetails( IDetailLayoutBuilder& DetailBui
 			];
 
 			int32 LandscapeCount = LandscapeInfo->StreamingProxies.Num() + (LandscapeInfo->LandscapeActor.Get() ? 1 : 0);
-			RowDisplayText = LOCTEXT("LandscapeCount", "Landscape Count");
+			RowDisplayText = LOCTEXT("LandscapeCount", "Landscape Proxy Count");
 			CategoryBuilder.AddCustomRow(RowDisplayText)
 			.RowTag(TEXT("LandscapeCount"))
 			.NameContent()
@@ -178,7 +178,7 @@ void FLandscapeProxyUIDetails::CustomizeDetails( IDetailLayoutBuilder& DetailBui
 			]
 			.ValueContent()
 			[
-				GenerateTextWidget(FText::Format(LOCTEXT("LandscapeCountValue", "{0}"), LandscapeCount), true)
+				GenerateTextWidget(FText::Format(LOCTEXT("LandscapeProxyCountValue", "{0}"), LandscapeCount), true)
 			];
 
 			int32 TotalComponentCount = LandscapeInfo->XYtoComponentMap.Num();
