Improvements to Inertialization and Dead Blending node:
- Added bone filter to Dead Blending node.
- Made inertialization rewind debugger tracks clickable.
- State machines no longer grey-out the blend parameters when an inertialization transition is selected.

diff --git a/Engine/Plugins/Animation/GameplayInsights/Source/GameplayInsights/Private/InertializationsTrack.cpp b/Engine/Plugins/Animation/GameplayInsights/Source/GameplayInsights/Private/InertializationsTrack.cpp
--- a/Engine/Plugins/Animation/GameplayInsights/Source/GameplayInsights/Private/InertializationsTrack.cpp
+++ b/Engine/Plugins/Animation/GameplayInsights/Source/GameplayInsights/Private/InertializationsTrack.cpp
@@ -10,7 +10,9 @@
 
 #if WITH_EDITOR
 #include "Animation/AnimInstance.h"
-#include "Editor.h"
+#include "Animation/AnimBlueprintGeneratedClass.h"
+#include "IAnimationBlueprintEditor.h"
+#include "Animation/AnimBlueprint.h"
 #include "Subsystems/AssetEditorSubsystem.h"
 #endif
 
@@ -210,6 +212,47 @@ TSharedPtr<SWidget> FInertializationTrack::GetDetailsViewInternal()
 
 bool FInertializationTrack::HandleDoubleClickInternal()
 {
+
+#if WITH_EDITOR
+	IRewindDebugger* RewindDebugger = IRewindDebugger::Instance();
+	if (const TraceServices::IAnalysisSession* AnalysisSession = RewindDebugger->GetAnalysisSession())
+	{
+		TraceServices::FAnalysisSessionReadScope SessionReadScope(*AnalysisSession);
+
+		if (const FGameplayProvider* GameplayProvider = AnalysisSession->ReadProvider<FGameplayProvider>(FGameplayProvider::ProviderName))
+		{
+			if (const FObjectInfo* AnimInstanceInfo = GameplayProvider->FindObjectInfo(ObjectId))
+			{
+				if (const FClassInfo* AnimInstanceClassInfo = GameplayProvider->FindClassInfo(AnimInstanceInfo->ClassId))
+				{
+					TSoftObjectPtr<UAnimBlueprintGeneratedClass> InstanceClass;
+					InstanceClass = FSoftObjectPath(AnimInstanceClassInfo->PathName);
+
+					if (InstanceClass.LoadSynchronous())
+					{
+						if (UAnimBlueprint* AnimBlueprint = Cast<UAnimBlueprint>(InstanceClass.Get()->ClassGeneratedBy))
+						{
+							GEditor->GetEditorSubsystem<UAssetEditorSubsystem>()->OpenEditorForAsset(AnimBlueprint);
+
+							if (IAnimationBlueprintEditor* AnimBlueprintEditor = static_cast<IAnimationBlueprintEditor*>(GEditor->GetEditorSubsystem<UAssetEditorSubsystem>()->FindEditorForAsset(AnimBlueprint, true)))
+							{
+								int32 AnimNodeIndex = InstanceClass.Get()->GetAnimNodeProperties().Num() - NodeId - 1;
+								TWeakObjectPtr<const UEdGraphNode>* GraphNode = InstanceClass.Get()->AnimBlueprintDebugData.NodePropertyIndexToNodeMap.Find(AnimNodeIndex);
+								if (GraphNode != nullptr && GraphNode->Get())
+								{
+									AnimBlueprintEditor->JumpToHyperlink(GraphNode->Get());
+								}
+							}
+
+							return true;
+						}
+					}
+				}
+			}
+		}
+	}
+#endif
+
 	return false;
 }
 
diff --git a/Engine/Source/Editor/DetailCustomizations/Private/AnimTransitionNodeDetails.cpp b/Engine/Source/Editor/DetailCustomizations/Private/AnimTransitionNodeDetails.cpp
--- a/Engine/Source/Editor/DetailCustomizations/Private/AnimTransitionNodeDetails.cpp
+++ b/Engine/Source/Editor/DetailCustomizations/Private/AnimTransitionNodeDetails.cpp
@@ -209,24 +209,11 @@ void FAnimTransitionNodeDetails::CustomizeDetails( IDetailLayoutBuilder& DetailB
 
 		//////////////////////////////////////////////////////////////////////////
 
-		auto BlendSettingsEnabled = [LogicTypeHandle]()
-		{
-			uint8 LogicType;
-			if (LogicTypeHandle->GetValue(LogicType) == FPropertyAccess::Result::Success)
-			{
-				return LogicType != ETransitionLogicType::TLT_Inertialization;
-			}
-			return true;
-		};
-
-		auto BlendSettingsEnabledAttribute = TAttribute<bool>::Create(TAttribute<bool>::FGetter::CreateLambda(BlendSettingsEnabled));
-
 		IDetailCategoryBuilder& CrossfadeCategory = DetailBuilder.EditCategory("BlendSettings", LOCTEXT("BlendSettingsCategoryTitle", "Blend Settings") );
 		if (TransitionNode != NULL && SelectedObjects.Num() == 1)
 		{
 			// The sharing option for the crossfade settings
 			CrossfadeCategory.AddCustomRow( LOCTEXT("TransitionCrossfadeSharingLabel", "Transition Crossfade Sharing") )
-			.IsEnabled(BlendSettingsEnabledAttribute)
 			.NameContent()
 			[
 				SNew(STextBlock)
@@ -247,8 +234,8 @@ void FAnimTransitionNodeDetails::CustomizeDetails( IDetailLayoutBuilder& DetailB
 
 		//@TODO: Gate editing these on shared non-authoritative ones
 		CrossfadeCategory.AddProperty(GET_MEMBER_NAME_CHECKED(UAnimStateTransitionNode, CrossfadeDuration)).DisplayName( LOCTEXT("DurationLabel", "Duration") );
-		CrossfadeCategory.AddProperty(GET_MEMBER_NAME_CHECKED(UAnimStateTransitionNode, BlendMode)).DisplayName(LOCTEXT("ModeLabel", "Mode")).IsEnabled(BlendSettingsEnabledAttribute);
-		CrossfadeCategory.AddProperty(GET_MEMBER_NAME_CHECKED(UAnimStateTransitionNode, CustomBlendCurve)).DisplayName(LOCTEXT("CurveLabel", "Custom Blend Curve")).IsEnabled(BlendSettingsEnabledAttribute);
+		CrossfadeCategory.AddProperty(GET_MEMBER_NAME_CHECKED(UAnimStateTransitionNode, BlendMode)).DisplayName(LOCTEXT("ModeLabel", "Mode"));
+		CrossfadeCategory.AddProperty(GET_MEMBER_NAME_CHECKED(UAnimStateTransitionNode, CustomBlendCurve)).DisplayName(LOCTEXT("CurveLabel", "Custom Blend Curve"));
 		CrossfadeCategory.AddProperty(GET_MEMBER_NAME_CHECKED(UAnimStateTransitionNode, BlendProfile)).DisplayName(LOCTEXT("BlendProfileLabel", "Blend Profile"));
 
 		//////////////////////////////////////////////////////////////////////////
diff --git a/Engine/Source/Runtime/Engine/Classes/Animation/AnimNode_DeadBlending.h b/Engine/Source/Runtime/Engine/Classes/Animation/AnimNode_DeadBlending.h
--- a/Engine/Source/Runtime/Engine/Classes/Animation/AnimNode_DeadBlending.h
+++ b/Engine/Source/Runtime/Engine/Classes/Animation/AnimNode_DeadBlending.h
@@ -3,6 +3,7 @@
 #pragma once
 
 #include "AnimNode_Inertialization.h"
+#include "Interfaces/Interface_BoneReferenceSkeletonProvider.h"
 #include "AnimNode_DeadBlending.generated.h"
 
 /**
@@ -28,7 +29,7 @@
  * More specifically they specify the "half-life" - or how it takes for the velocities to be decayed by half.
  */
 USTRUCT(Experimental, BlueprintInternalUseOnly)
-struct ENGINE_API FAnimNode_DeadBlending : public FAnimNode_Base
+struct ENGINE_API FAnimNode_DeadBlending : public FAnimNode_Base, public IBoneReferenceSkeletonProvider
 {
 	GENERATED_BODY()
 
@@ -66,6 +67,10 @@ private:
 	UPROPERTY(EditAnywhere, Category = Filter)
 	TArray<FName> FilteredCurves;
 
+	// List of bones that should not use inertial blending. These bones will change instantly when the animation switches.
+	UPROPERTY(EditAnywhere, Category = Filter)
+	TArray<FBoneReference> FilteredBones;
+
 	/**
 	 * The average half-life of decay in seconds to use when extrapolating the animation. To get the final half-life of 
 	 * decay, this value will be scaled by the amount to which the velocities of the animation being transitioned from 
@@ -141,11 +146,17 @@ private:
 	 */
 	void ApplyTo(FCompactPose& InOutPose, FBlendedCurve& InOutCurves);
 
+public: // IBoneReferenceSkeletonProvider
+	class USkeleton* GetSkeleton(bool& bInvalidSkeletonIsError, const IPropertyHandle* PropertyHandle) override;
+
 private:
 
 	// Cached curve filter built from FilteredCurves
 	UE::Anim::FCurveFilter CurveFilter;
 
+	// Cache compact pose bone index for FilteredBones
+	TArray<FCompactPoseBoneIndex, TInlineAllocator<8>> BoneFilter;
+
 	// Snapshots of the actor pose generated as output.
 	TArray<FInertializationPose, TInlineAllocator<2>> PoseSnapshots;
 
diff --git a/Engine/Source/Runtime/Engine/Private/Animation/AnimNode_DeadBlending.cpp b/Engine/Source/Runtime/Engine/Private/Animation/AnimNode_DeadBlending.cpp
--- a/Engine/Source/Runtime/Engine/Private/Animation/AnimNode_DeadBlending.cpp
+++ b/Engine/Source/Runtime/Engine/Private/Animation/AnimNode_DeadBlending.cpp
@@ -382,7 +382,7 @@ void FAnimNode_DeadBlending::ApplyTo(FCompactPose& InOutPose, FBlendedCurve& InO
 	{
 		const int32 SkeletonPoseBoneIndex = BoneContainer.GetSkeletonIndex(BoneIndex);
 
-		if (SkeletonPoseBoneIndex == INDEX_NONE || !BoneValid[SkeletonPoseBoneIndex])
+		if (SkeletonPoseBoneIndex == INDEX_NONE || !BoneValid[SkeletonPoseBoneIndex] || BoneFilter.Contains(BoneIndex))
 		{
 			continue;
 		}
@@ -492,6 +492,18 @@ void FAnimNode_DeadBlending::ApplyTo(FCompactPose& InOutPose, FBlendedCurve& InO
 	}
 }
 
+class USkeleton* FAnimNode_DeadBlending::GetSkeleton(bool& bInvalidSkeletonIsError, const IPropertyHandle* PropertyHandle)
+{
+	bInvalidSkeletonIsError = false;
+
+	if (const IAnimClassInterface* AnimClassInterface = GetAnimClassInterface())
+	{
+		return AnimClassInterface->GetTargetSkeleton();
+	}
+
+	return nullptr;
+}
+
 FAnimNode_DeadBlending::FAnimNode_DeadBlending() {}
 
 void FAnimNode_DeadBlending::RequestInertialization(const FInertializationRequest& Request)
@@ -514,6 +526,8 @@ void FAnimNode_DeadBlending::Initialize_AnyThread(const FAnimationInitializeCont
 	CurveFilter.SetFilterMode(UE::Anim::ECurveFilterMode::DisallowFiltered);
 	CurveFilter.AppendNames(FilteredCurves);
 
+	BoneFilter.Init(FCompactPoseBoneIndex(INDEX_NONE), FilteredBones.Num());
+
 	PoseSnapshots.Empty(UE::Anim::DeadBlending::Private::MaxPoseSnapShotNum);
 
 	RequestQueue.Reserve(8);
@@ -556,6 +570,15 @@ void FAnimNode_DeadBlending::CacheBones_AnyThread(const FAnimationCacheBonesCont
 
 	FAnimNode_Base::CacheBones_AnyThread(Context);
 	Source.CacheBones(Context);
+
+	// Compute Compact Pose Bone Index for each bone in Filter
+
+	const FBoneContainer& RequiredBones = Context.AnimInstanceProxy->GetRequiredBones();
+	for (int32 FilterBoneIdx = 0; FilterBoneIdx < FilteredBones.Num(); FilterBoneIdx++)
+	{
+		FilteredBones[FilterBoneIdx].Initialize(Context.AnimInstanceProxy->GetSkeleton());
+		BoneFilter[FilterBoneIdx] = FilteredBones[FilterBoneIdx].GetCompactPoseIndex(RequiredBones);
+	}
 }
 
 
