Fixed landscape component count display in editor when selecting streaming proxies

diff --git a/Engine/Source/Editor/DetailCustomizations/Private/LandscapeProxyUIDetails.cpp b/Engine/Source/Editor/DetailCustomizations/Private/LandscapeProxyUIDetails.cpp
--- a/Engine/Source/Editor/DetailCustomizations/Private/LandscapeProxyUIDetails.cpp
+++ b/Engine/Source/Editor/DetailCustomizations/Private/LandscapeProxyUIDetails.cpp
@@ -84,6 +84,7 @@ void FLandscapeProxyUIDetails::CustomizeDetails( IDetailLayoutBuilder& DetailBui
 	Algo::Transform(EditingObjects, EditingProxies, [](TWeakObjectPtr<UObject> InObject) { return TWeakObjectPtr<ALandscapeProxy>(Cast<ALandscapeProxy>(InObject.Get())); });
 
 	TArray<TWeakObjectPtr<ALandscape>> LandscapeActors;
+	TArray<TWeakObjectPtr<ALandscapeStreamingProxy>> EditingStreamingProxies;
 	for (TWeakObjectPtr<ALandscapeProxy> EditingProxy : EditingProxies)
 	{
 		if (EditingProxy.IsValid())
@@ -92,6 +93,11 @@ void FLandscapeProxyUIDetails::CustomizeDetails( IDetailLayoutBuilder& DetailBui
 			{
 				LandscapeActors.AddUnique(LandscapeActor);
 			}
+
+			if (ALandscapeStreamingProxy* LandscapeStreamingProxy = Cast<ALandscapeStreamingProxy>(EditingProxy.Get()))
+			{
+				EditingStreamingProxies.Add(LandscapeStreamingProxy);
+			}
 		}
 	}
 
@@ -134,6 +140,16 @@ void FLandscapeProxyUIDetails::CustomizeDetails( IDetailLayoutBuilder& DetailBui
 			];
 
 			RowDisplayText = LOCTEXT("LandscapeComponentCount", "Component Count");
+			int32 NumComponents = LandscapeActor->LandscapeComponents.Num();
+			// If displaying streaming proxies, don't show the main landscape actor's component count, which is always 0 : 
+			if (!EditingStreamingProxies.IsEmpty())
+			{
+				NumComponents = EditingStreamingProxies[0]->LandscapeComponents.Num();
+				if (Algo::AnyOf(EditingStreamingProxies, [NumComponents](const TWeakObjectPtr<ALandscapeStreamingProxy> InStreamingProxy) { return InStreamingProxy.Get()->LandscapeComponents.Num() != NumComponents; }))
+				{
+					NumComponents = -1;
+				}
+			}
 			CategoryBuilder.AddCustomRow(RowDisplayText)
 			.RowTag(TEXT("LandscapeComponentCount"))
 			.NameContent()
@@ -142,7 +158,7 @@ void FLandscapeProxyUIDetails::CustomizeDetails( IDetailLayoutBuilder& DetailBui
 			]
 			.ValueContent()
 			[
-				GenerateTextWidget(FText::Format(LOCTEXT("LandscapeComponentCountValue", "{0}"), LandscapeActor->LandscapeComponents.Num()), true)
+				GenerateTextWidget((NumComponents == -1) ? LOCTEXT("MultipleValues", "Multiple Values") : FText::Format(LOCTEXT("LandscapeComponentCountValue", "{0}"), NumComponents), true)
 			];
 
 			RowDisplayText = LOCTEXT("LandscapeComponentSubsections", "Component Subsections");
@@ -328,18 +344,6 @@ void FLandscapeProxyUIDetails::CustomizeDetails( IDetailLayoutBuilder& DetailBui
 		];
 	}
 
-	TArray<TWeakObjectPtr<ALandscapeStreamingProxy>> EditingStreamingProxies;
-
-	for (const TWeakObjectPtr<UObject> EditingObject : EditingObjects)
-	{
-		TWeakObjectPtr<ALandscapeStreamingProxy> LandscapeStreamingProxyPtr = Cast<ALandscapeStreamingProxy>(EditingObject.Get());
-
-		if (LandscapeStreamingProxyPtr.IsValid())
-		{
-			EditingStreamingProxies.Add(MoveTemp(LandscapeStreamingProxyPtr));
-		}
-	}
-
 	ALandscapeStreamingProxy* LandscapeStreamingProxy = EditingStreamingProxies.IsEmpty() ? nullptr : EditingStreamingProxies[0].Get();
 
 	if (LandscapeStreamingProxy != nullptr)
