Changed PerPlatformProperty widgets to adhere to IPropertyHandle::IsEditable

diff --git a/Engine/Source/Editor/DetailCustomizations/Private/PerPlatformPropertyCustomization.cpp b/Engine/Source/Editor/DetailCustomizations/Private/PerPlatformPropertyCustomization.cpp
--- a/Engine/Source/Editor/DetailCustomizations/Private/PerPlatformPropertyCustomization.cpp
+++ b/Engine/Source/Editor/DetailCustomizations/Private/PerPlatformPropertyCustomization.cpp
@@ -61,6 +61,10 @@ void FPerPlatformPropertyCustomization<PerPlatformType>::CustomizeChildren(TShar
 	{
 		return GetWidget(PlatformGroupName, StructPropertyHandle, StructBuilder);
 	});
+	Args.IsEnabled = TAttribute<bool>::CreateLambda([StructPropertyHandle]()
+	{		
+		return StructPropertyHandle->IsEditable();
+	});		
 	
 	StructBuilder.AddCustomBuilder(MakeShared<FPerPlatformPropertyCustomNodeBuilder>(MoveTemp(Args)));
 }
