Added compiler note and visual warnings to state machine transitions if the transition rule has logic but the flag for automatic rule based condition is set

diff --git a/Engine/Source/Editor/AnimGraph/Private/AnimGraphNode_StateMachineBase.cpp b/Engine/Source/Editor/AnimGraph/Private/AnimGraphNode_StateMachineBase.cpp
--- a/Engine/Source/Editor/AnimGraph/Private/AnimGraphNode_StateMachineBase.cpp
+++ b/Engine/Source/Editor/AnimGraph/Private/AnimGraphNode_StateMachineBase.cpp
@@ -689,6 +689,23 @@ void UAnimGraphNode_StateMachineBase::OnProcessDuringCompilation(IAnimBlueprintC
 					Rule.CanTakeDelegateIndex = Extension->ExpandGraphAndProcessNodes(TransitionNode->BoundGraph, TransitionResultNode, InCompilationContext, OutCompiledData, TransitionNode);
 					AlreadyMergedTransitionList.Add(TransitionResultNode, Rule.CanTakeDelegateIndex);
 				}
+
+				if (TransitionNode->bAutomaticRuleBasedOnSequencePlayerInState)
+				{
+					if (UAnimationTransitionGraph* TransGraph = Cast<UAnimationTransitionGraph>(TransitionNode->BoundGraph))
+					{
+						if (UAnimGraphNode_TransitionResult* ResultNode = TransGraph->GetResultNode())
+						{
+							if (UEdGraphPin* CanExecPin = ResultNode->FindPin(TEXT("bCanEnterTransition")))
+							{
+								if (CanExecPin->LinkedTo.Num() > 0)
+								{
+									InCompilationContext.GetMessageLog().Note(*LOCTEXT("TransitionWithAutomaticRuleBased", "@@ has an automatic Rule Based Transition that will override graph exit rule.").ToString(), TransitionNode);
+								}
+							}
+						}
+					}
+				}
 			}
 			else
 			{
diff --git a/Engine/Source/Editor/AnimGraph/Private/AnimGraphNode_TransitionResult.cpp b/Engine/Source/Editor/AnimGraph/Private/AnimGraphNode_TransitionResult.cpp
--- a/Engine/Source/Editor/AnimGraph/Private/AnimGraphNode_TransitionResult.cpp
+++ b/Engine/Source/Editor/AnimGraph/Private/AnimGraphNode_TransitionResult.cpp
@@ -2,6 +2,7 @@
 
 #include "AnimGraphNode_TransitionResult.h"
 #include "GraphEditorSettings.h"
+#include "AnimStateTransitionNode.h"
 
 
 /////////////////////////////////////////////////////
@@ -29,6 +30,35 @@ FText UAnimGraphNode_TransitionResult::GetNodeTitle(ENodeTitleType::Type TitleTy
 	return LOCTEXT("Result", "Result");
 }
 
+bool UAnimGraphNode_TransitionResult::ShowVisualWarning() const
+{
+	UObject * Outer = GetOuter();
+	while (Outer != nullptr)
+	{
+		if (UAnimStateTransitionNode* AnimStateTransitionNode = Cast<UAnimStateTransitionNode>(Outer))
+		{
+			// only show warning if the auto rule is set and there is logic connected
+			if (AnimStateTransitionNode->bAutomaticRuleBasedOnSequencePlayerInState)
+			{
+				UEdGraphPin* CanExecPin = FindPin(TEXT("bCanEnterTransition"));
+				if (CanExecPin != nullptr && CanExecPin->LinkedTo.Num() > 0) 
+				{
+					return true;
+				}
+			}
+			break;
+		}
+		Outer = Outer->GetOuter();
+	}
+
+	return false;
+}
+
+FText UAnimGraphNode_TransitionResult::GetVisualWarningTooltipText() const
+{
+	return LOCTEXT("TransitionResult_VisualWarning", "Warning : Automatic Rule Based Transition will override graph exit rule.");
+}
+
 void UAnimGraphNode_TransitionResult::GetMenuActions(FBlueprintActionDatabaseRegistrar& ActionRegistrar) const
 {
 	// Intentionally empty. This node is auto-generated when a transition graph is created.
diff --git a/Engine/Source/Editor/AnimGraph/Public/AnimGraphNode_TransitionResult.h b/Engine/Source/Editor/AnimGraph/Public/AnimGraphNode_TransitionResult.h
--- a/Engine/Source/Editor/AnimGraph/Public/AnimGraphNode_TransitionResult.h
+++ b/Engine/Source/Editor/AnimGraph/Public/AnimGraphNode_TransitionResult.h
@@ -25,6 +25,8 @@ class UAnimGraphNode_TransitionResult : public UAnimGraphNode_Base
 	virtual FText GetTooltipText() const override;
 	virtual FText GetNodeTitle(ENodeTitleType::Type TitleType) const override;
 	virtual bool IsNodeRootSet() const override { return true; }
+	virtual bool ShowVisualWarning() const override;
+	virtual FText GetVisualWarningTooltipText() const override;
 	// End of UEdGraphNode interface
 
 	// UK2Node interface.
diff --git a/Engine/Source/Editor/AnimationBlueprintEditor/Private/AnimationStateNodes/SGraphNodeAnimTransition.cpp b/Engine/Source/Editor/AnimationBlueprintEditor/Private/AnimationStateNodes/SGraphNodeAnimTransition.cpp
--- a/Engine/Source/Editor/AnimationBlueprintEditor/Private/AnimationStateNodes/SGraphNodeAnimTransition.cpp
+++ b/Engine/Source/Editor/AnimationBlueprintEditor/Private/AnimationStateNodes/SGraphNodeAnimTransition.cpp
@@ -158,14 +158,29 @@ TSharedRef<SWidget> SGraphNodeAnimTransition::GenerateRichTooltip()
 
 	if(TransNode->bAutomaticRuleBasedOnSequencePlayerInState)
 	{
-		Widget->AddSlot()
-			.AutoHeight()
-			.Padding( 2.0f )
-			[
-				SNew(STextBlock)
-				.TextStyle( FAppStyle::Get(), TEXT("Graph.TransitionNode.TooltipRule") )
-				.Text(LOCTEXT("AnimGraphNodeAutomaticRule_ToolTip", "Automatic Rule"))
-			];
+		if (CanExecPin->LinkedTo.Num() > 0)
+		{
+			Widget->AddSlot()
+				.AutoHeight()
+				.Padding(2.0f)
+				[
+					SNew(STextBlock)
+					.TextStyle(FAppStyle::Get(), TEXT("Graph.TransitionNode.TooltipRule"))
+					.Text(LOCTEXT("AnimGraphNodeAutomaticRuleWarning_ToolTip", "Warning : Automatic Rule Based Transition will override graph exit rule."))
+					.ColorAndOpacity(FCoreStyle::Get().GetColor("ErrorReporting.WarningBackgroundColor"))
+				];
+		}
+		else
+		{
+			Widget->AddSlot()
+				.AutoHeight()
+				.Padding(2.0f)
+				[
+					SNew(STextBlock)
+					.TextStyle(FAppStyle::Get(), TEXT("Graph.TransitionNode.TooltipRule"))
+					.Text(LOCTEXT("AnimGraphNodeAutomaticRule_ToolTip", "Automatic Rule"))
+				];
+		}
 	}
 	else
 	{
diff --git a/Engine/Source/Editor/DetailCustomizations/Private/AnimTransitionNodeDetails.cpp b/Engine/Source/Editor/DetailCustomizations/Private/AnimTransitionNodeDetails.cpp
--- a/Engine/Source/Editor/DetailCustomizations/Private/AnimTransitionNodeDetails.cpp
+++ b/Engine/Source/Editor/DetailCustomizations/Private/AnimTransitionNodeDetails.cpp
@@ -175,36 +175,71 @@ void FAnimTransitionNodeDetails::CustomizeDetails( IDetailLayoutBuilder& DetailB
 				}
 			}
 
-			// indicate if a native transition rule applies to this
-			UBlueprint* Blueprint = FBlueprintEditorUtils::FindBlueprintForNodeChecked(TransitionNode.Get());
-			if(Blueprint && Blueprint->ParentClass)
+			if (TransNode->bAutomaticRuleBasedOnSequencePlayerInState)
 			{
-				UAnimInstance* AnimInstance = CastChecked<UAnimInstance>(Blueprint->ParentClass->GetDefaultObject());
-				if(AnimInstance)
+				if (CanExecPin->LinkedTo.Num() > 0)
 				{
-					UEdGraph* ParentGraph = TransitionNode->GetGraph();
-					UAnimStateNodeBase* PrevState = TransitionNode->GetPreviousState();
-					UAnimStateNodeBase* NextState = TransitionNode->GetNextState();
-					if(PrevState != nullptr && NextState != nullptr && ParentGraph != nullptr)
+					TransitionCategory.AddCustomRow(LOCTEXT("AnimGraphNodeDetailsAutomaticRule_RowWarning", "Automatic Rule"))
+					[
+						SNew(SBox)
+						.VAlign(VAlign_Center)
+						.HAlign(HAlign_Left)
+						[
+							SNew(STextBlock)
+							.Text(LOCTEXT("AnimGraphNodeDetailsAutomaticRule_Warning", "Warning : Automatic Rule Based Transition will override graph exit rule."))
+							.ColorAndOpacity(FCoreStyle::Get().GetColor("ErrorReporting.WarningBackgroundColor"))
+							.Font(IDetailLayoutBuilder::GetDetailFontBold())
+						]
+					];
+				}
+				else
+				{
+					TransitionCategory.AddCustomRow(LOCTEXT("AnimGraphNodeDetailsAutomaticRule_Row", "Automatic Rule"))
+					[
+						SNew(SBox)
+						.VAlign(VAlign_Center)
+						.HAlign(HAlign_Left)
+						[
+							SNew(STextBlock)
+							.Text(LOCTEXT("AnimGraphNodeDetailsAutomaticRule", "Automatic Rule Based Transition"))
+							.Font(IDetailLayoutBuilder::GetDetailFontBold())
+						]
+					];
+				}
+			}
+			else
+			{
+				// indicate if a native transition rule applies to this
+				UBlueprint* Blueprint = FBlueprintEditorUtils::FindBlueprintForNodeChecked(TransitionNode.Get());
+				if(Blueprint && Blueprint->ParentClass)
+				{
+					UAnimInstance* AnimInstance = CastChecked<UAnimInstance>(Blueprint->ParentClass->GetDefaultObject());
+					if(AnimInstance)
 					{
-						FName FunctionName;
-						if(AnimInstance->HasNativeTransitionBinding(ParentGraph->GetFName(), FName(*PrevState->GetStateName()), FName(*NextState->GetStateName()), FunctionName))
+						UEdGraph* ParentGraph = TransitionNode->GetGraph();
+						UAnimStateNodeBase* PrevState = TransitionNode->GetPreviousState();
+						UAnimStateNodeBase* NextState = TransitionNode->GetNextState();
+						if(PrevState != nullptr && NextState != nullptr && ParentGraph != nullptr)
 						{
-							TransitionCategory.AddCustomRow( LOCTEXT("NativeBindingPresent_Filter", "Transition has native binding") )
-							[
-								SNew(STextBlock)
-								.Text(FText::Format(LOCTEXT("NativeBindingPresent", "Transition has native binding to {0}()"), FText::FromName(FunctionName)))
-								.Font( IDetailLayoutBuilder::GetDetailFontBold() )
-							];
+							FName FunctionName;
+							if(AnimInstance->HasNativeTransitionBinding(ParentGraph->GetFName(), FName(*PrevState->GetStateName()), FName(*NextState->GetStateName()), FunctionName))
+							{
+								TransitionCategory.AddCustomRow( LOCTEXT("NativeBindingPresent_Filter", "Transition has native binding") )
+								[
+									SNew(STextBlock)
+									.Text(FText::Format(LOCTEXT("NativeBindingPresent", "Transition has native binding to {0}()"), FText::FromName(FunctionName)))
+									.Font( IDetailLayoutBuilder::GetDetailFontBold() )
+								];
+							}
 						}
 					}
 				}
-			}
 
-			TransitionCategory.AddCustomRow( CanExecPin ? CanExecPin->PinFriendlyName : FText::GetEmpty() )
-			[
-				SNew(SKismetLinearExpression, CanExecPin)
-			];
+				TransitionCategory.AddCustomRow( CanExecPin ? CanExecPin->PinFriendlyName : FText::GetEmpty() )
+				[
+					SNew(SKismetLinearExpression, CanExecPin)
+				];
+			}
 		}
 
 		//////////////////////////////////////////////////////////////////////////
