Static analysis fix at SGraphNodeAnimTransition.cpp and AnimTransitionNodeDetails.cpp

diff --git a/Engine/Source/Editor/AnimationBlueprintEditor/Private/AnimationStateNodes/SGraphNodeAnimTransition.cpp b/Engine/Source/Editor/AnimationBlueprintEditor/Private/AnimationStateNodes/SGraphNodeAnimTransition.cpp
--- a/Engine/Source/Editor/AnimationBlueprintEditor/Private/AnimationStateNodes/SGraphNodeAnimTransition.cpp
+++ b/Engine/Source/Editor/AnimationBlueprintEditor/Private/AnimationStateNodes/SGraphNodeAnimTransition.cpp
@@ -158,7 +158,7 @@ TSharedRef<SWidget> SGraphNodeAnimTransition::GenerateRichTooltip()
 
 	if(TransNode->bAutomaticRuleBasedOnSequencePlayerInState)
 	{
-		if (CanExecPin->LinkedTo.Num() > 0)
+		if (CanExecPin != nullptr && CanExecPin->LinkedTo.Num() > 0)
 		{
 			Widget->AddSlot()
 				.AutoHeight()
diff --git a/Engine/Source/Editor/DetailCustomizations/Private/AnimTransitionNodeDetails.cpp b/Engine/Source/Editor/DetailCustomizations/Private/AnimTransitionNodeDetails.cpp
--- a/Engine/Source/Editor/DetailCustomizations/Private/AnimTransitionNodeDetails.cpp
+++ b/Engine/Source/Editor/DetailCustomizations/Private/AnimTransitionNodeDetails.cpp
@@ -177,7 +177,7 @@ void FAnimTransitionNodeDetails::CustomizeDetails( IDetailLayoutBuilder& DetailB
 
 			if (TransNode->bAutomaticRuleBasedOnSequencePlayerInState)
 			{
-				if (CanExecPin->LinkedTo.Num() > 0)
+				if (CanExecPin != nullptr && CanExecPin->LinkedTo.Num() > 0)
 				{
 					TransitionCategory.AddCustomRow(LOCTEXT("AnimGraphNodeDetailsAutomaticRule_RowWarning", "Automatic Rule"))
 					[
