[PropertyEditor] SlateFontInfo dropdown crash

diff --git a/Engine/Source/Editor/DetailCustomizations/Private/SlateFontInfoCustomization.cpp b/Engine/Source/Editor/DetailCustomizations/Private/SlateFontInfoCustomization.cpp
--- a/Engine/Source/Editor/DetailCustomizations/Private/SlateFontInfoCustomization.cpp
+++ b/Engine/Source/Editor/DetailCustomizations/Private/SlateFontInfoCustomization.cpp
@@ -110,7 +110,6 @@ void FSlateFontInfoStructCustomization::CustomizeChildren(TSharedRef<IPropertyHa
 		.OnComboBoxOpening(this, &FSlateFontInfoStructCustomization::OnFontEntryComboOpening)
 		.OnSelectionChanged(this, &FSlateFontInfoStructCustomization::OnFontEntrySelectionChanged)
 		.OnGenerateWidget(this, &FSlateFontInfoStructCustomization::MakeFontEntryWidget)
-		.IsEnabled(!TypefaceFontNameProperty->IsEditConst())
 		[
 			SNew(STextBlock)
 			.Text(this, &FSlateFontInfoStructCustomization::GetFontEntryComboText)
@@ -325,6 +324,11 @@ void FSlateFontInfoStructCustomization::OnFontChanged(const FAssetData& InAssetD
 
 bool FSlateFontInfoStructCustomization::IsFontEntryComboEnabled() const
 {
+	if (TypefaceFontNameProperty->IsEditConst())
+	{
+		return false;
+	}
+	
 	TArray<const FSlateFontInfo*> SlateFontInfoStructs = GetFontInfoBeingEdited();
 	if(SlateFontInfoStructs.Num() == 0)
 	{
