Fixed "Dynamic Shadow Distance" being unlocalizable

#rnx

diff --git a/Engine/Source/Editor/DetailCustomizations/Private/DirectionalLightComponentDetails.cpp b/Engine/Source/Editor/DetailCustomizations/Private/DirectionalLightComponentDetails.cpp
--- a/Engine/Source/Editor/DetailCustomizations/Private/DirectionalLightComponentDetails.cpp
+++ b/Engine/Source/Editor/DetailCustomizations/Private/DirectionalLightComponentDetails.cpp
@@ -2,6 +2,7 @@
 
 #include "DirectionalLightComponentDetails.h"
 
+#include "Components/DirectionalLightComponent.h"
 #include "Components/LightComponentBase.h"
 #include "Components/SceneComponent.h"
 #include "Delegates/Delegate.h"
@@ -26,11 +27,8 @@ TSharedRef<IDetailCustomization> FDirectionalLightComponentDetails::MakeInstance
 
 void FDirectionalLightComponentDetails::CustomizeDetails( IDetailLayoutBuilder& DetailBuilder )
 {
-	// Get cascaded shadow map category
-	IDetailCategoryBuilder& ShadowMapCategory = DetailBuilder.EditCategory("CascadedShadowMaps", FText::GetEmpty(), ECategoryPriority::Default );
-
-	TSharedPtr<IPropertyHandle> MovableShadowRadiusPropertyHandle = DetailBuilder.GetProperty("DynamicShadowDistanceMovableLight");
-	TSharedPtr<IPropertyHandle> StationaryShadowRadiusPropertyHandle = DetailBuilder.GetProperty("DynamicShadowDistanceStationaryLight");
+	TSharedPtr<IPropertyHandle> MovableShadowRadiusPropertyHandle = DetailBuilder.GetProperty(GET_MEMBER_NAME_CHECKED(UDirectionalLightComponent, DynamicShadowDistanceMovableLight));
+	TSharedPtr<IPropertyHandle> StationaryShadowRadiusPropertyHandle = DetailBuilder.GetProperty(GET_MEMBER_NAME_CHECKED(UDirectionalLightComponent, DynamicShadowDistanceStationaryLight));
 
 	static const auto AllowStaticLightingVar = IConsoleManager::Get().FindTConsoleVariableDataInt(TEXT("r.AllowStaticLighting"));
 	const bool bAllowStaticLighting = (!AllowStaticLightingVar || AllowStaticLightingVar->GetValueOnGameThread() != 0);
@@ -38,15 +36,11 @@ void FDirectionalLightComponentDetails::CustomizeDetails( IDetailLayoutBuilder&
 	if(!bAllowStaticLighting)
 	{
 		// If static lighting is not allowed, hide DynamicShadowDistanceStationaryLight and rename DynamicShadowDistanceMovableLight to "Dynamic Shadow Distance"
-
-		FProperty* MovableShadowRadiusProperty = MovableShadowRadiusPropertyHandle->GetProperty();
-		MovableShadowRadiusProperty->SetMetaData(TEXT("DisplayName"), TEXT("Dynamic Shadow Distance"));
-
-		ShadowMapCategory.AddProperty(StationaryShadowRadiusPropertyHandle)
-			.Visibility(EVisibility::Hidden);
+		MovableShadowRadiusPropertyHandle->SetPropertyDisplayName(LOCTEXT("DynamicShadowDistanceDisplayName", "Dynamic Shadow Distance"));
+		StationaryShadowRadiusPropertyHandle->MarkHiddenByCustomization();
 	}
 
-	TSharedPtr<IPropertyHandle> LightIntensityProperty = DetailBuilder.GetProperty("Intensity", ULightComponentBase::StaticClass());
+	TSharedPtr<IPropertyHandle> LightIntensityProperty = DetailBuilder.GetProperty(GET_MEMBER_NAME_CHECKED(ULightComponentBase, Intensity), ULightComponentBase::StaticClass());
 	// Point lights need to override the ui min and max for units of lumens, so we have to undo that
 	LightIntensityProperty->SetInstanceMetaData("UIMin",TEXT("0.0f"));
 	LightIntensityProperty->SetInstanceMetaData("UIMax",TEXT("150.0f"));
