Path Tracer: Add Path Tracing properties to Rendering quick category

diff --git a/Engine/Source/Editor/DetailCustomizations/Private/DetailCustomizations.cpp b/Engine/Source/Editor/DetailCustomizations/Private/DetailCustomizations.cpp
--- a/Engine/Source/Editor/DetailCustomizations/Private/DetailCustomizations.cpp
+++ b/Engine/Source/Editor/DetailCustomizations/Private/DetailCustomizations.cpp
@@ -575,6 +575,7 @@ void FDetailCustomizationsModule::RegisterSectionMappings()
 			Section->AddCategory("Materials");
 			Section->AddCategory("Mobile");
 			Section->AddCategory("Ray Tracing");
+			Section->AddCategory("Path Tracing");
 			Section->AddCategory("Rendering");
 			Section->AddCategory("Texture Streaming");
 			Section->AddCategory("Virtual Texture");
