Changed CallInEditor UFunction button label to respect UFunction display name

diff --git a/Engine/Source/Editor/DetailCustomizations/Private/ObjectDetails.cpp b/Engine/Source/Editor/DetailCustomizations/Private/ObjectDetails.cpp
--- a/Engine/Source/Editor/DetailCustomizations/Private/ObjectDetails.cpp
+++ b/Engine/Source/Editor/DetailCustomizations/Private/ObjectDetails.cpp
@@ -15,6 +15,7 @@
 #include "HAL/PlatformCrt.h"
 #include "Internationalization/Internationalization.h"
 #include "Internationalization/Text.h"
+#include "K2Node_CallFunction.h"
 #include "Kismet2/BlueprintEditorUtils.h"
 #include "Layout/Margin.h"
 #include "Math/NumericLimits.h"
@@ -193,14 +194,12 @@ void FObjectDetails::AddCallInEditorMethods(IDetailLayoutBuilder& DetailBuilder)
 			}
 			FCategoryEntry& CategoryEntry = CategoryList.Last();
 
-			//@TODO: Expose the code in UK2Node_CallFunction::GetUserFacingFunctionName / etc...
-			const FText ButtonCaption = FText::FromString(FName::NameToDisplayString(*Function->GetName(), false));
+			const FText ButtonCaption = UK2Node_CallFunction::GetUserFacingFunctionName(Function);
 			FText FunctionTooltip = Function->GetToolTipText();
 			if (FunctionTooltip.IsEmpty())
 			{
-				FunctionTooltip = FText::FromString(Function->GetName());
+				FunctionTooltip = ButtonCaption;
 			}
-			
 
 			TWeakObjectPtr<UFunction> WeakFunctionPtr(Function);
 			CategoryEntry.WrapBox->AddSlot()
@@ -215,6 +214,11 @@ void FObjectDetails::AddCallInEditorMethods(IDetailLayoutBuilder& DetailBuilder)
 			CategoryEntry.RowTag = Function->GetFName();
 			CategoryEntry.FunctionSearchText.AppendLine(ButtonCaption);
 			CategoryEntry.FunctionSearchText.AppendLine(FunctionTooltip);
+
+			if (ButtonCaption.ToString() != Function->GetName())
+			{
+				CategoryEntry.FunctionSearchText.AppendLine(FText::FromString(Function->GetName()));
+			}
 		}
 		
 		// Now edit the categories, adding the button strips to the details panel
