Add Probe and QueryAndProbe to CollisionProfileDetails panel

diff --git a/Engine/Source/Editor/DetailCustomizations/Private/CollisionProfileDetails.cpp b/Engine/Source/Editor/DetailCustomizations/Private/CollisionProfileDetails.cpp
--- a/Engine/Source/Editor/DetailCustomizations/Private/CollisionProfileDetails.cpp
+++ b/Engine/Source/Editor/DetailCustomizations/Private/CollisionProfileDetails.cpp
@@ -676,6 +676,8 @@ void SProfileEditDialog::FillCollisionEnabledString()
 	CollisionEnabledComboBoxString.Add(MakeShareable(new FString(TEXT("Query Only (No Physics Collision)"))));
 	CollisionEnabledComboBoxString.Add(MakeShareable(new FString(TEXT("Physics Only (No Query Collision)"))));
 	CollisionEnabledComboBoxString.Add(MakeShareable(new FString(TEXT("Collision Enabled (Query and Physics)"))));
+	CollisionEnabledComboBoxString.Add(MakeShareable(new FString(TEXT("Probe Only (No Query or Physics Collision)"))));
+	CollisionEnabledComboBoxString.Add(MakeShareable(new FString(TEXT("Query and Probe (No Physics Collision)"))));
 }
 
 bool	SProfileEditDialog::IsAcceptAvailable() const
@@ -791,7 +793,7 @@ void SProfileEditDialog::HandleCollisionEnabledComboBoxSelectionChanged(TSharedP
 		if(*Iter == StringItem)
 		{
 			ECollisionEnabled::Type NewCollisionEnabled = (ECollisionEnabled::Type)Iter.GetIndex();
-			check(NewCollisionEnabled >= ECollisionEnabled::NoCollision && NewCollisionEnabled <= ECollisionEnabled::QueryAndPhysics);
+			check(NewCollisionEnabled >= ECollisionEnabled::NoCollision && NewCollisionEnabled <= ECollisionEnabled::QueryAndProbe);
 			ProfileTemplate.CollisionEnabled = NewCollisionEnabled;
 			return;
 		}
@@ -1264,6 +1266,10 @@ FText SProfileListItem::GetCollsionEnabled() const
 		return LOCTEXT("ECollisionEnabled_PhysicsOnly", "Physics Only (No Query Collision)");
 	case ECollisionEnabled::QueryAndPhysics:
 		return LOCTEXT("ECollisionEnabled_QueryAndPhysics", "Collision Enabled (Query and Physics)");
+	case ECollisionEnabled::ProbeOnly:
+		return LOCTEXT("ECollisionEnabled_Probe", "Probe Only (No Query or Physics Collision)");
+	case ECollisionEnabled::QueryAndProbe:
+		return LOCTEXT("ECollisionEnabled_QueryAndProbe", "Query and Probe (No Physics Collision)");
 	}
 
 	return LOCTEXT("ECollisionEnabled_Error", "ERROR");
