Make some functionality virtual so that it can be overriden by a derivation. This allows interfacing with a type that doesn't follow the same color requirements.

diff --git a/Engine/Source/Editor/DetailCustomizations/Private/ColorStructCustomization.cpp b/Engine/Source/Editor/DetailCustomizations/Private/ColorStructCustomization.cpp
--- a/Engine/Source/Editor/DetailCustomizations/Private/ColorStructCustomization.cpp
+++ b/Engine/Source/Editor/DetailCustomizations/Private/ColorStructCustomization.cpp
@@ -63,7 +63,7 @@ void FColorStructCustomization::CustomizeHeader(TSharedRef<class IPropertyHandle
 	StructPropertyHandle = InStructPropertyHandle;
 
 	bIsLinearColor = CastFieldChecked<FStructProperty>(StructPropertyHandle->GetProperty())->Struct->GetFName() == NAME_LinearColor;
-	bIgnoreAlpha = StructPropertyHandle->GetProperty()->HasMetaData(TEXT("HideAlphaChannel"));
+	bIgnoreAlpha = TypeSupportsAlpha() == false || StructPropertyHandle->GetProperty()->HasMetaData(TEXT("HideAlphaChannel"));
 	
 	if (StructPropertyHandle->GetProperty()->HasMetaData(TEXT("sRGB")))
 	{
@@ -204,11 +204,8 @@ void FColorStructCustomization::GetSortedChildren(TSharedRef<IPropertyHandle> In
 	}
 }
 
-
-void FColorStructCustomization::CreateColorPicker(bool bUseAlpha)
+void FColorStructCustomization::GatherSavedPreColorPickerColors()
 {
-	GEditor->BeginTransaction(FText::Format(LOCTEXT("SetColorProperty", "Edit {0}"), StructPropertyHandle->GetPropertyDisplayName()));
-
 	SavedPreColorPickerColors.Empty();
 	TArray<FString> PerObjectValues;
 	StructPropertyHandle->GetPerObjectValues(PerObjectValues);
@@ -228,6 +225,13 @@ void FColorStructCustomization::CreateColorPicker(bool bUseAlpha)
 			SavedPreColorPickerColors.Add(FLinearOrSrgbColor(Color));
 		}
 	}
+}
+
+void FColorStructCustomization::CreateColorPicker(bool bUseAlpha)
+{
+	GEditor->BeginTransaction(FText::Format(LOCTEXT("SetColorProperty", "Edit {0}"), StructPropertyHandle->GetPropertyDisplayName()));
+
+	GatherSavedPreColorPickerColors();
 
 	FLinearColor InitialColor;
 	GetColorAsLinear(InitialColor);
@@ -265,25 +269,7 @@ TSharedRef<SColorPicker> FColorStructCustomization::CreateInlineColorPicker(TWea
 {
 	GEditor->BeginTransaction(FText::Format(LOCTEXT("SetColorProperty", "Edit {0}"), StructPropertyHandle->GetPropertyDisplayName()));
 
-	SavedPreColorPickerColors.Empty();
-	TArray<FString> PerObjectValues;
-	StructPropertyHandle->GetPerObjectValues(PerObjectValues);
-
-	for (int32 ObjectIndex = 0; ObjectIndex < PerObjectValues.Num(); ++ObjectIndex)
-	{
-		if (bIsLinearColor)
-		{
-			FLinearColor Color;
-			Color.InitFromString(PerObjectValues[ObjectIndex]);
-			SavedPreColorPickerColors.Add(FLinearOrSrgbColor(Color));	
-		}
-		else
-		{
-			FColor Color;
-			Color.InitFromString(PerObjectValues[ObjectIndex]);
-			SavedPreColorPickerColors.Add(FLinearOrSrgbColor(Color));
-		}
-	}
+	GatherSavedPreColorPickerColors();
 
 	FLinearColor InitialColor;
 	GetColorAsLinear(InitialColor);
@@ -306,8 +292,7 @@ TSharedRef<SColorPicker> FColorStructCustomization::CreateInlineColorPicker(TWea
 		.IsEnabled(this, &FColorStructCustomization::IsValueEnabled, StructWeakHandlePtr);
 }
 
-
-void FColorStructCustomization::OnSetColorFromColorPicker(FLinearColor NewColor)
+void FColorStructCustomization::SetLastPickerColorString(const FLinearColor NewColor)
 {
 	if (bIsLinearColor)
 	{
@@ -319,6 +304,11 @@ void FColorStructCustomization::OnSetColorFromColorPicker(FLinearColor NewColor)
 		FColor NewFColor = NewColor.ToFColor(bSRGB);
 		LastPickerColorString = NewFColor.ToString();
 	}
+}
+
+void FColorStructCustomization::OnSetColorFromColorPicker(FLinearColor NewColor)
+{
+	SetLastPickerColorString(NewColor);
 
 	EPropertyValueSetFlags::Type PropertyFlags = EPropertyValueSetFlags::NotTransactable;
 	PropertyFlags |= bIsInteractive ? EPropertyValueSetFlags::InteractiveChange : 0;
@@ -326,7 +316,7 @@ void FColorStructCustomization::OnSetColorFromColorPicker(FLinearColor NewColor)
 	StructPropertyHandle->NotifyFinishedChangingProperties();
 }
 
-void FColorStructCustomization::ResetColors()
+TArray<FString> FColorStructCustomization::ConvertToPerObjectColors(const TArray<FLinearOrSrgbColor>& Colors) const
 {
 	TArray<FString> PerObjectColors;
 
@@ -343,6 +333,13 @@ void FColorStructCustomization::ResetColors()
 		}
 	}
 
+	return PerObjectColors;
+}
+
+void FColorStructCustomization::ResetColors()
+{
+	TArray<FString> PerObjectColors = ConvertToPerObjectColors(SavedPreColorPickerColors);
+
 	if (PerObjectColors.Num() > 0)
 	{
 		// See @TODO in FColorStructCustomization::OnColorPickerWindowClosed
diff --git a/Engine/Source/Editor/DetailCustomizations/Public/Customizations/ColorStructCustomization.h b/Engine/Source/Editor/DetailCustomizations/Public/Customizations/ColorStructCustomization.h
--- a/Engine/Source/Editor/DetailCustomizations/Public/Customizations/ColorStructCustomization.h
+++ b/Engine/Source/Editor/DetailCustomizations/Public/Customizations/ColorStructCustomization.h
@@ -51,15 +51,42 @@ protected:
 
 protected:
 
-	/** Creates the color widget that when clicked spawns the color picker window. */
-	TSharedRef<SWidget> CreateColorWidget(TWeakPtr<IPropertyHandle>);
+	struct FLinearOrSrgbColor;
 
 	/**
 	 * Get the color used by this struct as a linear color value
-	 * @param InColor To be filled with the color value used by this struct, or white if this struct is being used to edit multiple values  
+	 * @param InColor To be filled with the color value used by this struct, or white if this struct is being used to edit multiple values
 	 * @return The result of trying to get the color value
 	 */
-	FPropertyAccess::Result GetColorAsLinear(FLinearColor& InColor) const;
+	virtual FPropertyAccess::Result GetColorAsLinear(FLinearColor& InColor) const;
+
+	/**
+	 * Does the type have Alpha Support
+	 * @return true if it does
+	*/
+	virtual bool TypeSupportsAlpha() const { return true; }
+
+	/**
+	 * Stores the colors from the property into SavedPreColorPickerColors
+	*/
+	virtual void GatherSavedPreColorPickerColors();
+
+	/**
+	 * Stores the color as a string in LastPickerColorString
+	*/
+	virtual void SetLastPickerColorString(const FLinearColor NewColor);
+
+	/**
+	 * Converts Colors into strings
+	 * @param Colors Array of colors to convert
+	 * @return Array of colors converted to string
+	*/
+	virtual TArray<FString> ConvertToPerObjectColors(const TArray<FLinearOrSrgbColor>& Colors) const;
+
+protected:
+
+	/** Creates the color widget that when clicked spawns the color picker window. */
+	TSharedRef<SWidget> CreateColorWidget(TWeakPtr<IPropertyHandle>);
 
 	/**
 	 * Does this struct have multiple values?
