ComponentTransformDetails: keep same world rotation and scale when going from relative to absolute and vice versa.
It was already the behavior for location but not for rotation and scale. With this change, all transform components will behave the same when switching from absolute to relative and vice versa.

diff --git a/Engine/Source/Editor/DetailCustomizations/Private/ComponentTransformDetails.cpp b/Engine/Source/Editor/DetailCustomizations/Private/ComponentTransformDetails.cpp
--- a/Engine/Source/Editor/DetailCustomizations/Private/ComponentTransformDetails.cpp
+++ b/Engine/Source/Editor/DetailCustomizations/Private/ComponentTransformDetails.cpp
@@ -749,31 +749,45 @@ void FComponentTransformDetails::OnSetAbsoluteTransform(ETransformField::Type Tr
 					NotifyHook->NotifyPreChange(AbsoluteProperty);
 				}
 
+				TOptional<FTransform> TransformToPreserve;
+				if (SceneComponent->GetAttachParent())
+				{
+					if (bAbsoluteEnabled)
+					{
+						TransformToPreserve = SceneComponent->GetComponentTransform();
+					}
+					else
+					{
+						FTransform ParentToWorld = SceneComponent->GetAttachParent()->GetSocketTransform(SceneComponent->GetAttachSocketName());
+						TransformToPreserve = SceneComponent->GetComponentTransform().GetRelativeTransform(ParentToWorld);
+					}
+				}
+
 				switch (TransformField)
 				{
 				case ETransformField::Location:
 					SceneComponent->SetUsingAbsoluteLocation(bAbsoluteEnabled);
 
-					// Update RelativeLocation to maintain/stabilize position when switching between relative and world.
-					if (SceneComponent->GetAttachParent())
+					if (TransformToPreserve.IsSet())
 					{
-						if (SceneComponent->IsUsingAbsoluteLocation())
-						{
-							SceneComponent->SetRelativeLocation_Direct(SceneComponent->GetComponentTransform().GetTranslation());
-						}
-						else
-						{
-							FTransform ParentToWorld = SceneComponent->GetAttachParent()->GetSocketTransform(SceneComponent->GetAttachSocketName());
-							FTransform RelativeTM = SceneComponent->GetComponentTransform().GetRelativeTransform(ParentToWorld);
-							SceneComponent->SetRelativeLocation_Direct(RelativeTM.GetTranslation());
-						}
+						SceneComponent->SetRelativeLocation_Direct(TransformToPreserve->GetTranslation());
 					}
 					break;
 				case ETransformField::Rotation:
 					SceneComponent->SetUsingAbsoluteRotation(bAbsoluteEnabled);
+
+					if (TransformToPreserve.IsSet())
+					{
+						SceneComponent->SetRelativeRotation_Direct(FRotator(TransformToPreserve->GetRotation()));
+					}
 					break;
 				case ETransformField::Scale:
 					SceneComponent->SetUsingAbsoluteScale(bAbsoluteEnabled);
+
+					if (TransformToPreserve.IsSet())
+					{
+						SceneComponent->SetRelativeScale3D_Direct(TransformToPreserve->GetScale3D());
+					}
 					break;
 				}
 
