Color Customization: Keep track of transactions and end any open transaction when the color picker is closed. This fixes an issue with the inline color picker which doesn't end the transaction when the color picker is closed.

diff --git a/Engine/Source/Editor/DetailCustomizations/Private/ColorStructCustomization.cpp b/Engine/Source/Editor/DetailCustomizations/Private/ColorStructCustomization.cpp
--- a/Engine/Source/Editor/DetailCustomizations/Private/ColorStructCustomization.cpp
+++ b/Engine/Source/Editor/DetailCustomizations/Private/ColorStructCustomization.cpp
@@ -108,6 +108,10 @@ void FColorStructCustomization::MakeHeaderRow(TSharedRef<class IPropertyHandle>&
 
 FColorStructCustomization::~FColorStructCustomization()
 {
+	if (TransactionIndex.IsSet())
+	{
+		GEditor->EndTransaction();
+	}
 }
 
 TSharedRef<SWidget> FColorStructCustomization::CreateColorWidget(TWeakPtr<IPropertyHandle> StructWeakHandlePtr)
@@ -229,7 +233,7 @@ void FColorStructCustomization::GatherSavedPreColorPickerColors()
 
 void FColorStructCustomization::CreateColorPicker(bool bUseAlpha)
 {
-	GEditor->BeginTransaction(FText::Format(LOCTEXT("SetColorProperty", "Edit {0}"), StructPropertyHandle->GetPropertyDisplayName()));
+	TransactionIndex = GEditor->BeginTransaction(FText::Format(LOCTEXT("SetColorProperty", "Edit {0}"), StructPropertyHandle->GetPropertyDisplayName()));
 
 	GatherSavedPreColorPickerColors();
 
@@ -267,7 +271,7 @@ void FColorStructCustomization::CreateColorPicker(bool bUseAlpha)
 
 TSharedRef<SColorPicker> FColorStructCustomization::CreateInlineColorPicker(TWeakPtr<IPropertyHandle> StructWeakHandlePtr)
 {
-	GEditor->BeginTransaction(FText::Format(LOCTEXT("SetColorProperty", "Edit {0}"), StructPropertyHandle->GetPropertyDisplayName()));
+	TransactionIndex = GEditor->BeginTransaction(FText::Format(LOCTEXT("SetColorProperty", "Edit {0}"), StructPropertyHandle->GetPropertyDisplayName()));
 
 	GatherSavedPreColorPickerColors();
 
@@ -354,6 +358,7 @@ void FColorStructCustomization::OnColorPickerCancelled(FLinearColor OriginalColo
 	LastPickerColorString.Reset();
 
 	GEditor->CancelTransaction(0);
+	TransactionIndex.Reset();
 }
 
 void FColorStructCustomization::OnColorPickerWindowClosed(const TSharedRef<SWindow>& Window)
@@ -371,6 +376,7 @@ void FColorStructCustomization::OnColorPickerWindowClosed(const TSharedRef<SWind
 	}
 
 	GEditor->EndTransaction();
+	TransactionIndex.Reset();
 }
 
 
diff --git a/Engine/Source/Editor/DetailCustomizations/Public/Customizations/ColorStructCustomization.h b/Engine/Source/Editor/DetailCustomizations/Public/Customizations/ColorStructCustomization.h
--- a/Engine/Source/Editor/DetailCustomizations/Public/Customizations/ColorStructCustomization.h
+++ b/Engine/Source/Editor/DetailCustomizations/Public/Customizations/ColorStructCustomization.h
@@ -222,6 +222,5 @@ protected:
 	/** The value won;t be updated while editing */
 	bool bDontUpdateWhileEditing;
 
-
-	
+	TOptional<int32> TransactionIndex;	
 };
