PR #10556: Fix FlushValues method not looping through all raw data elements

diff --git a/Engine/Source/Editor/DetailCustomizations/Private/MathStructProxyCustomizations.cpp b/Engine/Source/Editor/DetailCustomizations/Private/MathStructProxyCustomizations.cpp
--- a/Engine/Source/Editor/DetailCustomizations/Private/MathStructProxyCustomizations.cpp
+++ b/Engine/Source/Editor/DetailCustomizations/Private/MathStructProxyCustomizations.cpp
@@ -645,7 +645,7 @@ bool FTransformStructCustomization<T>::FlushValues( TWeakPtr<IPropertyHandle> Pr
 	bool bNotifiedPreChange = false;
 	for (int32 ValueIndex = 0; ValueIndex < RawData.Num(); ValueIndex++)
 	{
-		UE::Math::TTransform<T>* TransformValue = reinterpret_cast<UE::Math::TTransform<T>*>(RawData[0]);
+		UE::Math::TTransform<T>* TransformValue = reinterpret_cast<UE::Math::TTransform<T>*>(RawData[ValueIndex]);
 		if (TransformValue != NULL)
 		{
 			const UE::Math::TTransform<T> PreviousValue = *TransformValue;
