Fix Localization issue in World Setting Details (World Partition)

diff --git a/Engine/Source/Editor/DetailCustomizations/Private/WorldSettingsDetails.cpp b/Engine/Source/Editor/DetailCustomizations/Private/WorldSettingsDetails.cpp
--- a/Engine/Source/Editor/DetailCustomizations/Private/WorldSettingsDetails.cpp
+++ b/Engine/Source/Editor/DetailCustomizations/Private/WorldSettingsDetails.cpp
@@ -195,6 +195,7 @@ void FWorldSettingsDetails::AddWorldCustomization(IDetailLayoutBuilder& DetailBu
 				[
 					SNew(SHorizontalBox)
 					+ SHorizontalBox::Slot()
+					.AutoWidth()
 					[
 						SNew(SButton)
 						.OnClicked_Lambda([CustomizedLevel]()
@@ -211,6 +212,7 @@ void FWorldSettingsDetails::AddWorldCustomization(IDetailLayoutBuilder& DetailBu
 						]
 					]
 					+ SHorizontalBox::Slot()
+					.AutoWidth()
 					[
 						SNew(SButton)
 						.OnClicked_Lambda([CustomizedLevel]()
