[Custom Primitive Data] Prevent Colorpicker from closing when picking a color

diff --git a/Engine/Source/Editor/DetailCustomizations/Private/CustomPrimitiveDataCustomization.cpp b/Engine/Source/Editor/DetailCustomizations/Private/CustomPrimitiveDataCustomization.cpp
--- a/Engine/Source/Editor/DetailCustomizations/Private/CustomPrimitiveDataCustomization.cpp
+++ b/Engine/Source/Editor/DetailCustomizations/Private/CustomPrimitiveDataCustomization.cpp
@@ -896,6 +896,7 @@ FReply FCustomPrimitiveDataCustomization::OnMouseButtonDownColorBlock(const FGeo
 	GEditor->BeginTransaction(FText::Format(LOCTEXT("SetVectorColor", "Edit Primitive Data Vector: {0}"), FText::AsNumber(PrimIdx)));
 
 	FColorPickerArgs PickerArgs;
+	PickerArgs.bOnlyRefreshOnOk = true;
 	PickerArgs.bUseAlpha = true;
 	PickerArgs.InitialColor = GetVectorColor(PrimIdx);
 	PickerArgs.ParentWidget = ColorBlocks[PrimIdx];
