[Custom Primitive Data] Fixing detail customization where the color picker would previously only commit the first channel.

diff --git a/Engine/Source/Editor/DetailCustomizations/Private/CustomPrimitiveDataCustomization.cpp b/Engine/Source/Editor/DetailCustomizations/Private/CustomPrimitiveDataCustomization.cpp
--- a/Engine/Source/Editor/DetailCustomizations/Private/CustomPrimitiveDataCustomization.cpp
+++ b/Engine/Source/Editor/DetailCustomizations/Private/CustomPrimitiveDataCustomization.cpp
@@ -808,12 +808,24 @@ void FCustomPrimitiveDataCustomization::SetVectorColor(FLinearColor NewColor, ui
 	uint32 NumElems;
 	if (GetNumElements(NumElems) == FPropertyAccess::Success)
 	{
-		const int32 MaxElems = FMath::Min((int32)NumElems, PrimIdx + 4);
-
-		for (int32 i = PrimIdx; i < MaxElems; ++i)
+		DataHandle->NotifyPreChange();
+		DataHandle->EnumerateRawData([&](void* Ptr, const int32 ObjectIndex, const int32 NumObjects)
 		{
-			DataArrayHandle->GetElement(i)->SetValue(Color[i - PrimIdx]);
-		}
+			const uint32 MaxElems = FMath::Min(NumElems, PrimIdx + 4u);
+
+			if (TArray<float>* const DataArray = static_cast<TArray<float>*>(Ptr))
+			{
+				for (uint32 i = PrimIdx; i < MaxElems; ++i)
+				{
+					(*DataArray)[i] = Color[i - PrimIdx];
+				}
+			}
+ 
+			return true;
+		});
+		DataHandle->NotifyPostChange(EPropertyChangeType::ValueSet);
+		DataHandle->NotifyFinishedChangingProperties();
+
 	}
 }
 
