[Custom Primitive Data] Fix crash when adding elements on a BP instance placed in the scene

diff --git a/Engine/Source/Editor/DetailCustomizations/Private/CustomPrimitiveDataCustomization.cpp b/Engine/Source/Editor/DetailCustomizations/Private/CustomPrimitiveDataCustomization.cpp
--- a/Engine/Source/Editor/DetailCustomizations/Private/CustomPrimitiveDataCustomization.cpp
+++ b/Engine/Source/Editor/DetailCustomizations/Private/CustomPrimitiveDataCustomization.cpp
@@ -444,7 +444,7 @@ void FCustomPrimitiveDataCustomization::CreateParameterRow(IDetailChildrenBuilde
 
 		TSharedRef<SWidget> ValueWidget = ElementHandle->CreatePropertyValueWidget(false);
 		ValueWidget->SetEnabled(bDataEditable);
-		ElementHandleRef->SetOnPropertyResetToDefault(FSimpleDelegate::CreateSP(this, &FCustomPrimitiveDataCustomization::SetDefaultValue, ElementHandle, PrimIdx));
+		ElementHandleRef->SetOnPropertyResetToDefault(FSimpleDelegate::CreateSP(this, &FCustomPrimitiveDataCustomization::SetDefaultValue, PrimIdx));
 
 		Row.CustomWidget()
 		.NameContent()
@@ -756,12 +756,22 @@ void FCustomPrimitiveDataCustomization::OnAddedDesiredPrimitiveData(uint8 PrimId
 	{
 		GEditor->BeginTransaction(LOCTEXT("OnAddedDesiredPrimitiveData", "Added Items"));
 
-		for (int32 i = NumElements; i <= PrimIdx; ++i)
+		DataHandle->NotifyPreChange();
+
+		TArray<void*> RawArray;
+		DataHandle->AccessRawData(RawArray);
+		
+		if (RawArray.Num() > 0)
 		{
-			DataArrayHandle->AddItem();
+			TArray<float>* Data = reinterpret_cast<TArray<float>*>(RawArray[0]);
 
-			SetDefaultValue(DataArrayHandle->GetElement(i), i);
+			for (int32 i = NumElements; i <= PrimIdx; ++i)
+			{
+				SetDefaultValue(i);
+			}
 		}
+		
+		DataHandle->NotifyPostChange(EPropertyChangeType::ArrayAdd);
 
 		GEditor->EndTransaction();
 	}
@@ -829,52 +839,41 @@ void FCustomPrimitiveDataCustomization::SetVectorColor(FLinearColor NewColor, ui
 	}
 }
 
-void FCustomPrimitiveDataCustomization::SetDefaultValue(TSharedPtr<IPropertyHandle> Handle, uint8 PrimIdx)
+void FCustomPrimitiveDataCustomization::SetDefaultValue(uint8 PrimIdx)
 {
-	if (Handle.IsValid())
-	{
-		TSet<TWeakObjectPtr<UPrimitiveComponent>> ChangedComponents;
+	TSet<TWeakObjectPtr<UPrimitiveComponent>> ChangedComponents;
 
-		// Prioritize vector data since we have a color picker
-		if (VectorParameterData.Contains(PrimIdx))
+	// Prioritize vector data since we have a color picker
+	if (VectorParameterData.Contains(PrimIdx))
+	{
+		for (const FParameterData& ParameterData : VectorParameterData[PrimIdx])
 		{
-			for (const FParameterData& ParameterData : VectorParameterData[PrimIdx])
+			if (ParameterData.Component.IsValid() && !ChangedComponents.Contains(ParameterData.Component))
 			{
-				if (ParameterData.Component.IsValid() && !ChangedComponents.Contains(ParameterData.Component))
+				FLinearColor Color(ForceInitToZero);
+				if (!ParameterData.Material.IsValid() || ParameterData.Material->GetVectorParameterValue(ParameterData.Info, Color))
 				{
-					FLinearColor Color(ForceInitToZero);
-					if (!ParameterData.Material.IsValid() || ParameterData.Material->GetVectorParameterValue(ParameterData.Info, Color))
-					{
-						float* ColorPtr = reinterpret_cast<float*>(&Color);
-						ParameterData.Component->SetDefaultCustomPrimitiveDataFloat(PrimIdx, ColorPtr[ParameterData.IndexOffset]);
+					float* ColorPtr = reinterpret_cast<float*>(&Color);
+					ParameterData.Component->SetDefaultCustomPrimitiveDataFloat(PrimIdx, ColorPtr[ParameterData.IndexOffset]);
 
-						FPropertyChangedEvent PropertyChangedEvent(Handle->GetParentHandle()->GetParentHandle()->GetProperty());
-						PropertyChangedEvent.SetActiveMemberProperty(Handle->GetParentHandle()->GetProperty());
-						ParameterData.Component->PostEditChangeProperty(PropertyChangedEvent);
-
-						ChangedComponents.Add(ParameterData.Component);
-					}
+					ChangedComponents.Add(ParameterData.Component);
 				}
 			}
 		}
+	}
 
-		if (ScalarParameterData.Contains(PrimIdx))
+	if (ScalarParameterData.Contains(PrimIdx))
+	{
+		for (const FParameterData& ParameterData : ScalarParameterData[PrimIdx])
 		{
-			for (const FParameterData& ParameterData : ScalarParameterData[PrimIdx])
+			if (ParameterData.Component.IsValid() && !ChangedComponents.Contains(ParameterData.Component))
 			{
-				if (ParameterData.Component.IsValid() && !ChangedComponents.Contains(ParameterData.Component))
+				float Value = 0.f;
+				if (!ParameterData.Material.IsValid() || ParameterData.Material->GetScalarParameterValue(ParameterData.Info, Value))
 				{
-					float Value = 0.f;
-					if (!ParameterData.Material.IsValid() || ParameterData.Material->GetScalarParameterValue(ParameterData.Info, Value))
-					{
-						ParameterData.Component->SetDefaultCustomPrimitiveDataFloat(PrimIdx, Value);
+					ParameterData.Component->SetDefaultCustomPrimitiveDataFloat(PrimIdx, Value);
 
-						FPropertyChangedEvent PropertyChangedEvent(Handle->GetParentHandle()->GetParentHandle()->GetProperty());
-						PropertyChangedEvent.SetActiveMemberProperty(Handle->GetParentHandle()->GetProperty());
-						ParameterData.Component->PostEditChangeProperty(PropertyChangedEvent);
-
-						ChangedComponents.Add(ParameterData.Component);
-					}
+					ChangedComponents.Add(ParameterData.Component);
 				}
 			}
 		}
diff --git a/Engine/Source/Editor/DetailCustomizations/Private/CustomPrimitiveDataCustomization.h b/Engine/Source/Editor/DetailCustomizations/Private/CustomPrimitiveDataCustomization.h
--- a/Engine/Source/Editor/DetailCustomizations/Private/CustomPrimitiveDataCustomization.h
+++ b/Engine/Source/Editor/DetailCustomizations/Private/CustomPrimitiveDataCustomization.h
@@ -116,7 +116,7 @@ private:
 	void SetVectorColor(FLinearColor NewColor, uint8 PrimIdx);
 
 	/** Iterates through each parameter, and sets their values to the material interface's defaults on a per-component basis */
-	void SetDefaultValue(TSharedPtr<IPropertyHandle> Handle, uint8 PrimIdx);
+	void SetDefaultValue(uint8 PrimIdx);
 
 	/** Helper method to SetDefaultValue for vectors, but makes one transaction for all edits */
 	void SetDefaultVectorValue(uint8 PrimIdx);
