Fixed regression with stable text keying in UMG widgets

CL# 23873743 preferred querying the object for its package, but the widgets you edit in UMG are outered to the transient package rather than the asset. This resulted in edits to text on those widgets not using a package ID.

#rnx

diff --git a/Engine/Source/Editor/DetailCustomizations/Private/TextCustomization.cpp b/Engine/Source/Editor/DetailCustomizations/Private/TextCustomization.cpp
--- a/Engine/Source/Editor/DetailCustomizations/Private/TextCustomization.cpp
+++ b/Engine/Source/Editor/DetailCustomizations/Private/TextCustomization.cpp
@@ -120,7 +120,7 @@ namespace
 						bForceUniqueId = PropertyHandle->HasMetaData(PropertyObjects[InIndex]->HasAnyFlags(RF_ClassDefaultObject | RF_ArchetypeObject) ? NAME_UniqueIdWhenEditedOnDefault : NAME_UniqueIdWhenEditedOnInstance);
 					}
 				}
-				if (!PropertyPackage)
+				if (!PropertyPackage || PropertyPackage == GetTransientPackage())
 				{
 					TArray<UPackage*> PropertyPackages;
 					PropertyHandle->GetOuterPackages(PropertyPackages);
