Skeletal Mesh: Ensure the "Base LOD" param in the skelmesh LOD settings data asset gets properly exposed for higher LODs.

diff --git a/Engine/Source/Editor/DetailCustomizations/Private/SkeletalMeshReductionSettingsDetails.cpp b/Engine/Source/Editor/DetailCustomizations/Private/SkeletalMeshReductionSettingsDetails.cpp
--- a/Engine/Source/Editor/DetailCustomizations/Private/SkeletalMeshReductionSettingsDetails.cpp
+++ b/Engine/Source/Editor/DetailCustomizations/Private/SkeletalMeshReductionSettingsDetails.cpp
@@ -8,6 +8,7 @@
 #include "DetailLayoutBuilder.h"
 #include "DetailWidgetRow.h"
 #include "Engine/SkeletalMesh.h"
+#include "Engine/SkeletalMeshLODSettings.h"
 #include "Fonts/SlateFontInfo.h"
 #include "HAL/PlatformCrt.h"
 #include "IDetailChildrenBuilder.h"
@@ -92,7 +93,8 @@ void FSkeletalMeshReductionSettingsDetails::CustomizeChildren(TSharedRef<IProper
 	{
 		if (StructPropertyHandle->GetParentHandle().IsValid())
 		{
-			if (StructPropertyHandle->GetParentHandle()->GetProperty()->GetFName() == GET_MEMBER_NAME_CHECKED(FSkeletalMeshObject, LODInfo))
+			if (StructPropertyHandle->GetParentHandle()->GetProperty()->GetFName() == GET_MEMBER_NAME_CHECKED(FSkeletalMeshObject, LODInfo) ||
+				StructPropertyHandle->GetParentHandle()->GetProperty()->GetFName() == GET_MEMBER_NAME_CHECKED(USkeletalMeshLODSettings, LODGroups))
 			{
 				return StructPropertyHandle->GetParentHandle()->GetIndexInArray();
 			}
diff --git a/Engine/Source/Runtime/Engine/Classes/Engine/SkeletalMeshLODSettings.h b/Engine/Source/Runtime/Engine/Classes/Engine/SkeletalMeshLODSettings.h
--- a/Engine/Source/Runtime/Engine/Classes/Engine/SkeletalMeshLODSettings.h
+++ b/Engine/Source/Runtime/Engine/Classes/Engine/SkeletalMeshLODSettings.h
@@ -151,6 +151,7 @@ protected:
 	UPROPERTY(globalconfig, EditAnywhere, Category=LODGroups)
 	TArray<FSkeletalMeshLODGroupSettings> LODGroups;
 
+	friend class FSkeletalMeshReductionSettingsDetails;
 public:
 	/** Retrieves the Skeletal mesh LOD group settings for the given name */
 	ENGINE_API const FSkeletalMeshLODGroupSettings& GetSettingsForLODLevel(const int32 LODIndex) const;
