Fix raw data acess image size

#okforgitub public

diff --git a/Engine/Source/Editor/DetailCustomizations/Private/SlateBrushCustomization.cpp b/Engine/Source/Editor/DetailCustomizations/Private/SlateBrushCustomization.cpp
--- a/Engine/Source/Editor/DetailCustomizations/Private/SlateBrushCustomization.cpp
+++ b/Engine/Source/Editor/DetailCustomizations/Private/SlateBrushCustomization.cpp
@@ -1242,19 +1242,21 @@ private:
 					TemporaryBrush->InvalidateResourceHandle();
 				}
 			}
-			FVector2D CachedTextureSize;
+
+			using ImageSizeType = decltype(FSlateBrush::ImageSize);
+			ImageSizeType CachedTextureSize;
 
 			TArray<void*> RawData;
 			ImageSizeProperty->AccessRawData(RawData);
 			if ( RawData.Num() > 0 && RawData[0] != NULL )
 			{
-				CachedTextureSize = *static_cast<FVector2D*>( RawData[0] );
+				CachedTextureSize = *static_cast<ImageSizeType*>( RawData[0] );
 			}
 
 			UTexture2D* BrushTexture = Cast<UTexture2D>(ResourceObject);
 			if ( BrushTexture )
 			{
-				CachedTextureSize = FVector2D(BrushTexture->GetSizeX(), BrushTexture->GetSizeY());
+				CachedTextureSize = ImageSizeType(BrushTexture->GetSizeX(), BrushTexture->GetSizeY());
 			}
 			else if ( ISlateTextureAtlasInterface* AtlasedTextureObject = Cast<ISlateTextureAtlasInterface>(ResourceObject) )
 			{
