Property Editor: Remove 0-360 restriction on rotation dragging so that you can drag to negative rotations and beyond 360.

The delta is set to 1 (ie. slider rotates by 1 degrees)

diff --git a/Engine/Source/Editor/DetailCustomizations/Private/ComponentTransformDetails.cpp b/Engine/Source/Editor/DetailCustomizations/Private/ComponentTransformDetails.cpp
--- a/Engine/Source/Editor/DetailCustomizations/Private/ComponentTransformDetails.cpp
+++ b/Engine/Source/Editor/DetailCustomizations/Private/ComponentTransformDetails.cpp
@@ -432,13 +432,15 @@ void FComponentTransformDetails::GenerateChildContent( IDetailChildrenBuilder& C
 	TSharedPtr<IPropertyHandle> ScalePropertyHandle = LayoutBuilder.GetProperty(USceneComponent::GetRelativeScale3DPropertyName(), USceneComponent::StaticClass());
 
 	const FString& MetaLocationDeltaString = LocationPropertyHandle->GetMetaData("Delta");
+	const FString& MetaRotationDeltaString = RotationPropertyHandle->GetMetaData("Delta");
 	const FString& MetaRotationMinString = RotationPropertyHandle->GetMetaData("UIMin");
 	const FString& MetaRotationMaxString = RotationPropertyHandle->GetMetaData("UIMax");
 	const FString& MetaScaleDeltaString = ScalePropertyHandle->GetMetaData("Delta");
 
 	float LocationSpinDelta = !MetaLocationDeltaString.IsEmpty() ? FCString::Atof(*MetaLocationDeltaString) : 1.f;
-	float RotationMin = !MetaRotationMinString.IsEmpty() ? FCString::Atof(*MetaRotationMinString) : 0.f;
-	float RotationMax = !MetaRotationMaxString.IsEmpty() ? FCString::Atof(*MetaRotationMaxString) : 359.999f;
+	float RotationSpinDelta = !MetaRotationDeltaString.IsEmpty() ? FCString::Atof(*MetaRotationDeltaString) : 1.f;
+	TOptional<FRotator::FReal> RotationMin = !MetaRotationMinString.IsEmpty() ? FCString::Atof(*MetaRotationMinString) : TOptional<FRotator::FReal>();
+	TOptional<FRotator::FReal> RotationMax = !MetaRotationMaxString.IsEmpty() ? FCString::Atof(*MetaRotationMaxString) : TOptional<FRotator::FReal>();
 	float ScaleSpinDelta = !MetaScaleDeltaString.IsEmpty() ? FCString::Atof(*MetaScaleDeltaString) : 0.0025f;
 
 	// Location
@@ -518,6 +520,7 @@ void FComponentTransformDetails::GenerateChildContent( IDetailChildrenBuilder& C
 		[
 			SNew( SNumericRotatorInputBox<FRotator::FReal> )
 			.AllowSpin( SelectedObjects.Num() == 1 ) 
+			.SpinDelta(RotationSpinDelta)
 			.MinSliderValue(RotationMin)
 			.MaxSliderValue(RotationMax)
 			.Roll( this, &FComponentTransformDetails::GetRotationX )
diff --git a/Engine/Source/Runtime/Engine/Classes/Components/SceneComponent.h b/Engine/Source/Runtime/Engine/Classes/Components/SceneComponent.h
--- a/Engine/Source/Runtime/Engine/Classes/Components/SceneComponent.h
+++ b/Engine/Source/Runtime/Engine/Classes/Components/SceneComponent.h
@@ -134,7 +134,7 @@ private:
 	FVector RelativeLocation;
 
 	/** Rotation of the component relative to its parent */
-	UPROPERTY(EditAnywhere, BlueprintReadOnly, ReplicatedUsing=OnRep_Transform, Category=Transform, meta=(AllowPrivateAccess="true", UIMin = "0.0", UIMax = "359.999"))
+	UPROPERTY(EditAnywhere, BlueprintReadOnly, ReplicatedUsing=OnRep_Transform, Category=Transform, meta=(AllowPrivateAccess="true", LinearDeltaSensitivity = "1", Delta = "1.0"))
 	FRotator RelativeRotation;
 
 	/**
diff --git a/Engine/Source/Runtime/Slate/Public/Widgets/Input/SRotatorInputBox.h b/Engine/Source/Runtime/Slate/Public/Widgets/Input/SRotatorInputBox.h
--- a/Engine/Source/Runtime/Slate/Public/Widgets/Input/SRotatorInputBox.h
+++ b/Engine/Source/Runtime/Slate/Public/Widgets/Input/SRotatorInputBox.h
@@ -32,8 +32,9 @@ public:
 		: _bColorAxisLabels(false)
 		, _Font(FAppStyle::Get().GetFontStyle("NormalFont"))
 		, _AllowSpin(true)
-		, _MinSliderValue(0.0f)
-		, _MaxSliderValue(359.999f)
+		, _SpinDelta(1.f)
+		, _MinSliderValue(TOptional<NumericType>())
+		, _MaxSliderValue(TOptional<NumericType>())
 		, _DisplayToggle( false )
 		, _TogglePitchChecked( ECheckBoxState::Checked )
 		, _ToggleYawChecked( ECheckBoxState::Checked )
@@ -65,6 +66,9 @@ public:
 		/** Whether or not values can be spun or if they should be typed in */
 		SLATE_ARGUMENT( bool, AllowSpin )
 
+		/** The delta amount to apply, per pixel, when the spinner is dragged. */
+		SLATE_ATTRIBUTE(NumericType, SpinDelta)
+
 		/** The minimum value that can be specified by using the slider */
 		SLATE_ATTRIBUTE( TOptional<NumericType>, MinSliderValue )
 
@@ -156,6 +160,9 @@ public:
 			[
 				SNew(SNumericEntryBox<NumericType>)
 				.AllowSpin(InArgs._AllowSpin)
+				.Delta(InArgs._SpinDelta)
+				.MinValue(InArgs._MinSliderValue)
+				.MaxValue(InArgs._MaxSliderValue)
 				.MinSliderValue(InArgs._MinSliderValue)
 				.MaxSliderValue(InArgs._MaxSliderValue)
 				.LabelPadding(FMargin(3.0f))
@@ -188,6 +195,9 @@ public:
 			[
 				SNew(SNumericEntryBox<NumericType>)
 				.AllowSpin(InArgs._AllowSpin)
+				.Delta(InArgs._SpinDelta)
+				.MinValue(InArgs._MinSliderValue)
+				.MaxValue(InArgs._MaxSliderValue)
 				.MinSliderValue(InArgs._MinSliderValue)
 				.MaxSliderValue(InArgs._MaxSliderValue)
 				.LabelPadding(FMargin(3.0f))
@@ -220,6 +230,9 @@ public:
 			[
 				SNew(SNumericEntryBox<NumericType>)
 				.AllowSpin(InArgs._AllowSpin)
+				.Delta(InArgs._SpinDelta)
+				.MinValue(InArgs._MinSliderValue)
+				.MaxValue(InArgs._MaxSliderValue)
 				.MinSliderValue(InArgs._MinSliderValue)
 				.MaxSliderValue(InArgs._MaxSliderValue)
 				.LabelPadding(FMargin(3.0f))
