ComponentTransformDetails:
- Fix Actor locked location not being taken into account when selecting Root Component (bIsEnabledCache)
- Disable ResetToDefault when bIsEnabledCache is false

#rnx

diff --git a/Engine/Source/Editor/DetailCustomizations/Private/ComponentTransformDetails.cpp b/Engine/Source/Editor/DetailCustomizations/Private/ComponentTransformDetails.cpp
--- a/Engine/Source/Editor/DetailCustomizations/Private/ComponentTransformDetails.cpp
+++ b/Engine/Source/Editor/DetailCustomizations/Private/ComponentTransformDetails.cpp
@@ -135,6 +135,7 @@ FComponentTransformDetails::FComponentTransformDetails( const TArray< TWeakObjec
 	, bEditingRotationInUI( false )
 	, bIsSliderTransaction( false )
 	, HiddenFieldMask( 0 )
+	, bIsEnabledCache( false )
 {
 	GConfig->GetBool(TEXT("SelectionDetails"), TEXT("PreserveScaleRatio"), bPreserveScaleRatio, GEditorPerProjectIni);
 	FCoreUObjectDelegates::OnObjectsReplaced.AddRaw(this, &FComponentTransformDetails::OnObjectsReplaced);
@@ -294,7 +295,7 @@ bool FComponentTransformDetails::OnCanCopy( ETransformField::Type TransformField
 
 void FComponentTransformDetails::OnCopy( ETransformField::Type TransformField )
 {
-	CacheTransform();
+	CacheDetails();
 
 	FString CopyStr;
 	switch (TransformField)
@@ -593,7 +594,7 @@ END_SLATE_FUNCTION_BUILD_OPTIMIZATION
 
 void FComponentTransformDetails::Tick( float DeltaTime ) 
 {
-	CacheTransform();
+	CacheDetails();
 	if (!FixedDisplayUnits.IsSet())
 	{
 		CacheCommonLocationUnits();
@@ -654,7 +655,7 @@ void FComponentTransformDetails::UpdatePropertyHandlesObjects(const TArray<UObje
 
 bool FComponentTransformDetails::GetIsEnabled() const
 {
-	return !GEditor->HasLockedActors() || SelectedActorInfo.NumSelected == 0;
+	return bIsEnabledCache;
 }
 
 const FSlateBrush* FComponentTransformDetails::GetPreserveScaleRatioImage() const
@@ -871,13 +872,16 @@ bool FComponentTransformDetails::GetLocationResetVisibility() const
 
 void FComponentTransformDetails::OnLocationResetClicked()
 {
-	const FText TransactionName = LOCTEXT("ResetLocation", "Reset Location");
-	FScopedTransaction Transaction(TransactionName);
+	if (bIsEnabledCache)
+	{
+		const FText TransactionName = LOCTEXT("ResetLocation", "Reset Location");
+		FScopedTransaction Transaction(TransactionName);
 
-	const USceneComponent* Archetype = FGetRootComponentArchetype::Get(SelectedObjects[0].Get());
-	const FVector Data = Archetype ? Archetype->GetRelativeLocation() : FVector::ZeroVector;
+		const USceneComponent* Archetype = FGetRootComponentArchetype::Get(SelectedObjects[0].Get());
+		const FVector Data = Archetype ? Archetype->GetRelativeLocation() : FVector::ZeroVector;
 
-	OnSetTransform(ETransformField::Location, EAxisList::All, Data, false, true);
+		OnSetTransform(ETransformField::Location, EAxisList::All, Data, false, true);
+	}
 }
 
 bool FComponentTransformDetails::GetRotationResetVisibility() const
@@ -891,13 +895,16 @@ bool FComponentTransformDetails::GetRotationResetVisibility() const
 
 void FComponentTransformDetails::OnRotationResetClicked()
 {
-	const FText TransactionName = LOCTEXT("ResetRotation", "Reset Rotation");
-	FScopedTransaction Transaction(TransactionName);
+	if (bIsEnabledCache)
+	{
+		const FText TransactionName = LOCTEXT("ResetRotation", "Reset Rotation");
+		FScopedTransaction Transaction(TransactionName);
 
-	const USceneComponent* Archetype = FGetRootComponentArchetype::Get(SelectedObjects[0].Get());
-	const FVector Data = Archetype ? Archetype->GetRelativeRotation().Euler() : FVector::ZeroVector;
+		const USceneComponent* Archetype = FGetRootComponentArchetype::Get(SelectedObjects[0].Get());
+		const FVector Data = Archetype ? Archetype->GetRelativeRotation().Euler() : FVector::ZeroVector;
 
-	OnSetTransform(ETransformField::Rotation, EAxisList::All, Data, false, true);
+		OnSetTransform(ETransformField::Rotation, EAxisList::All, Data, false, true);
+	}
 }
 
 bool FComponentTransformDetails::GetScaleResetVisibility() const
@@ -911,13 +918,16 @@ bool FComponentTransformDetails::GetScaleResetVisibility() const
 
 void FComponentTransformDetails::OnScaleResetClicked()
 {
-	const FText TransactionName = LOCTEXT("ResetScale", "Reset Scale");
-	FScopedTransaction Transaction(TransactionName);
+	if (bIsEnabledCache)
+	{
+		const FText TransactionName = LOCTEXT("ResetScale", "Reset Scale");
+		FScopedTransaction Transaction(TransactionName);
 
-	const USceneComponent* Archetype = FGetRootComponentArchetype::Get(SelectedObjects[0].Get());
-	const FVector Data = Archetype ? Archetype->GetRelativeScale3D() : FVector(1.0f);
+		const USceneComponent* Archetype = FGetRootComponentArchetype::Get(SelectedObjects[0].Get());
+		const FVector Data = Archetype ? Archetype->GetRelativeScale3D() : FVector(1.0f);
 
-	OnSetTransform(ETransformField::Scale, EAxisList::All, Data, false, true);
+		OnSetTransform(ETransformField::Scale, EAxisList::All, Data, false, true);
+	}
 }
 
 void FComponentTransformDetails::ExtendXScaleContextMenu( FMenuBuilder& MenuBuilder )
@@ -977,12 +987,12 @@ void FComponentTransformDetails::OnZScaleMirrored()
 	OnSetTransform(ETransformField::Scale, EAxisList::Z, FVector(1.0f), true, true);
 }
 
-void FComponentTransformDetails::CacheTransform()
+void FComponentTransformDetails::CacheDetails()
 {
 	FVector CurLoc;
 	FRotator CurRot;
 	FVector CurScale;
-
+	bIsEnabledCache = true;
 	for( int32 ObjectIndex = 0; ObjectIndex < SelectedObjects.Num(); ++ObjectIndex )
 	{
 		TWeakObjectPtr<UObject> ObjectPtr = SelectedObjects[ObjectIndex];
@@ -996,6 +1006,11 @@ void FComponentTransformDetails::CacheTransform()
 			FVector Scale;
 			if( SceneComponent )
 			{
+				if (AActor* Owner = SceneComponent->GetOwner(); Owner && Owner->GetRootComponent() == SceneComponent)
+				{
+					bIsEnabledCache &= !Owner->IsLockLocation();
+				}
+				
 				Loc = SceneComponent->GetRelativeLocation();
 				FRotator* FoundRotator = ObjectToRelativeRotationMap.Find(SceneComponent);
 				Rot = (bEditingRotationInUI && !Object->IsTemplate() && FoundRotator) ? *FoundRotator : SceneComponent->GetRelativeRotation();
@@ -1425,7 +1440,7 @@ void FComponentTransformDetails::OnSetTransform(ETransformField::Type TransformF
 	if (bCommitted && bBeganTransaction)
 	{
 		GEditor->EndTransaction();
-		CacheTransform();
+		CacheDetails();
 	}
 
 	GUnrealEd->UpdatePivotLocationForSelection();
diff --git a/Engine/Source/Editor/DetailCustomizations/Private/ComponentTransformDetails.h b/Engine/Source/Editor/DetailCustomizations/Private/ComponentTransformDetails.h
--- a/Engine/Source/Editor/DetailCustomizations/Private/ComponentTransformDetails.h
+++ b/Engine/Source/Editor/DetailCustomizations/Private/ComponentTransformDetails.h
@@ -51,11 +51,6 @@ public:
 	FComponentTransformDetails( const TArray< TWeakObjectPtr<UObject> >& InSelectedObjects, const FSelectedActorInfo& InSelectedActorInfo, IDetailLayoutBuilder& DetailBuilder );
 	~FComponentTransformDetails();
 
-	/**
-	 * Caches the representation of the actor transform for the user input boxes                   
-	 */
-	void CacheTransform();
-
 	virtual void GenerateHeaderRowContent( FDetailWidgetRow& NodeRow ) override {}
 	virtual void GenerateChildContent( IDetailChildrenBuilder& ChildrenBuilder ) override;
 	virtual bool RequiresTick() const override { return true; }
@@ -70,6 +65,9 @@ public:
 	}
 
 private:
+	/** Caches some information of the actor (transform, locked location) for the user input boxes */
+	void CacheDetails();
+
 	/** @return Whether the transform details panel should be enabled (editable) or not (read-only / greyed out) */
 	bool GetIsEnabled() const;
 
@@ -332,4 +330,6 @@ private:
 	TArray< TSharedPtr< IPropertyHandle> > PropertyHandles;
 	/** Holds the property handles' outer objects. Used to update the handles' objects when the actor construction script runs. */
 	TArray< TWeakObjectPtr<UObject> > CachedHandlesObjects;
+	/** Cached enabled value of the selected set */
+	bool bIsEnabledCache;
 };
