Make sure the fbx options hide the option we cannot use when class permission is denied.

#rnx

diff --git a/Engine/Source/Editor/DetailCustomizations/Private/FbxImportUIDetails.cpp b/Engine/Source/Editor/DetailCustomizations/Private/FbxImportUIDetails.cpp
--- a/Engine/Source/Editor/DetailCustomizations/Private/FbxImportUIDetails.cpp
+++ b/Engine/Source/Editor/DetailCustomizations/Private/FbxImportUIDetails.cpp
@@ -2,6 +2,7 @@
 
 #include "FbxImportUIDetails.h"
 
+#include "AssetToolsModule.h"
 #include "Containers/EnumAsByte.h"
 #include "Containers/UnrealString.h"
 #include "Delegates/Delegate.h"
@@ -18,6 +19,7 @@
 #include "Factories/FbxTextureImportData.h"
 #include "Fonts/SlateFontInfo.h"
 #include "HAL/PlatformCrt.h"
+#include "IAssetTools.h"
 #include "IDetailGroup.h"
 #include "IDetailPropertyRow.h"
 #include "Internationalization/Internationalization.h"
@@ -31,7 +33,9 @@
 #include "Misc/AssertionMacros.h"
 #include "Misc/Attribute.h"
 #include "Misc/Guid.h"
+#include "Misc/NamePermissionList.h"
 #include "Misc/StringFormatArg.h"
+#include "PhysicsEngine/PhysicsAsset.h"
 #include "PropertyEditorModule.h"
 #include "PropertyHandle.h"
 #include "SEnumCombo.h"
@@ -260,6 +264,18 @@ void FFbxImportUIDetails::CustomizeDetails( IDetailLayoutBuilder& DetailBuilder
 
 	ImportUI = Cast<UFbxImportUI>(EditingObjects[0].Get());
 	
+	auto GetClassPermission = [](UClass* Class, EAssetClassAction ClassAction)
+		{
+			IAssetTools& AssetTools = FAssetToolsModule::GetModule().Get();
+			TSharedPtr<FPathPermissionList> AssetClassPermissionList = AssetTools.GetAssetClassPathPermissionList(ClassAction);
+			if (Class && AssetClassPermissionList && AssetClassPermissionList->HasFiltering())
+			{
+				return AssetClassPermissionList->PassesFilter(Class->GetPathName());
+			}
+			return true;
+		};
+	bool bShowPhysicsAssetOptions = GetClassPermission(UPhysicsAsset::StaticClass(), EAssetClassAction::ImportAsset);
+
 	bool bShowCompareResult = ShowCompareResult();
 
 	auto AddRefreshCustomDetailEvent = [this](TSharedPtr<IPropertyHandle> Handle)
@@ -419,6 +435,13 @@ void FFbxImportUIDetails::CustomizeDetails( IDetailLayoutBuilder& DetailBuilder
 		}
 	}
 
+	if (!bShowPhysicsAssetOptions)
+	{
+		TSharedRef<IPropertyHandle> CreatePhysicsAssetProperty = DetailBuilder.GetProperty(GET_MEMBER_NAME_CHECKED(UFbxImportUI, bCreatePhysicsAsset));
+		DetailBuilder.HideProperty(CreatePhysicsAssetProperty);
+		ImportUI->bCreatePhysicsAsset = false;
+	}
+
 	TMap<FString, TArray<TSharedPtr<IPropertyHandle>>> SubCategoriesProperties;
 	TMap<FString, bool > SubCategoriesAdvanced;
 	TMap<FString, FText > SubCategoriesTooltip;
