ComponentTransformDetails: Properly handle CanEditChange on properties: bAbsoluteLocation, bAbsoluteRotation, bAbsoluteScale, RelativeLocation, RelativeRotation, RelativeScale3D

#rnx

diff --git a/Engine/Source/Editor/DetailCustomizations/Private/ComponentTransformDetails.cpp b/Engine/Source/Editor/DetailCustomizations/Private/ComponentTransformDetails.cpp
--- a/Engine/Source/Editor/DetailCustomizations/Private/ComponentTransformDetails.cpp
+++ b/Engine/Source/Editor/DetailCustomizations/Private/ComponentTransformDetails.cpp
@@ -203,35 +203,34 @@ TSharedRef<SWidget> FComponentTransformDetails::BuildTransformFieldLabel( ETrans
 
 	MenuBuilder.EndSection();
 
-	TSharedRef<SWidget> NameContent =
-		SNew(SComboButton)
-		.ContentPadding(0)
-		.MenuContent()
-		[
-			MenuBuilder.MakeWidget()
-		]
-		.ButtonContent()
+	TSharedRef<SHorizontalBox> NameContent =
+		SNew(SHorizontalBox)
+		+ SHorizontalBox::Slot()
+		.VAlign(VAlign_Center)
 		[
-			SNew( SBox )
-			.Padding( FMargin( 0.0f, 0.0f, 2.0f, 0.0f ) )
-			.MinDesiredWidth(50.f)
+			SNew(SComboButton)
+			.ContentPadding(0)
+			.IsEnabled(this, &FComponentTransformDetails::CanChangeAbsoluteFlag, TransformField)
+			.MenuContent()
 			[
-				SNew(STextBlock)
-				.Text(this, &FComponentTransformDetails::GetTransformFieldText, TransformField)
-				.Font(IDetailLayoutBuilder::GetDetailFont())
+				MenuBuilder.MakeWidget()
 			]
-		];
-
-	if(TransformField == ETransformField::Scale)
-	{
-		NameContent =
-			SNew(SHorizontalBox)
-			+ SHorizontalBox::Slot()
-			.VAlign(VAlign_Center)
+			.ButtonContent()
 			[
-				NameContent
+				SNew( SBox )
+				.Padding( FMargin( 0.0f, 0.0f, 2.0f, 0.0f ) )
+				.MinDesiredWidth(50.f)
+				[
+					SNew(STextBlock)
+					.Text(this, &FComponentTransformDetails::GetTransformFieldText, TransformField)
+					.Font(IDetailLayoutBuilder::GetDetailFont())
+				]
 			]
-			+ SHorizontalBox::Slot()
+		];
+	
+	if (TransformField == ETransformField::Scale)
+	{
+		NameContent->AddSlot()
 			.AutoWidth()
 			.VAlign(VAlign_Center)
 			.Padding(FMargin(4.0f, 0.0f, 0.0f, 0.0f))
@@ -239,7 +238,7 @@ TSharedRef<SWidget> FComponentTransformDetails::BuildTransformFieldLabel( ETrans
 				// Add a checkbox to toggle between preserving the ratio of x,y,z components of scale when a value is entered
 				SNew(SCheckBox)
 				.IsChecked(this, &FComponentTransformDetails::IsPreserveScaleRatioChecked)
-				.IsEnabled(this, &FComponentTransformDetails::GetIsEnabled)
+				.IsEnabled(this, &FComponentTransformDetails::GetIsScaleEnabled)
 				.OnCheckStateChanged(this, &FComponentTransformDetails::OnPreserveScaleRatioToggled)
 				.Style(FAppStyle::Get(), "TransparentCheckBox")
 				.ToolTipText(LOCTEXT("PreserveScaleToolTip", "When locked, scales uniformly based on the current xyz scale values so the object maintains its shape in each direction when scaled"))
@@ -455,12 +454,17 @@ void FComponentTransformDetails::GenerateChildContent( IDetailChildrenBuilder& C
 
 		ParentCategory.OnPasteFromText()->AddSP(this, &FComponentTransformDetails::OnPasteFromText, ETransformField::Location);
 
+		FindOrCreatePropertyHandle(USceneComponent::GetAbsoluteLocationPropertyName(), ChildrenBuilder);
+
+		TSharedPtr<IPropertyHandle> PropertyHandle = FindOrCreatePropertyHandle(USceneComponent::GetRelativeLocationPropertyName(), ChildrenBuilder);
+
 		ChildrenBuilder.AddCustomRow( LOCTEXT("LocationFilter", "Location") )
 		.RowTag("Location")
 		.CopyAction( CreateCopyAction( ETransformField::Location ) )
 		.PasteAction( CreatePasteAction( ETransformField::Location ) )
 		.OverrideResetToDefault(FResetToDefaultOverride::Create(TAttribute<bool>(this, &FComponentTransformDetails::GetLocationResetVisibility), FSimpleDelegate::CreateSP(this, &FComponentTransformDetails::OnLocationResetClicked)))
-		.PropertyHandleList({ GeneratePropertyHandle(USceneComponent::GetRelativeLocationPropertyName(), ChildrenBuilder) })
+		.PropertyHandleList({ PropertyHandle })
+		.IsEnabled(TAttribute<bool>(this, &FComponentTransformDetails::GetIsEnabled))
 		.NameContent()
 		.VAlign(VAlign_Center)
 		[
@@ -476,7 +480,7 @@ void FComponentTransformDetails::GenerateChildContent( IDetailChildrenBuilder& C
 			.Y(this, &FComponentTransformDetails::GetLocationY)
 			.Z(this, &FComponentTransformDetails::GetLocationZ)
 			.bColorAxisLabels(true)
-			.IsEnabled(this, &FComponentTransformDetails::GetIsEnabled)
+			.IsEnabled(this, &FComponentTransformDetails::GetIsLocationEnabled)
 			.OnXChanged(this, &FComponentTransformDetails::OnSetTransformAxis, ETextCommit::Default, ETransformField::Location, EAxisList::X, false)
 			.OnYChanged(this, &FComponentTransformDetails::OnSetTransformAxis, ETextCommit::Default, ETransformField::Location, EAxisList::Y, false)
 			.OnZChanged(this, &FComponentTransformDetails::OnSetTransformAxis, ETextCommit::Default, ETransformField::Location, EAxisList::Z, false)
@@ -503,12 +507,17 @@ void FComponentTransformDetails::GenerateChildContent( IDetailChildrenBuilder& C
 
 		ParentCategory.OnPasteFromText()->AddSP(this, &FComponentTransformDetails::OnPasteFromText, ETransformField::Rotation);
 
+		FindOrCreatePropertyHandle(USceneComponent::GetAbsoluteRotationPropertyName(), ChildrenBuilder);
+
+		TSharedPtr<IPropertyHandle> PropertyHandle = FindOrCreatePropertyHandle(USceneComponent::GetRelativeRotationPropertyName(), ChildrenBuilder);
+
 		ChildrenBuilder.AddCustomRow( LOCTEXT("RotationFilter", "Rotation") )
 		.RowTag("Rotation")
 		.CopyAction( CreateCopyAction(ETransformField::Rotation) )
 		.PasteAction( CreatePasteAction(ETransformField::Rotation) )
 		.OverrideResetToDefault(FResetToDefaultOverride::Create(TAttribute<bool>(this, &FComponentTransformDetails::GetRotationResetVisibility), FSimpleDelegate::CreateSP(this, &FComponentTransformDetails::OnRotationResetClicked)))
-		.PropertyHandleList({ GeneratePropertyHandle(USceneComponent::GetRelativeRotationPropertyName(), ChildrenBuilder) })
+		.PropertyHandleList({ PropertyHandle })
+		.IsEnabled(TAttribute<bool>(this, &FComponentTransformDetails::GetIsEnabled))
 		.NameContent()
 		.VAlign(VAlign_Center)
 		[
@@ -528,7 +537,7 @@ void FComponentTransformDetails::GenerateChildContent( IDetailChildrenBuilder& C
 			.Pitch( this, &FComponentTransformDetails::GetRotationY )
 			.Yaw( this, &FComponentTransformDetails::GetRotationZ )
 			.bColorAxisLabels( true )
-			.IsEnabled( this, &FComponentTransformDetails::GetIsEnabled )
+			.IsEnabled( this, &FComponentTransformDetails::GetIsRotationEnabled )
 			.OnPitchBeginSliderMovement( this, &FComponentTransformDetails::OnBeginRotationSlider )
 			.OnYawBeginSliderMovement( this, &FComponentTransformDetails::OnBeginRotationSlider )
 			.OnRollBeginSliderMovement( this, &FComponentTransformDetails::OnBeginRotationSlider )
@@ -551,12 +560,17 @@ void FComponentTransformDetails::GenerateChildContent( IDetailChildrenBuilder& C
 	{
 		ParentCategory.OnPasteFromText()->AddSP(this, &FComponentTransformDetails::OnPasteFromText, ETransformField::Scale);
 		
+		FindOrCreatePropertyHandle(USceneComponent::GetAbsoluteScalePropertyName(), ChildrenBuilder);
+
+		TSharedPtr<IPropertyHandle> PropertyHandle = FindOrCreatePropertyHandle(USceneComponent::GetRelativeScale3DPropertyName(), ChildrenBuilder);
+
 		ChildrenBuilder.AddCustomRow( LOCTEXT("ScaleFilter", "Scale") )
 		.RowTag("Scale")
 		.CopyAction( CreateCopyAction(ETransformField::Scale) )
 		.PasteAction( CreatePasteAction(ETransformField::Scale) )
 		.OverrideResetToDefault(FResetToDefaultOverride::Create(TAttribute<bool>(this, &FComponentTransformDetails::GetScaleResetVisibility), FSimpleDelegate::CreateSP(this, &FComponentTransformDetails::OnScaleResetClicked)))
-		.PropertyHandleList({ GeneratePropertyHandle(USceneComponent::GetRelativeScale3DPropertyName(), ChildrenBuilder) })
+		.PropertyHandleList({ PropertyHandle })
+		.IsEnabled(TAttribute<bool>(this, &FComponentTransformDetails::GetIsEnabled))
 		.NameContent()
 		.VAlign(VAlign_Center)
 		[
@@ -572,7 +586,7 @@ void FComponentTransformDetails::GenerateChildContent( IDetailChildrenBuilder& C
 			.Y( this, &FComponentTransformDetails::GetScaleY )
 			.Z( this, &FComponentTransformDetails::GetScaleZ )
 			.bColorAxisLabels( true )
-			.IsEnabled( this, &FComponentTransformDetails::GetIsEnabled )
+			.IsEnabled( this, &FComponentTransformDetails::GetIsScaleEnabled )
 			.OnXChanged( this, &FComponentTransformDetails::OnSetTransformAxis, ETextCommit::Default, ETransformField::Scale, EAxisList::X, false )
 			.OnYChanged( this, &FComponentTransformDetails::OnSetTransformAxis, ETextCommit::Default, ETransformField::Scale, EAxisList::Y, false )
 			.OnZChanged( this, &FComponentTransformDetails::OnSetTransformAxis, ETextCommit::Default, ETransformField::Scale, EAxisList::Z, false )
@@ -620,8 +634,13 @@ void FComponentTransformDetails::CacheCommonLocationUnits()
 	SetupFixedDisplay(LargestValue);
 }
 
-TSharedPtr<IPropertyHandle> FComponentTransformDetails::GeneratePropertyHandle(FName PropertyName, IDetailChildrenBuilder& ChildrenBuilder)
+TSharedPtr<IPropertyHandle> FComponentTransformDetails::FindOrCreatePropertyHandle(FName PropertyName, IDetailChildrenBuilder& ChildrenBuilder)
 {
+	if (TSharedPtr<IPropertyHandle>* HandlePtr = PropertyHandles.Find(PropertyName))
+	{
+		return *HandlePtr;
+	}
+	
 	// Try finding the property handle in the details panel's property map first.
 	IDetailLayoutBuilder& LayoutBuilder = ChildrenBuilder.GetParentCategory().GetParentLayout();
 	TSharedPtr<IPropertyHandle> PropertyHandle = LayoutBuilder.GetProperty(PropertyName, USceneComponent::StaticClass());
@@ -634,7 +653,7 @@ TSharedPtr<IPropertyHandle> FComponentTransformDetails::GeneratePropertyHandle(F
 		CachedHandlesObjects.Append(SceneComponents);
 	}
 
-	PropertyHandles.Add(PropertyHandle);
+	PropertyHandles.Add(PropertyName, PropertyHandle);
 	return PropertyHandle;
 }
 
@@ -644,11 +663,12 @@ void FComponentTransformDetails::UpdatePropertyHandlesObjects(const TArray<UObje
 	CachedHandlesObjects.Reset(NewSceneComponents.Num());
 	Algo::Transform(NewSceneComponents, CachedHandlesObjects, [](UObject* Obj) { return TWeakObjectPtr<UObject>(Obj); });
 
-	for (TSharedPtr<IPropertyHandle>& Handle : PropertyHandles)
+	for (TMap<FName, TSharedPtr<IPropertyHandle>>::TIterator It(PropertyHandles); It; ++It)
 	{
-		if (Handle && Handle->IsValidHandle())
+		TSharedPtr<IPropertyHandle> PropertyHandle = It.Value();
+		if (PropertyHandle && PropertyHandle->IsValidHandle())
 		{
-			Handle->ReplaceOuterObjects(NewSceneComponents);
+			PropertyHandle->ReplaceOuterObjects(NewSceneComponents);
 		}
 	}
 }
@@ -658,6 +678,33 @@ bool FComponentTransformDetails::GetIsEnabled() const
 	return bIsEnabledCache;
 }
 
+bool FComponentTransformDetails::GetIsLocationEnabled() const
+{
+	return GetIsTransformComponentEnabled(USceneComponent::GetRelativeLocationPropertyName());
+}
+
+bool FComponentTransformDetails::GetIsRotationEnabled() const
+{
+	return GetIsTransformComponentEnabled(USceneComponent::GetRelativeRotationPropertyName());
+}
+
+bool FComponentTransformDetails::GetIsScaleEnabled() const
+{
+	return GetIsTransformComponentEnabled(USceneComponent::GetRelativeScale3DPropertyName());
+}
+
+bool FComponentTransformDetails::GetIsTransformComponentEnabled(FName ComponentName) const
+{
+	if (GetIsEnabled())
+	{
+		if (const TSharedPtr<IPropertyHandle>* PropertyHandle = PropertyHandles.Find(ComponentName))
+		{
+			return (*PropertyHandle)->IsEditable();
+		}
+	}
+	return false;
+}
+
 const FSlateBrush* FComponentTransformDetails::GetPreserveScaleRatioImage() const
 {
 	return bPreserveScaleRatio ? FAppStyle::GetBrush( TEXT("Icons.Lock") ) : FAppStyle::GetBrush( TEXT("Icons.Unlock") ) ;
@@ -852,6 +899,36 @@ bool FComponentTransformDetails::IsAbsoluteTransformChecked(ETransformField::Typ
 	}
 }
 
+bool FComponentTransformDetails::CanChangeAbsoluteFlag(ETransformField::Type TransformField) const
+{
+	FName PropertyName;
+
+	switch (TransformField)
+	{
+	case ETransformField::Location:
+		PropertyName = USceneComponent::GetAbsoluteLocationPropertyName();
+		break;
+	case ETransformField::Rotation:
+		PropertyName = USceneComponent::GetAbsoluteRotationPropertyName();
+		break;
+	case ETransformField::Scale:
+		PropertyName = USceneComponent::GetAbsoluteScalePropertyName();
+		break;
+	default:
+		break;
+	}
+
+	if (!PropertyName.IsNone())
+	{
+		if (const TSharedPtr<IPropertyHandle>* HandlePtr = PropertyHandles.Find(PropertyName))
+		{
+			return (*HandlePtr)->IsEditable();
+		}
+	}
+
+	return false;
+}
+
 struct FGetRootComponentArchetype
 {
 	static USceneComponent* Get(UObject* Object)
@@ -872,7 +949,7 @@ bool FComponentTransformDetails::GetLocationResetVisibility() const
 
 void FComponentTransformDetails::OnLocationResetClicked()
 {
-	if (bIsEnabledCache)
+	if (GetIsLocationEnabled())
 	{
 		const FText TransactionName = LOCTEXT("ResetLocation", "Reset Location");
 		FScopedTransaction Transaction(TransactionName);
@@ -895,7 +972,7 @@ bool FComponentTransformDetails::GetRotationResetVisibility() const
 
 void FComponentTransformDetails::OnRotationResetClicked()
 {
-	if (bIsEnabledCache)
+	if (GetIsRotationEnabled())
 	{
 		const FText TransactionName = LOCTEXT("ResetRotation", "Reset Rotation");
 		FScopedTransaction Transaction(TransactionName);
@@ -918,7 +995,7 @@ bool FComponentTransformDetails::GetScaleResetVisibility() const
 
 void FComponentTransformDetails::OnScaleResetClicked()
 {
-	if (bIsEnabledCache)
+	if (GetIsScaleEnabled())
 	{
 		const FText TransactionName = LOCTEXT("ResetScale", "Reset Scale");
 		FScopedTransaction Transaction(TransactionName);
diff --git a/Engine/Source/Editor/DetailCustomizations/Private/ComponentTransformDetails.h b/Engine/Source/Editor/DetailCustomizations/Private/ComponentTransformDetails.h
--- a/Engine/Source/Editor/DetailCustomizations/Private/ComponentTransformDetails.h
+++ b/Engine/Source/Editor/DetailCustomizations/Private/ComponentTransformDetails.h
@@ -70,6 +70,10 @@ private:
 
 	/** @return Whether the transform details panel should be enabled (editable) or not (read-only / greyed out) */
 	bool GetIsEnabled() const;
+	bool GetIsLocationEnabled() const;
+	bool GetIsRotationEnabled() const;
+	bool GetIsScaleEnabled() const;
+	bool GetIsTransformComponentEnabled(FName ComponentName) const;
 
 	/** Sets a vector based on two source vectors and an axis list */
 	FVector GetAxisFilteredVector(EAxisList::Type Axis, const FVector& NewValue, const FVector& OldValue);
@@ -173,6 +177,9 @@ private:
 	/** @return true if Absolute flag of transform type matches passed in bCheckAbsolute*/
 	bool IsAbsoluteTransformChecked( ETransformField::Type TransformField, bool bAbsoluteEnabled=true) const;
 
+	/** @return true if Absolute flag of transform can be changed */
+	bool CanChangeAbsoluteFlag(ETransformField::Type TransformField) const;
+
 	/** @return true of copy is enabled for the specified field */
 	bool OnCanCopy( ETransformField::Type TransformField ) const;
 
@@ -253,8 +260,8 @@ private:
 	/** Cache a single unit to display all location components in */
 	void CacheCommonLocationUnits();
 
-	/** Generate a property handle from a property name. */
-	TSharedPtr<IPropertyHandle> GeneratePropertyHandle(FName PropertyName, IDetailChildrenBuilder& ChildrenBuilder);
+	/** Get a property handle from a property name. */
+	TSharedPtr<IPropertyHandle> FindOrCreatePropertyHandle(FName PropertyName, IDetailChildrenBuilder& ChildrenBuilder);
 	/** Update the outer objects of the property handles generated from this transform. */
 	void UpdatePropertyHandlesObjects(const TArray<UObject*> NewSceneComponents);
 
@@ -327,7 +334,7 @@ private:
 	/** Bitmask to indicate which fields should be hidden (if any) */
 	uint8 HiddenFieldMask;
 	/** Holds this transform's property handles. */
-	TArray< TSharedPtr< IPropertyHandle> > PropertyHandles;
+	TMap<FName, TSharedPtr< IPropertyHandle> > PropertyHandles;
 	/** Holds the property handles' outer objects. Used to update the handles' objects when the actor construction script runs. */
 	TArray< TWeakObjectPtr<UObject> > CachedHandlesObjects;
 	/** Cached enabled value of the selected set */
