Fixed Per LOD material overrides not working on landscape streaming proxies
#tests editor

diff --git a/Engine/Source/Editor/DetailCustomizations/Private/LandscapeProxyUIDetails.cpp b/Engine/Source/Editor/DetailCustomizations/Private/LandscapeProxyUIDetails.cpp
--- a/Engine/Source/Editor/DetailCustomizations/Private/LandscapeProxyUIDetails.cpp
+++ b/Engine/Source/Editor/DetailCustomizations/Private/LandscapeProxyUIDetails.cpp
@@ -390,7 +390,7 @@ void FLandscapeProxyUIDetails::CustomizeDetails( IDetailLayoutBuilder& DetailBui
 
 					DetailRow->GetDefaultWidgets(NameWidget, ValueWidget);
 
-					DetailRow->CustomWidget()
+					DetailRow->CustomWidget(/*bShowChildren = */true)
 					.NameContent()
 					[
 						SNew(SHorizontalBox)
