PR #10335: Remove restrictions for call-in-editor Blutility functions

diff --git a/Engine/Source/Editor/DetailCustomizations/Private/ObjectDetails.cpp b/Engine/Source/Editor/DetailCustomizations/Private/ObjectDetails.cpp
--- a/Engine/Source/Editor/DetailCustomizations/Private/ObjectDetails.cpp
+++ b/Engine/Source/Editor/DetailCustomizations/Private/ObjectDetails.cpp
@@ -22,6 +22,7 @@
 #include "Misc/Attribute.h"
 #include "Misc/CString.h"
 #include "ObjectEditorUtils.h"
+#include "Settings/BlueprintEditorProjectSettings.h"
 #include "SWarningOrErrorBox.h"
 #include "ScopedTransaction.h"
 #include "SlotBase.h"
@@ -87,6 +88,8 @@ void FObjectDetails::AddCallInEditorMethods(IDetailLayoutBuilder& DetailBuilder)
 	// metadata tag for defining sort order of function buttons within a Category
 	static const FName NAME_DisplayPriority("DisplayPriority");
 
+	const bool bDisallowEditorUtilityBlueprintFunctions = GetDefault<UBlueprintEditorProjectSettings>()->bDisallowEditorUtilityBlueprintFunctionsInDetailsView;
+
 	// Get all of the functions we need to display (done ahead of time so we can sort them)
 	TArray<UFunction*, TInlineAllocator<8>> CallInEditorFunctions;
 	for (TFieldIterator<UFunction> FunctionIter(DetailBuilder.GetBaseClass(), EFieldIteratorFlags::IncludeSuper); FunctionIter; ++FunctionIter)
@@ -95,22 +98,26 @@ void FObjectDetails::AddCallInEditorMethods(IDetailLayoutBuilder& DetailBuilder)
 
 		if (TestFunction->GetBoolMetaData(FBlueprintMetadata::MD_CallInEditor) && (TestFunction->ParmsSize == 0))
 		{
+			bool bAllowFunction = true;
 			if (UClass* TestFunctionOwnerClass = TestFunction->GetOwnerClass())
 			{
 				if (UBlueprint* Blueprint = Cast<UBlueprint>(TestFunctionOwnerClass->ClassGeneratedBy))
 				{
 					if (FBlueprintEditorUtils::IsEditorUtilityBlueprint(Blueprint))
 					{
-						// Skip Blutilities as these are handled by FEditorUtilityInstanceDetails
-						continue;
+						// Skip Blutilities if disabled via project settings
+						bAllowFunction = !bDisallowEditorUtilityBlueprintFunctions;
 					}
 				}
 			}
 
-			const FName FunctionName = TestFunction->GetFName();
-			if (!CallInEditorFunctions.FindByPredicate([&FunctionName](const UFunction* Func) { return Func->GetFName() == FunctionName; }))
+			if (bAllowFunction)
 			{
-				CallInEditorFunctions.Add(*FunctionIter);
+				const FName FunctionName = TestFunction->GetFName();
+				if (!CallInEditorFunctions.FindByPredicate([&FunctionName](const UFunction* Func) { return Func->GetFName() == FunctionName; }))
+				{
+					CallInEditorFunctions.Add(*FunctionIter);
+				}
 			}
 		}
 	}
diff --git a/Engine/Source/Editor/UnrealEd/Private/Settings/BlueprintEditorProjectSettings.cpp b/Engine/Source/Editor/UnrealEd/Private/Settings/BlueprintEditorProjectSettings.cpp
--- a/Engine/Source/Editor/UnrealEd/Private/Settings/BlueprintEditorProjectSettings.cpp
+++ b/Engine/Source/Editor/UnrealEd/Private/Settings/BlueprintEditorProjectSettings.cpp
@@ -10,6 +10,7 @@
 UBlueprintEditorProjectSettings::UBlueprintEditorProjectSettings(const FObjectInitializer& ObjectInitializer)
 	: Super(ObjectInitializer)
 	, DefaultChildActorTreeViewMode(EChildActorComponentTreeViewVisualizationMode::ComponentOnly)
+	, bDisallowEditorUtilityBlueprintFunctionsInDetailsView(false)
 {
 }
 
diff --git a/Engine/Source/Editor/UnrealEd/Public/Settings/BlueprintEditorProjectSettings.h b/Engine/Source/Editor/UnrealEd/Public/Settings/BlueprintEditorProjectSettings.h
--- a/Engine/Source/Editor/UnrealEd/Public/Settings/BlueprintEditorProjectSettings.h
+++ b/Engine/Source/Editor/UnrealEd/Public/Settings/BlueprintEditorProjectSettings.h
@@ -78,6 +78,15 @@ public:
 	UPROPERTY(EditAnywhere, config, Category=Play, meta=(AllowAbstract))
 	TArray<TSoftClassPtr<UObject>> BaseClassesToAllowRecompilingDuringPlayInEditor;
 
+	/**
+	 * Disallow call-in-editor functions in Editor Utility Blueprint assets from adding buttons to the Details view when an instance of the Blueprint is selected.
+	 * 
+	 * @todo - Eventually "Blutility" functions are no longer to appear in the Details view by default, in favor of a more flexible solution. At that point, we can
+	 * expose this option and allow users to opt the project back into the legacy path until it can be fully deprecated.
+	 */
+	UPROPERTY(config)
+	bool bDisallowEditorUtilityBlueprintFunctionsInDetailsView;
+
 public:
 	// UObject interface
 	UNREALED_API virtual void PostEditChangeProperty(struct FPropertyChangedEvent& PropertyChangedEvent) override;
