* Fixed translation/rotation/scale axis sub-property not being passed to PropertyChangedEvent anymore
* Fixed landscape XY scale which could become non-uniform when pasting a non-uniform scale in the details panel

#tests editor

diff --git a/Engine/Source/Editor/DetailCustomizations/Private/ComponentTransformDetails.cpp b/Engine/Source/Editor/DetailCustomizations/Private/ComponentTransformDetails.cpp
--- a/Engine/Source/Editor/DetailCustomizations/Private/ComponentTransformDetails.cpp
+++ b/Engine/Source/Editor/DetailCustomizations/Private/ComponentTransformDetails.cpp
@@ -1163,15 +1163,18 @@ void FComponentTransformDetails::OnSetTransform(ETransformField::Type TransformF
 		// Only set axis property for single axis set
 		if (Axis == EAxisList::X)
 		{
-			AxisProperty = FindFProperty<FFloatProperty>(TBaseStructure<FVector>::Get(), GET_MEMBER_NAME_CHECKED(FVector, X));
+			AxisProperty = FindFProperty<FDoubleProperty>(TBaseStructure<FVector>::Get(), GET_MEMBER_NAME_CHECKED(FVector, X));
+			check(AxisProperty != nullptr);
 		}
 		else if (Axis == EAxisList::Y)
 		{
-			AxisProperty = FindFProperty<FFloatProperty>(TBaseStructure<FVector>::Get(), GET_MEMBER_NAME_CHECKED(FVector, Y));
+			AxisProperty = FindFProperty<FDoubleProperty>(TBaseStructure<FVector>::Get(), GET_MEMBER_NAME_CHECKED(FVector, Y));
+			check(AxisProperty != nullptr);
 		}
 		else if (Axis == EAxisList::Z)
 		{
-			AxisProperty = FindFProperty<FFloatProperty>(TBaseStructure<FVector>::Get(), GET_MEMBER_NAME_CHECKED(FVector, Z));
+			AxisProperty = FindFProperty<FDoubleProperty>(TBaseStructure<FVector>::Get(), GET_MEMBER_NAME_CHECKED(FVector, Z));
+			check(AxisProperty != nullptr);
 		}
 		break;
 	case ETransformField::Rotation:
@@ -1181,15 +1184,18 @@ void FComponentTransformDetails::OnSetTransform(ETransformField::Type TransformF
 		// Only set axis property for single axis set
 		if (Axis == EAxisList::X)
 		{
-			AxisProperty = FindFProperty<FFloatProperty>(TBaseStructure<FRotator>::Get(), GET_MEMBER_NAME_CHECKED(FRotator, Roll));
+			AxisProperty = FindFProperty<FDoubleProperty>(TBaseStructure<FRotator>::Get(), GET_MEMBER_NAME_CHECKED(FRotator, Roll));
+			check(AxisProperty != nullptr);
 		}
 		else if (Axis == EAxisList::Y)
 		{
-			AxisProperty = FindFProperty<FFloatProperty>(TBaseStructure<FRotator>::Get(), GET_MEMBER_NAME_CHECKED(FRotator, Pitch));
+			AxisProperty = FindFProperty<FDoubleProperty>(TBaseStructure<FRotator>::Get(), GET_MEMBER_NAME_CHECKED(FRotator, Pitch));
+			check(AxisProperty != nullptr);
 		}
 		else if (Axis == EAxisList::Z)
 		{
-			AxisProperty = FindFProperty<FFloatProperty>(TBaseStructure<FRotator>::Get(), GET_MEMBER_NAME_CHECKED(FRotator, Yaw));
+			AxisProperty = FindFProperty<FDoubleProperty>(TBaseStructure<FRotator>::Get(), GET_MEMBER_NAME_CHECKED(FRotator, Yaw));
+			check(AxisProperty != nullptr);
 		}
 		break;
 	case ETransformField::Scale:
@@ -1199,15 +1205,18 @@ void FComponentTransformDetails::OnSetTransform(ETransformField::Type TransformF
 		// If keep scale is set, don't set axis property
 		if (!bPreserveScaleRatio && Axis == EAxisList::X)
 		{
-			AxisProperty = FindFProperty<FFloatProperty>(TBaseStructure<FVector>::Get(), GET_MEMBER_NAME_CHECKED(FVector, X));
+			AxisProperty = FindFProperty<FDoubleProperty>(TBaseStructure<FVector>::Get(), GET_MEMBER_NAME_CHECKED(FVector, X));
+			check(AxisProperty != nullptr);
 		}
 		else if (!bPreserveScaleRatio && Axis == EAxisList::Y)
 		{
-			AxisProperty = FindFProperty<FFloatProperty>(TBaseStructure<FVector>::Get(), GET_MEMBER_NAME_CHECKED(FVector, Y));
+			AxisProperty = FindFProperty<FDoubleProperty>(TBaseStructure<FVector>::Get(), GET_MEMBER_NAME_CHECKED(FVector, Y));
+			check(AxisProperty != nullptr);
 		}
 		else if (!bPreserveScaleRatio && Axis == EAxisList::Z)
 		{
-			AxisProperty = FindFProperty<FFloatProperty>(TBaseStructure<FVector>::Get(), GET_MEMBER_NAME_CHECKED(FVector, Z));
+			AxisProperty = FindFProperty<FDoubleProperty>(TBaseStructure<FVector>::Get(), GET_MEMBER_NAME_CHECKED(FVector, Z));
+			check(AxisProperty != nullptr);
 		}
 		break;
 	default:
diff --git a/Engine/Source/Runtime/Landscape/Private/LandscapeEdit.cpp b/Engine/Source/Runtime/Landscape/Private/LandscapeEdit.cpp
--- a/Engine/Source/Runtime/Landscape/Private/LandscapeEdit.cpp
+++ b/Engine/Source/Runtime/Landscape/Private/LandscapeEdit.cpp
@@ -5696,11 +5696,22 @@ void ALandscapeProxy::PostEditChangeProperty(FPropertyChangedEvent& PropertyChan
 			// Lock X and Y scaling to the same value
 			if (SubPropertyName == FName("Y"))
 			{
-				ModifiedScale.X = FMath::Abs(OriginalScale.Y)*FMath::Sign(ModifiedScale.X);
+				ModifiedScale.X = FMath::Abs(OriginalScale.Y) * FMath::Sign(ModifiedScale.X);
 			}
 			else if (SubPropertyName == FName("X"))
 			{
-				ModifiedScale.Y = FMath::Abs(OriginalScale.X)*FMath::Sign(ModifiedScale.Y);
+				ModifiedScale.Y = FMath::Abs(OriginalScale.X) * FMath::Sign(ModifiedScale.Y);
+			}
+			else if (SubPropertyName != FName("Z"))
+			{
+				// When changing all axis values at once (e.g. when copy/pasting the scale), we receive only one event and the sub-property is not set :
+				check(SubPropertyName == PropertyName); // Any other combination of property / sub-property is invalid				
+				if (!FMath::IsNearlyEqual(ModifiedScale.X, ModifiedScale.Y))
+				{
+					UE_LOG(LogLandscape, Warning, TEXT("Non-uniform XY scale for landscape (%f, %f) : scale will be forced to (%f, %f)"), ModifiedScale.X, ModifiedScale.Y, ModifiedScale.X, ModifiedScale.X);
+					// Arbitrarily favor the X axis as the uniform scale value (but retain the sign) : 
+					ModifiedScale.Y = FMath::Abs(OriginalScale.X) * FMath::Sign(ModifiedScale.Y);
+				}
 			}
 
 			ULandscapeInfo* Info = GetLandscapeInfo();
