PR #11180: Fixed crash with the SCurveEditor widget when it has null curves.

#rnx

diff --git a/Engine/Source/Editor/DetailCustomizations/Private/CurveVectorCustomization.cpp b/Engine/Source/Editor/DetailCustomizations/Private/CurveVectorCustomization.cpp
--- a/Engine/Source/Editor/DetailCustomizations/Private/CurveVectorCustomization.cpp
+++ b/Engine/Source/Editor/DetailCustomizations/Private/CurveVectorCustomization.cpp
@@ -147,7 +147,7 @@ void FCurveVectorCustomization::CustomizeHeader(TSharedRef<IPropertyHandle> InSt
 		{
 			CurveWidget->SetCurveOwner(RuntimeCurve->ExternalCurve, false);
 		}
-		else
+		else if (RuntimeCurve)
 		{
 			CurveWidget->SetCurveOwner(this, InStructPropertyHandle->IsEditable());
 		}
