Splines: Scoped all spline detail transactions so that they only include the actual change and close the transaction before FComponentVisualizer::NotifyPropertiesModified.

diff --git a/Engine/Source/Editor/DetailCustomizations/Private/SplineComponentDetails.cpp b/Engine/Source/Editor/DetailCustomizations/Private/SplineComponentDetails.cpp
--- a/Engine/Source/Editor/DetailCustomizations/Private/SplineComponentDetails.cpp
+++ b/Engine/Source/Editor/DetailCustomizations/Private/SplineComponentDetails.cpp
@@ -913,40 +913,43 @@ void FSplinePointDetails::OnSetInputKey(float NewValue, ETextCommit::Type Commit
 		bModifyOtherPoints = true;
 	}
 
-	const FScopedTransaction Transaction(LOCTEXT("SetSplinePointInputKey", "Set spline point input key"));
-	SplineComp->Modify();
+	// Scope the transaction to only include the value change and none of the derived data changes that might arise from NotifyPropertyModified
+	{
+		const FScopedTransaction Transaction(LOCTEXT("SetSplinePointInputKey", "Set spline point input key"));
+		SplineComp->Modify();
 
-	TArray<FInterpCurvePoint<FQuat>>& Rotations = SplineComp->GetSplinePointsRotation().Points;
-	TArray<FInterpCurvePoint<FVector>>& Scales = SplineComp->GetSplinePointsScale().Points;
+		TArray<FInterpCurvePoint<FQuat>>& Rotations = SplineComp->GetSplinePointsRotation().Points;
+		TArray<FInterpCurvePoint<FVector>>& Scales = SplineComp->GetSplinePointsScale().Points;
 
-	if (bModifyOtherPoints)
-	{
-		// Shuffle the previous or next input keys down or up so the input value remains in sequence
-		if (Index > 0 && NewValue <= Positions[Index - 1].InVal)
+		if (bModifyOtherPoints)
 		{
-			float Delta = (NewValue - Positions[Index].InVal);
-			for (int32 PrevIndex = 0; PrevIndex < Index; PrevIndex++)
+			// Shuffle the previous or next input keys down or up so the input value remains in sequence
+			if (Index > 0 && NewValue <= Positions[Index - 1].InVal)
 			{
-				Positions[PrevIndex].InVal += Delta;
-				Rotations[PrevIndex].InVal += Delta;
-				Scales[PrevIndex].InVal += Delta;
+				float Delta = (NewValue - Positions[Index].InVal);
+				for (int32 PrevIndex = 0; PrevIndex < Index; PrevIndex++)
+				{
+					Positions[PrevIndex].InVal += Delta;
+					Rotations[PrevIndex].InVal += Delta;
+					Scales[PrevIndex].InVal += Delta;
+				}
 			}
-		}
-		else if (Index < NumPoints - 1 && NewValue >= Positions[Index + 1].InVal)
-		{
-			float Delta = (NewValue - Positions[Index].InVal);
-			for (int32 NextIndex = Index + 1; NextIndex < NumPoints; NextIndex++)
+			else if (Index < NumPoints - 1 && NewValue >= Positions[Index + 1].InVal)
 			{
-				Positions[NextIndex].InVal += Delta;
-				Rotations[NextIndex].InVal += Delta;
-				Scales[NextIndex].InVal += Delta;
+				float Delta = (NewValue - Positions[Index].InVal);
+				for (int32 NextIndex = Index + 1; NextIndex < NumPoints; NextIndex++)
+				{
+					Positions[NextIndex].InVal += Delta;
+					Rotations[NextIndex].InVal += Delta;
+					Scales[NextIndex].InVal += Delta;
+				}
 			}
 		}
-	}
 
-	Positions[Index].InVal = NewValue;
-	Rotations[Index].InVal = NewValue;
-	Scales[Index].InVal = NewValue;
+		Positions[Index].InVal = NewValue;
+		Rotations[Index].InVal = NewValue;
+		Scales[Index].InVal = NewValue;
+	}
 
 	SplineComp->UpdateSpline();
 	SplineComp->bSplineHasBeenEdited = true;
@@ -963,33 +966,36 @@ void FSplinePointDetails::OnSetPosition(float NewValue, ETextCommit::Type Commit
 		return;
 	}
 
-	const FScopedTransaction Transaction(LOCTEXT("SetSplinePointPosition", "Set spline point position"), !bInSliderTransaction);
-	SplineComp->Modify();
-
-	for (int32 Index : SelectedKeys)
+	// Scope the transaction to only include the value change and none of the derived data changes that might arise from NotifyPropertyModified
 	{
-		if (Index < 0 || Index >= SplineComp->GetSplinePointsPosition().Points.Num())
-		{
-			UE_LOG(LogSplineComponentDetails, Error, TEXT("Set spline point location: invalid index %d in selected points for spline component %s which contains %d spline points."),
-				Index, *SplineComp->GetPathName(), SplineComp->GetSplinePointsPosition().Points.Num());
-			continue;
-		}
+		const FScopedTransaction Transaction(LOCTEXT("SetSplinePointPosition", "Set spline point position"), !bInSliderTransaction);
+		SplineComp->Modify();
 
-		if (bEditingLocationAbsolute)
+		for (int32 Index : SelectedKeys)
 		{
-			const FTransform SplineToWorld = SplineComp->GetComponentToWorld();
-			const FVector RelativePos = SplineComp->GetSplinePointsPosition().Points[Index].OutVal;
-			FVector AbsolutePos = SplineToWorld.TransformPosition(RelativePos);
-			AbsolutePos.SetComponentForAxis(Axis, NewValue);
-			FVector PointPosition = SplineToWorld.InverseTransformPosition(AbsolutePos);
+			if (Index < 0 || Index >= SplineComp->GetSplinePointsPosition().Points.Num())
+			{
+				UE_LOG(LogSplineComponentDetails, Error, TEXT("Set spline point location: invalid index %d in selected points for spline component %s which contains %d spline points."),
+					Index, *SplineComp->GetPathName(), SplineComp->GetSplinePointsPosition().Points.Num());
+				continue;
+			}
 
-			SplineComp->GetSplinePointsPosition().Points[Index].OutVal = PointPosition;
-		}
-		else
-		{
-			FVector PointPosition = SplineComp->GetSplinePointsPosition().Points[Index].OutVal;
-			PointPosition.SetComponentForAxis(Axis, NewValue);
-			SplineComp->GetSplinePointsPosition().Points[Index].OutVal = PointPosition;
+			if (bEditingLocationAbsolute)
+			{
+				const FTransform SplineToWorld = SplineComp->GetComponentToWorld();
+				const FVector RelativePos = SplineComp->GetSplinePointsPosition().Points[Index].OutVal;
+				FVector AbsolutePos = SplineToWorld.TransformPosition(RelativePos);
+				AbsolutePos.SetComponentForAxis(Axis, NewValue);
+				FVector PointPosition = SplineToWorld.InverseTransformPosition(AbsolutePos);
+
+				SplineComp->GetSplinePointsPosition().Points[Index].OutVal = PointPosition;
+			}
+			else
+			{
+				FVector PointPosition = SplineComp->GetSplinePointsPosition().Points[Index].OutVal;
+				PointPosition.SetComponentForAxis(Axis, NewValue);
+				SplineComp->GetSplinePointsPosition().Points[Index].OutVal = PointPosition;
+			}
 		}
 	}
 
@@ -1011,22 +1017,25 @@ void FSplinePointDetails::OnSetArriveTangent(float NewValue, ETextCommit::Type C
 		return;
 	}
 
-	const FScopedTransaction Transaction(LOCTEXT("SetSplinePointTangent", "Set spline point tangent"));
-	SplineComp->Modify();
-
-	for (int32 Index : SelectedKeys)
+	// Scope the transaction to only include the value change and none of the derived data changes that might arise from NotifyPropertyModified
 	{
-		if (Index < 0 || Index >= SplineComp->GetSplinePointsPosition().Points.Num())
+		const FScopedTransaction Transaction(LOCTEXT("SetSplinePointTangent", "Set spline point tangent"));
+		SplineComp->Modify();
+
+		for (int32 Index : SelectedKeys)
 		{
-			UE_LOG(LogSplineComponentDetails, Error, TEXT("Set spline point arrive tangent: invalid index %d in selected points for spline component %s which contains %d spline points."),
-				Index, *SplineComp->GetPathName(), SplineComp->GetSplinePointsPosition().Points.Num());
-			continue;
-		}
+			if (Index < 0 || Index >= SplineComp->GetSplinePointsPosition().Points.Num())
+			{
+				UE_LOG(LogSplineComponentDetails, Error, TEXT("Set spline point arrive tangent: invalid index %d in selected points for spline component %s which contains %d spline points."),
+					Index, *SplineComp->GetPathName(), SplineComp->GetSplinePointsPosition().Points.Num());
+				continue;
+			}
 
-		FVector PointTangent = SplineComp->GetSplinePointsPosition().Points[Index].ArriveTangent;
-		PointTangent.SetComponentForAxis(Axis, NewValue);
-		SplineComp->GetSplinePointsPosition().Points[Index].ArriveTangent = PointTangent;
-		SplineComp->GetSplinePointsPosition().Points[Index].InterpMode = CIM_CurveUser;
+			FVector PointTangent = SplineComp->GetSplinePointsPosition().Points[Index].ArriveTangent;
+			PointTangent.SetComponentForAxis(Axis, NewValue);
+			SplineComp->GetSplinePointsPosition().Points[Index].ArriveTangent = PointTangent;
+			SplineComp->GetSplinePointsPosition().Points[Index].InterpMode = CIM_CurveUser;
+		}
 	}
 
 	if (CommitInfo == ETextCommit::OnEnter || CommitInfo == ETextCommit::OnUserMovedFocus)
@@ -1047,22 +1056,25 @@ void FSplinePointDetails::OnSetLeaveTangent(float NewValue, ETextCommit::Type Co
 		return;
 	}
 
-	const FScopedTransaction Transaction(LOCTEXT("SetSplinePointTangent", "Set spline point tangent"));
-	SplineComp->Modify();
-
-	for (int32 Index : SelectedKeys)
+	// Scope the transaction to only include the value change and none of the derived data changes that might arise from NotifyPropertyModified
 	{
-		if (Index < 0 || Index >= SplineComp->GetSplinePointsPosition().Points.Num())
+		const FScopedTransaction Transaction(LOCTEXT("SetSplinePointTangent", "Set spline point tangent"));
+		SplineComp->Modify();
+
+		for (int32 Index : SelectedKeys)
 		{
-			UE_LOG(LogSplineComponentDetails, Error, TEXT("Set spline point leave tangent: invalid index %d in selected points for spline component %s which contains %d spline points."),
-				Index, *SplineComp->GetPathName(), SplineComp->GetSplinePointsPosition().Points.Num());
-			continue;
-		}
+			if (Index < 0 || Index >= SplineComp->GetSplinePointsPosition().Points.Num())
+			{
+				UE_LOG(LogSplineComponentDetails, Error, TEXT("Set spline point leave tangent: invalid index %d in selected points for spline component %s which contains %d spline points."),
+					Index, *SplineComp->GetPathName(), SplineComp->GetSplinePointsPosition().Points.Num());
+				continue;
+			}
 
-		FVector PointTangent = SplineComp->GetSplinePointsPosition().Points[Index].LeaveTangent;
-		PointTangent.SetComponentForAxis(Axis, NewValue);
-		SplineComp->GetSplinePointsPosition().Points[Index].LeaveTangent = PointTangent;
-		SplineComp->GetSplinePointsPosition().Points[Index].InterpMode = CIM_CurveUser;
+			FVector PointTangent = SplineComp->GetSplinePointsPosition().Points[Index].LeaveTangent;
+			PointTangent.SetComponentForAxis(Axis, NewValue);
+			SplineComp->GetSplinePointsPosition().Points[Index].LeaveTangent = PointTangent;
+			SplineComp->GetSplinePointsPosition().Points[Index].InterpMode = CIM_CurveUser;
+		}
 	}
 
 	if (CommitInfo == ETextCommit::OnEnter || CommitInfo == ETextCommit::OnUserMovedFocus)
@@ -1082,59 +1094,63 @@ void FSplinePointDetails::OnSetRotation(float NewValue, ETextCommit::Type Commit
 	{
 		return;
 	}
-	const FScopedTransaction Transaction(LOCTEXT("SetSplinePointRotation", "Set spline point rotation"));
-	SplineComp->Modify();
-	FQuat SplineComponentRotation = SplineComp->GetComponentQuat();
+	
 	FQuat NewRotationRelative;
-	for (int32 Index : SelectedKeys)
+	// Scope the transaction to only include the value change and none of the derived data changes that might arise from NotifyPropertyModified
 	{
-		if (Index < 0 || Index >= SplineComp->GetSplinePointsRotation().Points.Num())
-		{
-			UE_LOG(LogSplineComponentDetails, Error, TEXT("Set spline point rotation: invalid index %d in selected points for spline component %s which contains %d spline points."),
-				Index, *SplineComp->GetPathName(), SplineComp->GetSplinePointsRotation().Points.Num());
-			continue;
-		}
-
-		FInterpCurvePoint<FVector>& EditedPoint = SplineComp->GetSplinePointsPosition().Points[Index];
-		FInterpCurvePoint<FQuat>& EditedRotPoint = SplineComp->GetSplinePointsRotation().Points[Index];
-		const FQuat CurrentRotationRelative = EditedRotPoint.OutVal;
-
-		if (bEditingRotationAbsolute)
+		const FScopedTransaction Transaction(LOCTEXT("SetSplinePointRotation", "Set spline point rotation"));
+		SplineComp->Modify();
+		FQuat SplineComponentRotation = SplineComp->GetComponentQuat();
+		for (int32 Index : SelectedKeys)
 		{
-			FRotator AbsoluteRot = (SplineComponentRotation * CurrentRotationRelative).Rotator();
-
-			switch (Axis)
+			if (Index < 0 || Index >= SplineComp->GetSplinePointsRotation().Points.Num())
 			{
-			case EAxis::X: AbsoluteRot.Roll = NewValue; break;
-			case EAxis::Y: AbsoluteRot.Pitch = NewValue; break;
-			case EAxis::Z: AbsoluteRot.Yaw = NewValue; break;
+				UE_LOG(LogSplineComponentDetails, Error, TEXT("Set spline point rotation: invalid index %d in selected points for spline component %s which contains %d spline points."),
+					Index, *SplineComp->GetPathName(), SplineComp->GetSplinePointsRotation().Points.Num());
+				continue;
 			}
 
-			NewRotationRelative = SplineComponentRotation.Inverse() * AbsoluteRot.Quaternion();
-		}
-		else
-		{
-			FRotator NewRotationRotator(CurrentRotationRelative);
+			FInterpCurvePoint<FVector>& EditedPoint = SplineComp->GetSplinePointsPosition().Points[Index];
+			FInterpCurvePoint<FQuat>& EditedRotPoint = SplineComp->GetSplinePointsRotation().Points[Index];
+			const FQuat CurrentRotationRelative = EditedRotPoint.OutVal;
 
-			switch (Axis)
+			if (bEditingRotationAbsolute)
 			{
-			case EAxis::X: NewRotationRotator.Roll = NewValue; break;
-			case EAxis::Y: NewRotationRotator.Pitch = NewValue; break;
-			case EAxis::Z: NewRotationRotator.Yaw = NewValue; break;
+				FRotator AbsoluteRot = (SplineComponentRotation * CurrentRotationRelative).Rotator();
+
+				switch (Axis)
+				{
+				case EAxis::X: AbsoluteRot.Roll = NewValue; break;
+				case EAxis::Y: AbsoluteRot.Pitch = NewValue; break;
+				case EAxis::Z: AbsoluteRot.Yaw = NewValue; break;
+				}
+
+				NewRotationRelative = SplineComponentRotation.Inverse() * AbsoluteRot.Quaternion();
 			}
+			else
+			{
+				FRotator NewRotationRotator(CurrentRotationRelative);
 
-			NewRotationRelative = NewRotationRotator.Quaternion();
-		}
+				switch (Axis)
+				{
+				case EAxis::X: NewRotationRotator.Roll = NewValue; break;
+				case EAxis::Y: NewRotationRotator.Pitch = NewValue; break;
+				case EAxis::Z: NewRotationRotator.Yaw = NewValue; break;
+				}
 
-		SplineComp->GetSplinePointsRotation().Points[Index].OutVal = NewRotationRelative;
+				NewRotationRelative = NewRotationRotator.Quaternion();
+			}
 
-		FQuat DeltaRotate(NewRotationRelative * CurrentRotationRelative.Inverse());
-		// Rotate tangent according to delta rotation
-		FVector NewTangent = SplineComponentRotation.RotateVector(EditedPoint.LeaveTangent); // convert local-space tangent vector to world-space
-		NewTangent = DeltaRotate.RotateVector(NewTangent); // apply world-space delta rotation to world-space tangent
-		NewTangent = SplineComponentRotation.Inverse().RotateVector(NewTangent); // convert world-space tangent vector back into local-space
-		EditedPoint.LeaveTangent = NewTangent;
-		EditedPoint.ArriveTangent = NewTangent;
+			SplineComp->GetSplinePointsRotation().Points[Index].OutVal = NewRotationRelative;
+
+			FQuat DeltaRotate(NewRotationRelative * CurrentRotationRelative.Inverse());
+			// Rotate tangent according to delta rotation
+			FVector NewTangent = SplineComponentRotation.RotateVector(EditedPoint.LeaveTangent); // convert local-space tangent vector to world-space
+			NewTangent = DeltaRotate.RotateVector(NewTangent); // apply world-space delta rotation to world-space tangent
+			NewTangent = SplineComponentRotation.Inverse().RotateVector(NewTangent); // convert world-space tangent vector back into local-space
+			EditedPoint.LeaveTangent = NewTangent;
+			EditedPoint.ArriveTangent = NewTangent;
+		}
 	}
 
 	SplineVisualizer->SetCachedRotation(NewRotationRelative);
@@ -1156,21 +1172,24 @@ void FSplinePointDetails::OnSetScale(float NewValue, ETextCommit::Type CommitInf
 		return;
 	}
 
-	const FScopedTransaction Transaction(LOCTEXT("SetSplinePointScale", "Set spline point scale"));
-	SplineComp->Modify();
-
-	for (int32 Index : SelectedKeys)
+	// Scope the transaction to only include the value change and none of the derived data changes that might arise from NotifyPropertyModified
 	{
-		if (Index < 0 || Index >= SplineComp->GetSplinePointsScale().Points.Num())
+		const FScopedTransaction Transaction(LOCTEXT("SetSplinePointScale", "Set spline point scale"));
+		SplineComp->Modify();
+
+		for (int32 Index : SelectedKeys)
 		{
-			UE_LOG(LogSplineComponentDetails, Error, TEXT("Set spline point scale: invalid index %d in selected points for spline component %s which contains %d spline points."),
-				Index, *SplineComp->GetPathName(), SplineComp->GetSplinePointsScale().Points.Num());
-			continue;
-		}
+			if (Index < 0 || Index >= SplineComp->GetSplinePointsScale().Points.Num())
+			{
+				UE_LOG(LogSplineComponentDetails, Error, TEXT("Set spline point scale: invalid index %d in selected points for spline component %s which contains %d spline points."),
+					Index, *SplineComp->GetPathName(), SplineComp->GetSplinePointsScale().Points.Num());
+				continue;
+			}
 
-		FVector PointScale = SplineComp->GetSplinePointsScale().Points[Index].OutVal;
-		PointScale.SetComponentForAxis(Axis, NewValue);
-		SplineComp->GetSplinePointsScale().Points[Index].OutVal = PointScale;
+			FVector PointScale = SplineComp->GetSplinePointsScale().Points[Index].OutVal;
+			PointScale.SetComponentForAxis(Axis, NewValue);
+			SplineComp->GetSplinePointsScale().Points[Index].OutVal = PointScale;
+		}
 	}
 
 	if (CommitInfo == ETextCommit::OnEnter || CommitInfo == ETextCommit::OnUserMovedFocus)
@@ -1203,29 +1222,32 @@ void FSplinePointDetails::OnSplinePointTypeChanged(TSharedPtr<FString> NewValue,
 		return;
 	}
 
-	const FScopedTransaction Transaction(LOCTEXT("SetSplinePointType", "Set spline point type"));
-	SplineComp->Modify();
-
-	EInterpCurveMode Mode = CIM_Unknown;
-	if (NewValue.IsValid() && SplinePointTypes.Contains(NewValue))
+	// Scope the transaction to only include the value change and none of the derived data changes that might arise from NotifyPropertyModified
 	{
-		const UEnum* SplinePointTypeEnum = StaticEnum<ESplinePointType::Type>();
-		check(SplinePointTypeEnum);
-		const int64 SplinePointType = SplinePointTypeEnum->GetValueByNameString(*NewValue);
-
-		Mode = ConvertSplinePointTypeToInterpCurveMode(static_cast<ESplinePointType::Type>(SplinePointType));
-	}
+		const FScopedTransaction Transaction(LOCTEXT("SetSplinePointType", "Set spline point type"));
+		SplineComp->Modify();
 
-	for (int32 Index : SelectedKeys)
-	{
-		if (Index < 0 || Index >= SplineComp->GetSplinePointsPosition().Points.Num())
+		EInterpCurveMode Mode = CIM_Unknown;
+		if (NewValue.IsValid() && SplinePointTypes.Contains(NewValue))
 		{
-			UE_LOG(LogSplineComponentDetails, Error, TEXT("Set spline point type: invalid index %d in selected points for spline component %s which contains %d spline points."),
-				Index, *SplineComp->GetPathName(), SplineComp->GetSplinePointsPosition().Points.Num());
-			continue;
+			const UEnum* SplinePointTypeEnum = StaticEnum<ESplinePointType::Type>();
+			check(SplinePointTypeEnum);
+			const int64 SplinePointType = SplinePointTypeEnum->GetValueByNameString(*NewValue);
+
+			Mode = ConvertSplinePointTypeToInterpCurveMode(static_cast<ESplinePointType::Type>(SplinePointType));
 		}
 
-		SplineComp->GetSplinePointsPosition().Points[Index].InterpMode = Mode;
+		for (int32 Index : SelectedKeys)
+		{
+			if (Index < 0 || Index >= SplineComp->GetSplinePointsPosition().Points.Num())
+			{
+				UE_LOG(LogSplineComponentDetails, Error, TEXT("Set spline point type: invalid index %d in selected points for spline component %s which contains %d spline points."),
+					Index, *SplineComp->GetPathName(), SplineComp->GetSplinePointsPosition().Points.Num());
+				continue;
+			}
+
+			SplineComp->GetSplinePointsPosition().Points[Index].InterpMode = Mode;
+		}
 	}
 
 	SplineComp->UpdateSpline();
