Add ShowNormalize option for math structs that will normalize the vector struct when pressed.
Prevent showing slider for vectors if NoSpinBox was requested.

diff --git a/Engine/Source/Editor/DetailCustomizations/Private/MathStructCustomizations.cpp b/Engine/Source/Editor/DetailCustomizations/Private/MathStructCustomizations.cpp
--- a/Engine/Source/Editor/DetailCustomizations/Private/MathStructCustomizations.cpp
+++ b/Engine/Source/Editor/DetailCustomizations/Private/MathStructCustomizations.cpp
@@ -28,6 +28,7 @@
 #include "UObject/UnrealType.h"
 #include "Widgets/DeclarativeSyntaxSupport.h"
 #include "Widgets/Images/SImage.h"
+#include "Widgets/Input/SButton.h"
 #include "Widgets/Input/SCheckBox.h"
 #include "Widgets/Input/SNumericEntryBox.h"
 #include "Widgets/SBoxPanel.h"
@@ -68,6 +69,62 @@ void FMathStructCustomization::CustomizeChildren(TSharedRef<class IPropertyHandl
 }
 
 
+namespace MathStructCustomization {
+	double ScaleStructComponent(double Val, double Scale)
+	{
+		return Val * Scale;
+	}
+
+	template<typename NumericType>
+	void NormalizePropertyVector(TWeakPtr<IPropertyHandle> PropertyHandle)
+	{
+		double SquareSum = 0;
+		TSharedPtr<IPropertyHandle> Property = PropertyHandle.Pin();
+
+		uint32 NumChildren;
+		Property->GetNumChildren(NumChildren);
+
+		// Loop through each child object and build a square sum
+		for (uint32 ChildIndex = 0; ChildIndex < NumChildren; ++ChildIndex)
+		{
+			TSharedPtr<IPropertyHandle> Child = Property->GetChildHandle(ChildIndex);
+			FProperty* ChildProperty = Child->GetProperty();
+			NumericType Val;
+			Child->GetValue(Val);
+			SquareSum += Val * Val;
+		}
+		if (SquareSum > UE_SMALL_NUMBER * UE_SMALL_NUMBER)
+		{
+			// Calculate the scale each object will need to be multiplied by to achieve a unit vector
+			double Scale = FMath::InvSqrt(SquareSum);
+
+			// Loop through each object and scale based on the normalized ratio for each object individually
+			for (uint32 ChildIndex = 0; ChildIndex < NumChildren; ++ChildIndex)
+			{
+				TSharedPtr<IPropertyHandle> Child = Property->GetChildHandle(ChildIndex);
+				FProperty* ChildProperty = Child->GetProperty();
+				NumericType Val;
+				Child->GetValue(Val);
+				NumericType Result = (NumericType)ScaleStructComponent(Val, Scale);
+				Child->SetValue(Result);
+			}
+		}
+	}
+
+	bool IsFloatVector(TSharedRef<class IPropertyHandle>& PropertyHandle)
+	{
+		// Look at the first child element to see if it's something we can normalize
+		TSharedPtr<IPropertyHandle> Child = PropertyHandle->GetChildHandle(0);
+		if (Child)
+		{
+			FNumericProperty* ChildProperty = CastField<FNumericProperty>(Child->GetProperty());
+
+			return ChildProperty && ChildProperty->IsFloatingPoint();
+		}
+		return false;
+	}
+}
+
 void FMathStructCustomization::MakeHeaderRow(TSharedRef<class IPropertyHandle>& StructPropertyHandle, FDetailWidgetRow& Row)
 {
 	TWeakPtr<IPropertyHandle> StructWeakHandlePtr = StructPropertyHandle;
@@ -116,7 +173,7 @@ void FMathStructCustomization::MakeHeaderRow(TSharedRef<class IPropertyHandle>&
 		];
 	}
 
-	if (StructPropertyHandle->GetProperty()->HasMetaData("AllowPreserveRatio"))
+	if (StructPropertyHandle->HasMetaData("AllowPreserveRatio"))
 	{
 		if (!GConfig->GetBool(TEXT("SelectionDetails"), *(StructPropertyHandle->GetProperty()->GetName() + TEXT("_PreserveScaleRatio")), bPreserveScaleRatio, GEditorPerProjectIni))
 		{
@@ -141,21 +198,38 @@ void FMathStructCustomization::MakeHeaderRow(TSharedRef<class IPropertyHandle>&
 			]
 		];
 	}
-}
 
+	if (StructPropertyHandle->HasMetaData("ShowNormalize") && MathStructCustomization::IsFloatVector(StructPropertyHandle))
+	{
+		HorizontalBox->AddSlot()
+			.AutoWidth()
+			.MaxWidth(18.0f)
+			.VAlign(VAlign_Center)
+			[
+				// Add a button to scale the vector uniformly to achieve a unit vector
+				SNew(SButton)
+					.OnClicked(this, &FMathStructCustomization::OnNormalizeClicked, StructWeakHandlePtr)
+					.ButtonStyle(FAppStyle::Get(), "NoBorder")
+					.ToolTipText(LOCTEXT("NormalizeToolTip", "When clicked, if the vector is large enough, it scales the vector uniformly to achieve a unit vector (vector with a length of 1)"))
+					[
+						SNew(SImage)
+							.ColorAndOpacity(FSlateColor::UseForeground())
+							.Image(FAppStyle::GetBrush(TEXT("Icons.Normalize")))	
+					]
+			];
+	}
+}
 
 const FSlateBrush* FMathStructCustomization::GetPreserveScaleRatioImage() const
 {
 	return bPreserveScaleRatio ? FAppStyle::GetBrush(TEXT("Icons.Lock")) : FAppStyle::GetBrush(TEXT("Icons.Unlock"));
 }
 
-
 ECheckBoxState FMathStructCustomization::IsPreserveScaleRatioChecked() const
 {
 	return bPreserveScaleRatio ? ECheckBoxState::Checked : ECheckBoxState::Unchecked;
 }
 
-
 void FMathStructCustomization::OnPreserveScaleRatioToggled(ECheckBoxState NewState, TWeakPtr<IPropertyHandle> PropertyHandle)
 {
 	bPreserveScaleRatio = (NewState == ECheckBoxState::Checked) ? true : false;
@@ -167,6 +241,24 @@ void FMathStructCustomization::OnPreserveScaleRatioToggled(ECheckBoxState NewSta
 	}
 }
 
+FReply FMathStructCustomization::OnNormalizeClicked(TWeakPtr<IPropertyHandle> PropertyHandle)
+{
+	if (PropertyHandle.IsValid())
+	{
+		TSharedPtr<IPropertyHandle> Property = PropertyHandle.Pin();
+		TSharedRef<IPropertyHandle> PropertyRef = Property.ToSharedRef();
+
+		if (MathStructCustomization::IsFloatVector(PropertyRef))
+		{
+			MathStructCustomization::NormalizePropertyVector<double>(Property);
+		}
+		else
+		{
+			ensureMsgf(false, TEXT("Unsupported type to Normalize"));
+		}
+	}
+	return FReply::Handled();
+}
 
 void FMathStructCustomization::GetSortedChildren(TSharedRef<IPropertyHandle> StructPropertyHandle, TArray< TSharedRef<IPropertyHandle> >& OutChildren)
 {
@@ -239,6 +331,7 @@ void FMathStructCustomization::ExtractNumericMetadata(TSharedRef<IPropertyHandle
 	// If no UIMin/Max was specified then use the clamp string
 	const FString& UIMinString = MetaUIMinString.Len() ? MetaUIMinString : ClampMinString;
 	const FString& UIMaxString = MetaUIMaxString.Len() ? MetaUIMaxString : ClampMaxString;
+	bool bAllowSpin = !Property->GetBoolMetaData(TEXT("NoSpinBox"));
 
 	NumericType ClampMin = TNumericLimits<NumericType>::Lowest();
 	NumericType ClampMax = TNumericLimits<NumericType>::Max();
@@ -312,6 +405,7 @@ void FMathStructCustomization::ExtractNumericMetadata(TSharedRef<IPropertyHandle
 	
 	MetadataOut.bSupportDynamicSliderMaxValue = SupportDynamicSliderMaxValueString.Len() > 0 && SupportDynamicSliderMaxValueString.ToBool();
 	MetadataOut.bSupportDynamicSliderMinValue = SupportDynamicSliderMinValueString.Len() > 0 && SupportDynamicSliderMinValueString.ToBool();
+	MetadataOut.bAllowSpinBox = bAllowSpin;
 }
 
 
@@ -337,7 +431,7 @@ TSharedRef<SWidget> FMathStructCustomization::MakeNumericWidget(
 			.OnBeginSliderMovement(this, &FMathStructCustomization::OnBeginSliderMovement)
 			.OnEndSliderMovement(this, &FMathStructCustomization::OnEndSliderMovement<NumericType>)
 			// Only allow spin on handles with one object.  Otherwise it is not clear what value to spin
-			.AllowSpin(PropertyHandle->GetNumOuterObjects() < 2)
+			.AllowSpin(PropertyHandle->GetNumOuterObjects() < 2 && Metadata.bAllowSpinBox)
 			.ShiftMouseMovePixelPerDelta(Metadata.ShiftMouseMovePixelPerDelta)
 			.SupportDynamicSliderMaxValue(Metadata.bSupportDynamicSliderMaxValue)
 			.SupportDynamicSliderMinValue(Metadata.bSupportDynamicSliderMinValue)
diff --git a/Engine/Source/Editor/DetailCustomizations/Public/Customizations/MathStructCustomizations.h b/Engine/Source/Editor/DetailCustomizations/Public/Customizations/MathStructCustomizations.h
--- a/Engine/Source/Editor/DetailCustomizations/Public/Customizations/MathStructCustomizations.h
+++ b/Engine/Source/Editor/DetailCustomizations/Public/Customizations/MathStructCustomizations.h
@@ -75,6 +75,7 @@ public:
 		int32 ShiftMouseMovePixelPerDelta;
 		bool bSupportDynamicSliderMaxValue;
 		bool bSupportDynamicSliderMinValue;
+		bool bAllowSpinBox;
 	};
 
 	/** Utility function that will extract common Math related numeric metadata */	
@@ -173,6 +174,8 @@ private:
 	/** Called when the user toggles preserve ratio. */
 	void OnPreserveScaleRatioToggled(ECheckBoxState NewState, TWeakPtr<IPropertyHandle> PropertyHandle);
 
+	FReply OnNormalizeClicked(TWeakPtr<IPropertyHandle> PropertyHandle);
+
 private:
 
 	/** Called when a value starts to be changed by a slider */
diff --git a/Engine/Source/Runtime/SlateCore/Private/Styling/StarshipCoreStyle.cpp b/Engine/Source/Runtime/SlateCore/Private/Styling/StarshipCoreStyle.cpp
--- a/Engine/Source/Runtime/SlateCore/Private/Styling/StarshipCoreStyle.cpp
+++ b/Engine/Source/Runtime/SlateCore/Private/Styling/StarshipCoreStyle.cpp
@@ -353,6 +353,8 @@ TSharedRef<ISlateStyle> FStarshipCoreStyle::Create()
 		Style->Set("Icons.Lock", new IMAGE_BRUSH_SVG("Starship/Common/lock", Icon16x16));
 		Style->Set("Icons.Unlock", new IMAGE_BRUSH_SVG("Starship/Common/lock-unlocked", Icon16x16));
 
+		Style->Set("Icons.Normalize", new IMAGE_BRUSH("Starship/Common/normalize", Icon16x16));
+
 		Style->Set("Icons.CircleArrowLeft", new IMAGE_BRUSH_SVG("Starship/Common/circle-arrow-left", Icon16x16));
 		Style->Set("Icons.CircleArrowRight", new IMAGE_BRUSH_SVG("Starship/Common/circle-arrow-right", Icon16x16));
 
