Would be helpful to have a pop-up dialogue with options for repairing montage/anim length discrepancies, and to add these checks to the in-editor content check button.
#feature added warning box + options buttons for user choice to Animation (montage) segment details
#misc had to make CachedPlayLength editable (but hidden by details customization) to use property-node SetValue API
#misc upgraded segment viewport client to use AdvancedPreviewScene

diff --git a/Engine/Source/Editor/DetailCustomizations/Private/AnimMontageSegmentDetails.cpp b/Engine/Source/Editor/DetailCustomizations/Private/AnimMontageSegmentDetails.cpp
--- a/Engine/Source/Editor/DetailCustomizations/Private/AnimMontageSegmentDetails.cpp
+++ b/Engine/Source/Editor/DetailCustomizations/Private/AnimMontageSegmentDetails.cpp
@@ -57,6 +57,8 @@
 #include "Widgets/SBoxPanel.h"
 #include "Widgets/SViewport.h"
 #include "Widgets/Text/STextBlock.h"
+#include "SWarningOrErrorBox.h"
+#include "Widgets/Input/SEditableText.h"
 
 class SWidget;
 struct FGeometry;
@@ -64,7 +66,7 @@ struct FGeometry;
 #define LOCTEXT_NAMESPACE "AnimMontageSegmentDetails"
 
 /////////////////////////////////////////////////////////////////////////
-FAnimationSegmentViewportClient::FAnimationSegmentViewportClient(FPreviewScene& InPreviewScene, const TWeakPtr<SEditorViewport>& InEditorViewportWidget)
+FAnimationSegmentViewportClient::FAnimationSegmentViewportClient(FAdvancedPreviewScene& InPreviewScene, const TWeakPtr<SEditorViewport>& InEditorViewportWidget)
 	: FEditorViewportClient(nullptr, &InPreviewScene, InEditorViewportWidget)
 {
 	SetViewMode(VMI_Lit);
@@ -324,6 +326,89 @@ void FAnimMontageSegmentDetails::CustomizeDetails( IDetailLayoutBuilder& DetailB
 			}
 		})
 	];
+	
+	const TSharedPtr<IPropertyHandle> PlayLengthHandle = AnimSegmentHandle->GetChildHandle(GET_MEMBER_NAME_CHECKED(FAnimSegment, CachedPlayLength));
+	PlayLengthHandle->MarkHiddenByCustomization();
+	IDetailCategoryBuilder& WarningCategory = DetailBuilder.EditCategory(TEXT("Warning"), LOCTEXT("WarningCategoryDisplayName", "Warning"), ECategoryPriority::Transform);
+	WarningCategory.AddCustomRow(LOCTEXT("WarningCategoryDisplayName", "Warning"))
+	.WholeRowContent()
+	[
+		SNew(SBox)
+		.Padding(FMargin(0.f, 4.f))
+		[
+			SNew(SWarningOrErrorBox)
+			.MessageStyle(EMessageStyle::Warning)
+			.Message_Lambda([this]() -> FText
+			{
+				const FAnimSegment* Segment = GetAnimationSegment();
+				static const FNumberFormattingOptions FormatOptions = FNumberFormattingOptions()
+					.SetMinimumFractionalDigits(2)
+					.SetMaximumFractionalDigits(2);
+				
+				return FText::Format(LOCTEXT("AnimSegmentLengthMismatchWarning", "Referenced Animation length has changed and mismatches with set Segment length. Animation length: {0}. Segment length: {1}."), FText::AsNumber(Segment->GetAnimReference()->GetPlayLength(), &FormatOptions), FText::AsNumber(Segment->CachedPlayLength, &FormatOptions));
+			})
+			[
+				SNew(SVerticalBox)
+				+SVerticalBox::Slot()
+				[
+					SNew(SButton)
+					.OnClicked_Lambda([this, PlayLengthHandle]()
+					{						
+						if (PlayLengthHandle.IsValid())
+						{
+							if (const FAnimSegment* AnimSegment = GetAnimationSegment())
+							{									
+								if(const IAnimationDataModel* DataModel = AnimSegment->GetAnimReference()->GetDataModel())
+								{
+									const float NewPlayLength = DataModel->GetPlayLength();
+									PlayLengthHandle->SetValue(NewPlayLength);
+								}
+							}
+						}
+							
+						return FReply::Handled();
+					})
+					.TextStyle(FAppStyle::Get(), "NormalText")
+					.Text(LOCTEXT("KeepSegmentLengthButtonText", "Keep Segment Length"))
+					.ToolTipText(LOCTEXT("KeepSegmentLengthButtonToolTip", "This will update the cached Animation Asset length with its current value while keeping the Segment Length the same."))
+					.HAlign(EHorizontalAlignment::HAlign_Center)
+				]
+				+SVerticalBox::Slot()
+				[
+					SNew(SButton)
+					.OnClicked_Lambda([this]()
+					{						
+						if (AnimEndTimeProperty.IsValid())
+						{
+							if (const FAnimSegment* AnimSegment = GetAnimationSegment())
+							{
+								if(const IAnimationDataModel* DataModel = AnimSegment->GetAnimReference()->GetDataModel())
+								{
+									const float NewPlayLength = DataModel->GetPlayLength();
+									AnimEndTimeProperty->SetValue(NewPlayLength);
+								}
+							}
+						}
+							
+						return FReply::Handled();
+					})
+					.TextStyle(FAppStyle::Get(), "NormalText")
+					.Text(LOCTEXT("SyncWithAnimLengthButtonText", "Sync with Animation Length"))
+					.ToolTipText(LOCTEXT("SyncWithAnimLengthButtonToolTip", "This will update the Segment Length to match the Animation Asset length its current value."))
+					.HAlign(EHorizontalAlignment::HAlign_Center)
+				]
+			]
+		]
+	]
+	.Visibility(TAttribute<EVisibility>::CreateLambda([this]() -> EVisibility
+	{
+		if (GetAnimationSegment() && GetAnimationSegment()->IsPlayLengthOutOfDate())
+		{
+			return EVisibility::Visible;
+		}
+
+		return EVisibility::Collapsed;
+	}));	
 }
 
 bool FAnimMontageSegmentDetails::OnShouldFilterAnimAsset(const FAssetData& AssetData) const
@@ -406,9 +491,18 @@ FAnimSegment* FAnimMontageSegmentDetails::GetAnimationSegment() const
 {
 	if (AnimSegmentHandle.IsValid())
 	{
-		void* StructPtr = nullptr;	
-		AnimSegmentHandle->GetValueData(StructPtr);
-		return (FAnimSegment*)StructPtr;
+		void* Data;
+		const FPropertyAccess::Result Result = AnimSegmentHandle->GetValueData(Data);
+		if (Result == FPropertyAccess::MultipleValues)
+		{
+			return nullptr;
+		}
+
+		if (Result == FPropertyAccess::Success)
+		{
+			FAnimSegment* AnimSegment = reinterpret_cast<FAnimSegment*>(Data);
+			return AnimSegment;
+		}
 	}
 
 	return nullptr;
@@ -471,7 +565,7 @@ void SAnimationSegmentViewport::CleanupComponent(USceneComponent* Component)
 }
 
 SAnimationSegmentViewport::SAnimationSegmentViewport()
-	: CurrentAnimSequenceBase(nullptr), PreviewScene(FPreviewScene::ConstructionValues())
+	: CurrentAnimSequenceBase(nullptr), AdvancedPreviewScene(FPreviewScene::ConstructionValues())
 {
 }
 
@@ -524,7 +618,7 @@ void SAnimationSegmentViewport::Construct(const FArguments& InArgs)
 	
 
 	// Create a viewport client
-	LevelViewportClient	= MakeShareable( new FAnimationSegmentViewportClient(PreviewScene) );
+	LevelViewportClient	= MakeShareable( new FAnimationSegmentViewportClient(AdvancedPreviewScene) );
 
 	LevelViewportClient->ViewportType = LVT_Perspective;
 	LevelViewportClient->bSetListenerPosition = false;
@@ -541,7 +635,7 @@ void SAnimationSegmentViewport::Construct(const FArguments& InArgs)
 	
 	PreviewComponent = NewObject<UDebugSkelMeshComponent>();
 	PreviewComponent->VisibilityBasedAnimTickOption = EVisibilityBasedAnimTickOption::AlwaysTickPoseAndRefreshBones;
-	PreviewScene.AddComponent(PreviewComponent, FTransform::Identity);
+	AdvancedPreviewScene.AddComponent(PreviewComponent, FTransform::Identity);
 
 	InitSkeleton();
 }
diff --git a/Engine/Source/Editor/DetailCustomizations/Private/AnimMontageSegmentDetails.h b/Engine/Source/Editor/DetailCustomizations/Private/AnimMontageSegmentDetails.h
--- a/Engine/Source/Editor/DetailCustomizations/Private/AnimMontageSegmentDetails.h
+++ b/Engine/Source/Editor/DetailCustomizations/Private/AnimMontageSegmentDetails.h
@@ -2,6 +2,7 @@
 
 #pragma once
 
+#include "AdvancedPreviewScene.h"
 #include "Containers/Array.h"
 #include "Delegates/Delegate.h"
 #include "EditorViewportClient.h"
@@ -108,10 +109,8 @@ private:
 
 	TSharedPtr<class FSceneViewport> SceneViewport;
 
-
 	TObjectPtr<UAnimSequenceBase> CurrentAnimSequenceBase;
-
-	FPreviewScene PreviewScene;
+	FAdvancedPreviewScene AdvancedPreviewScene;
 	class UDebugSkelMeshComponent* PreviewComponent;
 
 	TSharedPtr<STextBlock> Description;
@@ -197,7 +196,7 @@ protected:
 class FAnimationSegmentViewportClient : public FEditorViewportClient
 {
 public:
-	FAnimationSegmentViewportClient(FPreviewScene& InPreviewScene, const TWeakPtr<SEditorViewport>& InEditorViewportWidget = nullptr);
+	FAnimationSegmentViewportClient(FAdvancedPreviewScene& InPreviewScene, const TWeakPtr<SEditorViewport>& InEditorViewportWidget = nullptr);
 
 	// FlEditorViewportClient interface
 	virtual FSceneInterface* GetScene() const override;
diff --git a/Engine/Source/Runtime/Engine/Classes/Animation/AnimCompositeBase.h b/Engine/Source/Runtime/Engine/Classes/Animation/AnimCompositeBase.h
--- a/Engine/Source/Runtime/Engine/Classes/Animation/AnimCompositeBase.h
+++ b/Engine/Source/Runtime/Engine/Classes/Animation/AnimCompositeBase.h
@@ -81,7 +81,7 @@ struct FAnimSegment
 	TObjectPtr<UAnimSequenceBase> AnimReference;
 
 #if WITH_EDITORONLY_DATA
-	UPROPERTY()
+	UPROPERTY(EditAnywhere, Category=AnimSegment, meta=(DisplayName = "Cached Animation Asset length"))
 	float CachedPlayLength = 0.f;
 #endif
 
