Restrict world context signature matching to function libraries, as those have static functions. This fixes EditorUtilityActors not having buttons for their member functions in the details panel

diff --git a/Engine/Source/Editor/DetailCustomizations/Private/ObjectDetails.cpp b/Engine/Source/Editor/DetailCustomizations/Private/ObjectDetails.cpp
--- a/Engine/Source/Editor/DetailCustomizations/Private/ObjectDetails.cpp
+++ b/Engine/Source/Editor/DetailCustomizations/Private/ObjectDetails.cpp
@@ -96,7 +96,7 @@ static bool CanCallFunctionBasedOnParams(const UFunction* TestFunction)
 	{
 		if (UBlueprint* Blueprint = Cast<UBlueprint>(TestFunctionOwnerClass->ClassGeneratedBy))
 		{
-			if (FBlueprintEditorUtils::IsEditorUtilityBlueprint(Blueprint))
+			if (FBlueprintEditorUtils::IsEditorUtilityBlueprint(Blueprint) && Blueprint->BlueprintType == BPTYPE_FunctionLibrary)
 			{
 				using namespace UE::Reflection;
 				return TestFunction->HasMetaData(FBlueprintMetadata::MD_WorldContext) &&
