Anim Sequence Details now generates a list of bone animation tracks from the data model

diff --git a/Engine/Source/Editor/DetailCustomizations/Private/AnimSequenceDetails.cpp b/Engine/Source/Editor/DetailCustomizations/Private/AnimSequenceDetails.cpp
--- a/Engine/Source/Editor/DetailCustomizations/Private/AnimSequenceDetails.cpp
+++ b/Engine/Source/Editor/DetailCustomizations/Private/AnimSequenceDetails.cpp
@@ -56,6 +56,8 @@
 #include "Widgets/SBoxPanel.h"
 #include "Widgets/SViewport.h"
 #include "Widgets/Text/STextBlock.h"
+#include "PropertyCustomizationHelpers.h"
+#include "IDetailChildrenBuilder.h"
 
 class SWidget;
 struct FGeometry;
@@ -74,9 +76,11 @@ BEGIN_SLATE_FUNCTION_BUILD_OPTIMIZATION
 void FAnimSequenceDetails::CustomizeDetails(IDetailLayoutBuilder& DetailBuilder)
 {
 	/////////////////////////////////////////////////////////////////////////////////
-	// retarget source handler in Animation
+	// Animation
 	/////////////////////////////////////////////////////////////////////////////////
 	IDetailCategoryBuilder& AnimationCategory = DetailBuilder.EditCategory("Animation");
+
+	// *** Retarget source handler ***
 	RetargetSourceNameHandler = DetailBuilder.GetProperty(GET_MEMBER_NAME_CHECKED(UAnimSequence, RetargetSource));
 	RetargetSourceAssetHandle = DetailBuilder.GetProperty(GET_MEMBER_NAME_CHECKED(UAnimSequence, RetargetSourceAsset));
 
@@ -199,6 +203,38 @@ void FAnimSequenceDetails::CustomizeDetails(IDetailLayoutBuilder& DetailBuilder)
 	DetailBuilder.HideProperty(RetargetSourceNameHandler);
 	DetailBuilder.HideProperty(RetargetSourceAssetHandle);
 
+	// *** Animation Track Names ***
+	PRAGMA_DISABLE_DEPRECATION_WARNINGS
+	AnimationTrackNamesHandle = DetailBuilder.GetProperty(GET_MEMBER_NAME_CHECKED(UAnimSequence, AnimationTrackNames));
+	PRAGMA_ENABLE_DEPRECATION_WARNINGS
+
+	if (AnimationTrackNamesHandle.IsValid())
+	{
+		TArray<TWeakObjectPtr<UAnimSequence>> AnimSequences = DetailBuilder.GetSelectedObjectsOfType<UAnimSequence>();
+		AnimationTrackNamesList.Empty();
+
+		if (!AnimSequences.IsEmpty())
+		{
+			for (const TWeakObjectPtr<UAnimSequence>& AnimSequenceWeak : AnimSequences)
+			{
+				if(AnimSequenceWeak.IsValid())
+				{ 
+					if (AnimSequenceWeak->IsDataModelValid())
+					{
+						AnimSequenceWeak->GetDataModelInterface()->GetBoneTrackNames(AnimationTrackNamesList);
+					}
+				}
+			}
+		}
+
+		TSharedRef<FDetailArrayBuilder> AnimationTrackNamesArrayBuilder = MakeShareable(new FDetailArrayBuilder(AnimationTrackNamesHandle.ToSharedRef(), true, false, true));
+		AnimationTrackNamesArrayBuilder->OnGenerateArrayElementWidget(FOnGenerateArrayElementWidget::CreateSP(this, &FAnimSequenceDetails::GenerateAnimationTrackNameArrayElementWidget, &DetailBuilder));
+		AnimationCategory.AddCustomBuilder(AnimationTrackNamesArrayBuilder);
+
+		DetailBuilder.HideProperty(AnimationTrackNamesHandle);
+	}
+
+
 	/////////////////////////////////////////////////////////////////////////////
 	// Additive settings category
 	/////////////////////////////////////////////////////////////////////////////
@@ -245,6 +281,35 @@ void FAnimSequenceDetails::CustomizeDetails(IDetailLayoutBuilder& DetailBuilder)
 }
 END_SLATE_FUNCTION_BUILD_OPTIMIZATION
 
+void FAnimSequenceDetails::GenerateAnimationTrackNameArrayElementWidget(TSharedRef<IPropertyHandle> PropertyHandle, int32 ArrayIndex, IDetailChildrenBuilder& ChildrenBuilder, IDetailLayoutBuilder* DetailLayout)
+{
+	IDetailPropertyRow& PropRow = ChildrenBuilder.AddProperty(PropertyHandle);
+	PropRow.ShowPropertyButtons(false);
+	PropRow.OverrideResetToDefault(FResetToDefaultOverride::Hide());
+
+	FDetailWidgetRow& WidgetRow = PropRow.CustomWidget(true);
+
+	WidgetRow.NameContent()
+	[
+		PropertyHandle->CreatePropertyNameWidget()
+	];
+
+	WidgetRow.ValueContent()
+	[
+		SNew(SHorizontalBox)
+			+ SHorizontalBox::Slot()
+			.VAlign(VAlign_Center)
+			.HAlign(HAlign_Fill)
+			.Padding(5, 0, 0, 0)
+			.AutoWidth()
+			[
+				SNew(SEditableTextBox)
+					.Text(FText::FromName(AnimationTrackNamesList[ArrayIndex]))
+					.IsReadOnly(true)
+			]
+	];
+}
+
 void FAnimSequenceDetails::CreateOverridenProperty(IDetailLayoutBuilder& DetailBuilder, IDetailCategoryBuilder& AdditiveSettingsCategory, TSharedPtr<IPropertyHandle> PropertyHandle, TAttribute<EVisibility> VisibilityAttribute)
 {
 	DetailBuilder.HideProperty(PropertyHandle);
diff --git a/Engine/Source/Editor/DetailCustomizations/Private/AnimSequenceDetails.h b/Engine/Source/Editor/DetailCustomizations/Private/AnimSequenceDetails.h
--- a/Engine/Source/Editor/DetailCustomizations/Private/AnimSequenceDetails.h
+++ b/Engine/Source/Editor/DetailCustomizations/Private/AnimSequenceDetails.h
@@ -23,6 +23,7 @@
 class FEditorViewportClient;
 class FString;
 class IDetailCategoryBuilder;
+class IDetailChildrenBuilder;
 class IDetailLayoutBuilder;
 class IPropertyHandle;
 class STextBlock;
@@ -66,7 +67,7 @@ private:
 	TSharedPtr<IPropertyHandle> RetargetSourceAssetHandle;
 
 	TSharedPtr<class SSearchableComboBox> RetargetSourceComboBox;
-	TArray< TSharedPtr< FString > >						RetargetSourceComboList;
+	TArray< TSharedPtr< FString > > RetargetSourceComboList;
 
 	TSharedRef<SWidget> MakeRetargetSourceComboWidget( TSharedPtr<FString> InItem );
 	void OnRetargetSourceChanged( TSharedPtr<FString> NewSelection, ESelectInfo::Type SelectInfo  );
@@ -82,6 +83,11 @@ private:
 	void RegisterRetargetSourceChanged();
 	void DelegateRetargetSourceChanged();
 
+	// *** Animation Track Names ***
+	TSharedPtr<IPropertyHandle> AnimationTrackNamesHandle;
+	TArray<FName> AnimationTrackNamesList;
+	void GenerateAnimationTrackNameArrayElementWidget(TSharedRef<IPropertyHandle> PropertyHandle, int32 ArrayIndex, IDetailChildrenBuilder& ChildrenBuilder, IDetailLayoutBuilder* DetailLayout);
+
 	// button handler for Apply Compression
 	// cache all anim sequences that are selected
 	// but I need to know them before compress
diff --git a/Engine/Source/Runtime/Engine/Classes/Animation/AnimSequence.h b/Engine/Source/Runtime/Engine/Classes/Animation/AnimSequence.h
--- a/Engine/Source/Runtime/Engine/Classes/Animation/AnimSequence.h
+++ b/Engine/Source/Runtime/Engine/Classes/Animation/AnimSequence.h
@@ -938,4 +938,5 @@ public:
 	friend class FAnimDataControllerTestBase;
 	friend class UE::Anim::FAnimSequenceCompilingManager;
 	friend struct FAnimNextAnimSequenceKeyframeTask;
+	friend class FAnimSequenceDetails;
 };
