Fix brush style not applying on EUW due to SVG leftover data.

diff --git a/Engine/Source/Editor/DetailCustomizations/Private/SlateBrushCustomization.cpp b/Engine/Source/Editor/DetailCustomizations/Private/SlateBrushCustomization.cpp
--- a/Engine/Source/Editor/DetailCustomizations/Private/SlateBrushCustomization.cpp
+++ b/Engine/Source/Editor/DetailCustomizations/Private/SlateBrushCustomization.cpp
@@ -95,6 +95,7 @@ public:
 		SLATE_ARGUMENT(	TSharedPtr<IPropertyHandle>, ImageSizeProperty )
 		SLATE_ARGUMENT(	TSharedPtr<IPropertyHandle>, MarginProperty )
 		SLATE_ARGUMENT(	TSharedPtr<IPropertyHandle>, ResourceObjectProperty )
+		SLATE_ARGUMENT(	TSharedPtr<IPropertyHandle>, ImageTypeProperty )
 		SLATE_ARGUMENT( FSlateBrush*, SlateBrush )
 	SLATE_END_ARGS()
 
@@ -108,6 +109,7 @@ public:
 		ImageSizeProperty = InArgs._ImageSizeProperty;
 		MarginProperty = InArgs._MarginProperty;
 		ResourceObjectProperty = InArgs._ResourceObjectProperty;
+		ImageTypeProperty = InArgs._ImageTypeProperty;
 
 		FSimpleDelegate OnDrawAsChangedDelegate = FSimpleDelegate::CreateSP( this, &SSlateBrushPreview::OnDrawAsChanged );
 		DrawAsProperty->SetOnPropertyValueChanged( OnDrawAsChangedDelegate );
@@ -983,6 +985,7 @@ private:
 	TSharedPtr<IPropertyHandle> ImageSizeProperty;
 	TSharedPtr<IPropertyHandle> MarginProperty;
 	TSharedPtr<IPropertyHandle> ResourceObjectProperty;
+	TSharedPtr<IPropertyHandle> ImageTypeProperty;
 
 	/** Cached Slate Brush property values */
 	FVector2D CachedTextureSize;
@@ -1166,11 +1169,24 @@ class SBrushResourceObjectBox : public SCompoundWidget
 	
 	SLATE_END_ARGS()
 
-	void Construct(const FArguments& InArgs, IStructCustomizationUtils* StructCustomizationUtils, TSharedPtr<IPropertyHandle> InResourceObjectProperty, TSharedPtr<IPropertyHandle> InImageSizeProperty, TSharedPtr<IPropertyHandle> InDrawAsProperty)
+	struct FPropertyParams
 	{
-		ResourceObjectProperty = InResourceObjectProperty;
-		ImageSizeProperty = InImageSizeProperty;
-		DrawAsProperty = InDrawAsProperty;
+		TSharedPtr<IPropertyHandle> ResourceObjectProperty;
+		TSharedPtr<IPropertyHandle> ResourceNameProperty;
+		TSharedPtr<IPropertyHandle> ImageSizeProperty;
+		TSharedPtr<IPropertyHandle> ImageTypeProperty;
+		TSharedPtr<IPropertyHandle> DrawAsProperty;
+	};
+
+	void Construct(const FArguments& InArgs
+		, IStructCustomizationUtils* StructCustomizationUtils
+		, FPropertyParams InParams)
+	{
+		ResourceObjectProperty = InParams.ResourceObjectProperty;
+		ResourceNameProperty = InParams.ResourceNameProperty;
+		ImageSizeProperty = InParams.ImageSizeProperty;
+		ImageTypeProperty = InParams.ImageTypeProperty;
+		DrawAsProperty = InParams.DrawAsProperty;
 
 		FSimpleDelegate OnBrushResourceChangedDelegate = FSimpleDelegate::CreateSP(this, &SBrushResourceObjectBox::OnBrushResourceChanged);
 		ResourceObjectProperty->SetOnPropertyValueChanged(OnBrushResourceChangedDelegate);
@@ -1182,7 +1198,7 @@ class SBrushResourceObjectBox : public SCompoundWidget
 			.FillHeight(1)
 			[
 				SNew(SObjectPropertyEntryBox)
-				.PropertyHandle(InResourceObjectProperty)
+				.PropertyHandle(InParams.ResourceObjectProperty)
 				.ThumbnailPool(StructCustomizationUtils->GetThumbnailPool())
 			]
 			+ SVerticalBox::Slot()
@@ -1269,6 +1285,7 @@ private:
 			ImageSizeProperty->SetValue(CachedTextureSize);
 
 			// When you assign a resource object, if the current draw type is 'None' we go ahead and update it to 'Image'.
+			// Also update ResourceName to be null (Object name will be used), & set ImageType
 			if (ResourceObject)
 			{
 				TArray<FString> OutPerObjectValues;
@@ -1282,6 +1299,12 @@ private:
 				}
 
 				DrawAsProperty->SetPerObjectValues(NewPerObjectValues);
+
+				ResourceNameProperty->SetValue(NAME_None);
+
+				static_assert(sizeof(decltype(FSlateBrush::ImageType)) == sizeof(uint8));
+				uint8 Value = ESlateBrushImageType::FullColor;
+				ImageTypeProperty->SetValue(Value);
 			}
 		}
 	}
@@ -1350,7 +1373,9 @@ private:
 
 private:
 	TSharedPtr<IPropertyHandle> ResourceObjectProperty;
+	TSharedPtr<IPropertyHandle> ResourceNameProperty;
 	TSharedPtr<IPropertyHandle> ImageSizeProperty;
+	TSharedPtr<IPropertyHandle> ImageTypeProperty;
 	TSharedPtr<IPropertyHandle> DrawAsProperty;
 	TSharedPtr<SBrushResourceError> ResourceError;
 	TSharedPtr<SHyperlink> ChangeDomainLink;
@@ -1398,9 +1423,18 @@ void FSlateBrushStructCustomization::CustomizeChildren( TSharedRef<IPropertyHand
 	TSharedPtr<IPropertyHandle> TintProperty = StructPropertyHandle->GetChildHandle( TEXT("TintColor") );
 	TSharedPtr<IPropertyHandle> OutlineSettingsProperty = StructPropertyHandle->GetChildHandle(TEXT("OutlineSettings"));
 	ResourceObjectProperty = StructPropertyHandle->GetChildHandle( TEXT("ResourceObject") );
+	ResourceNameProperty = StructPropertyHandle->GetChildHandle(TEXT("ResourceName"));
+	ImageTypeProperty = StructPropertyHandle->GetChildHandle(TEXT("ImageType"));
 	
 	FDetailWidgetRow& ResourceObjectRow = StructBuilder.AddProperty(ResourceObjectProperty.ToSharedRef()).CustomWidget();
 
+	SBrushResourceObjectBox::FPropertyParams Params;
+	Params.ResourceObjectProperty = ResourceObjectProperty;
+	Params.ResourceNameProperty = ResourceNameProperty;
+	Params.ImageSizeProperty = ImageSizeProperty;
+	Params.ImageTypeProperty = ImageTypeProperty;
+	Params.DrawAsProperty = DrawAsProperty;
+
 	ResourceObjectRow
 		.NameContent()
 		[
@@ -1410,7 +1444,7 @@ void FSlateBrushStructCustomization::CustomizeChildren( TSharedRef<IPropertyHand
 		.MinDesiredWidth(250.0f)
 		.MaxDesiredWidth(0.0f)
 		[
-			SNew(SBrushResourceObjectBox, &StructCustomizationUtils, ResourceObjectProperty, ImageSizeProperty, DrawAsProperty)
+			SNew(SBrushResourceObjectBox, &StructCustomizationUtils, Params)
 		];
 
 	// Add the image size property with custom reset delegates that also affect the child properties (the components)
@@ -1448,6 +1482,7 @@ void FSlateBrushStructCustomization::CustomizeChildren( TSharedRef<IPropertyHand
 				.ImageSizeProperty(ImageSizeProperty)
 				.MarginProperty(MarginProperty)
 				.ResourceObjectProperty(ResourceObjectProperty)
+				.ImageTypeProperty(ImageTypeProperty)
 				.SlateBrush(Brush);
 
 			IDetailGroup& PreviewGroup = StructBuilder.AddGroup(TEXT("Preview"), FText::GetEmpty());
diff --git a/Engine/Source/Editor/DetailCustomizations/Public/Customizations/SlateBrushCustomization.h b/Engine/Source/Editor/DetailCustomizations/Public/Customizations/SlateBrushCustomization.h
--- a/Engine/Source/Editor/DetailCustomizations/Public/Customizations/SlateBrushCustomization.h
+++ b/Engine/Source/Editor/DetailCustomizations/Public/Customizations/SlateBrushCustomization.h
@@ -52,12 +52,18 @@ private:
 	/** Slate Brush DrawAs property */
 	TSharedPtr<IPropertyHandle> DrawAsProperty;
 
-/** Slate Brush Image Size property */
+	/** Slate Brush Image Size property */
 	TSharedPtr<IPropertyHandle> ImageSizeProperty;
 
-/** Slate Brush Resource Object property */
+	/** Slate Brush Resource Object property */
 	TSharedPtr<IPropertyHandle> ResourceObjectProperty;
 
+	/** Slate Brush Resource Name property */
+	TSharedPtr<IPropertyHandle> ResourceNameProperty;
+
+	/** Slate Brush Image Type property */
+	TSharedPtr<IPropertyHandle> ImageTypeProperty;
+
 	/** Error text to display if the resource object is not valid*/
 	TSharedPtr<SErrorText> ResourceErrorText;
 
diff --git a/Engine/Source/Runtime/SlateCore/Public/Styling/SlateBrush.h b/Engine/Source/Runtime/SlateCore/Public/Styling/SlateBrush.h
--- a/Engine/Source/Runtime/SlateCore/Public/Styling/SlateBrush.h
+++ b/Engine/Source/Runtime/SlateCore/Public/Styling/SlateBrush.h
@@ -263,7 +263,7 @@ public:
 	TEnumAsByte<enum ESlateBrushMirrorType::Type> Mirroring;
 
 	/** The type of image */
-	UPROPERTY()
+	UPROPERTY(EditAnywhere, AdvancedDisplay, Category = Brush)
 	TEnumAsByte<enum ESlateBrushImageType::Type> ImageType;
 
 	/** Size of the resource in Slate Units */
@@ -487,7 +487,7 @@ private:
 
 protected:
 	/** The name of the rendering resource to use */
-	UPROPERTY()
+	UPROPERTY(EditAnywhere, AdvancedDisplay, Category = Brush)
 	FName ResourceName;
 
 	/** 
