Fixed ImageSize properties not updating correctly when changing Slate Brushes to a texture resource that was currently being streamed in

diff --git a/Engine/Source/Editor/DetailCustomizations/Private/SlateBrushCustomization.cpp b/Engine/Source/Editor/DetailCustomizations/Private/SlateBrushCustomization.cpp
--- a/Engine/Source/Editor/DetailCustomizations/Private/SlateBrushCustomization.cpp
+++ b/Engine/Source/Editor/DetailCustomizations/Private/SlateBrushCustomization.cpp
@@ -55,6 +55,7 @@
 #include "UObject/Class.h"
 #include "UObject/Object.h"
 #include "UObject/UnrealType.h"
+#include "TextureCompiler.h"
 #include "Widgets/DeclarativeSyntaxSupport.h"
 #include "Widgets/Images/SImage.h"
 #include "Widgets/Input/SComboBox.h"
@@ -644,17 +645,25 @@ private:
 		UObject* ResourceObject;
 		FPropertyAccess::Result Result = ResourceObjectProperty->GetValue( ResourceObject );
 		if( Result == FPropertyAccess::Success )
-		{				
+		{
+			using ImageSizeType = decltype(FSlateBrush::ImageSize);
+
 			TArray<void*> RawData;
 			ImageSizeProperty->AccessRawData(RawData);
 			if (RawData.Num() > 0 && RawData[0] != NULL)
 			{
-				CachedImageSizeValue = *static_cast<FVector2D*>(RawData[0]);
+				CachedImageSizeValue = *static_cast<ImageSizeType*>(RawData[0]);
 			}
 
 			UTexture2D* BrushTexture = Cast<UTexture2D>(ResourceObject);
 			if( BrushTexture )
 			{
+				if ( BrushTexture->IsDefaultTexture() )
+				{
+					// GetSizeX/Y will return the incorrect value if this texture is being compiled so we need to wait for it here
+					UTexture* const BaseTexture = BrushTexture;
+					FTextureCompilingManager::Get().FinishCompilation( MakeArrayView(&BaseTexture, 1) );
+				}
 				CachedTextureSize = FVector2D( BrushTexture->GetSizeX(), BrushTexture->GetSizeY() );
 			}
 			else if ( ISlateTextureAtlasInterface* AtlasedTextureObject = Cast<ISlateTextureAtlasInterface>(ResourceObject) )
@@ -1272,6 +1281,12 @@ private:
 			UTexture2D* BrushTexture = Cast<UTexture2D>(ResourceObject);
 			if ( BrushTexture )
 			{
+				if ( BrushTexture->IsDefaultTexture() )
+				{
+					UTexture* const BaseTexture = BrushTexture;
+					// GetSizeX/Y will return the incorrect value if this texture is being compiled so we need to wait for it here
+					FTextureCompilingManager::Get().FinishCompilation( MakeArrayView(&BaseTexture, 1) );
+				}
 				CachedTextureSize = ImageSizeType(BrushTexture->GetSizeX(), BrushTexture->GetSizeY());
 			}
 			else if ( ISlateTextureAtlasInterface* AtlasedTextureObject = Cast<ISlateTextureAtlasInterface>(ResourceObject) )
@@ -1614,6 +1629,12 @@ FVector2D FSlateBrushStructCustomization::GetDefaultImageSize() const
 	{
 		if ( UTexture2D* Texture = Cast<UTexture2D>(ResourceObject) )
 		{
+			if ( Texture->IsDefaultTexture() )
+			{
+				// GetSizeX/Y will return the incorrect value if this texture is being compiled so we need to wait for it here
+				UTexture* const BaseTexture = Texture;
+				FTextureCompilingManager::Get().FinishCompilation(MakeArrayView(&BaseTexture, 1));
+			}
 			return FVector2D(Texture->GetSizeX(), Texture->GetSizeY());
 		}
 		else if ( ISlateTextureAtlasInterface* AtlasedTextureObject = Cast<ISlateTextureAtlasInterface>(ResourceObject) )
