UMG: Use an attribute for the font size property handler. The FontSize structure can be enabled/disabled. The property FontSize has to be evaluated dynamically.

diff --git a/Engine/Source/Editor/DetailCustomizations/Private/SlateFontInfoCustomization.cpp b/Engine/Source/Editor/DetailCustomizations/Private/SlateFontInfoCustomization.cpp
--- a/Engine/Source/Editor/DetailCustomizations/Private/SlateFontInfoCustomization.cpp
+++ b/Engine/Source/Editor/DetailCustomizations/Private/SlateFontInfoCustomization.cpp
@@ -193,7 +193,7 @@ void FSlateFontInfoStructCustomization::AddFontSizeProperty(IDetailChildrenBuild
 		.AllowWheel(true)
 		.WheelStep(1.0f)
 		.AllowSpin(FontSizePropertyRef->GetNumPerObjectValues() == 1) //Don't allow spin for multiple value select. Allowing it would result in the widget background not being displayed.
-		.IsEnabled(!FontSizePropertyRef->IsEditConst())
+		.IsEnabled(this, &FSlateFontInfoStructCustomization::IsFontSizeEnabled)
 		.ToolTip(IDocumentation::Get()->CreateToolTip(TAttribute<FText>(this, &FSlateFontInfoStructCustomization::GetFontSizeTooltipText),
 													  nullptr,
 													  TEXT("Shared/Types/FSlateFontInfo"),
@@ -201,6 +201,11 @@ void FSlateFontInfoStructCustomization::AddFontSizeProperty(IDetailChildrenBuild
 	];
 }
 
+bool FSlateFontInfoStructCustomization::IsFontSizeEnabled() const
+{
+	return FontSizeProperty && !FontSizeProperty->IsEditConst();
+}
+
 FText FSlateFontInfoStructCustomization::GetFontSizeTooltipText() const
 {
 	const UUserInterfaceSettings* UISettings = GetDefault<UUserInterfaceSettings>();
diff --git a/Engine/Source/Editor/DetailCustomizations/Public/Customizations/SlateFontInfoCustomization.h b/Engine/Source/Editor/DetailCustomizations/Public/Customizations/SlateFontInfoCustomization.h
--- a/Engine/Source/Editor/DetailCustomizations/Public/Customizations/SlateFontInfoCustomization.h
+++ b/Engine/Source/Editor/DetailCustomizations/Public/Customizations/SlateFontInfoCustomization.h
@@ -38,6 +38,7 @@ public:
 
 protected:
 	void AddFontSizeProperty(IDetailChildrenBuilder& InStructBuilder);
+	bool IsFontSizeEnabled() const;
 
 	/** @return The value or unset if properties with multiple values are viewed */
 	TOptional<float> OnFontSizeGetValue() const;
