Fix Rotator not shown in degrees by default Math Struct customization

diff --git a/Engine/Source/Editor/DetailCustomizations/Private/MathStructCustomizations.cpp b/Engine/Source/Editor/DetailCustomizations/Private/MathStructCustomizations.cpp
--- a/Engine/Source/Editor/DetailCustomizations/Private/MathStructCustomizations.cpp
+++ b/Engine/Source/Editor/DetailCustomizations/Private/MathStructCustomizations.cpp
@@ -25,6 +25,7 @@
 #include "Styling/SlateColor.h"
 #include "Templates/UnrealTemplate.h"
 #include "UObject/EnumProperty.h"
+#include "UObject/NoExportTypes.h"
 #include "UObject/UnrealType.h"
 #include "Widgets/DeclarativeSyntaxSupport.h"
 #include "Widgets/Images/SImage.h"
@@ -161,6 +162,16 @@ void FMathStructCustomization::MakeHeaderRow(TSharedRef<class IPropertyHandle>&
 		ChildHandle->SetInstanceMetaData(TEXT("ClampMin"), StructPropertyHandle->GetMetaData(TEXT("ClampMin")));
 		ChildHandle->SetInstanceMetaData(TEXT("ClampMax"), StructPropertyHandle->GetMetaData(TEXT("ClampMax")));
 
+		// Handle units directly since we can't set metadata on core object types
+		if (FStructProperty* StructProp = CastField<FStructProperty>(StructPropertyHandle->GetProperty()))
+		{
+			if (StructProp->Struct == TBaseStructure<FRotator>::Get())
+			{
+				const static int32 EUnitNamespaceSize = FCString::Strlen(TEXT("EUnit::"));
+				ChildHandle->SetInstanceMetaData(TEXT("Units"), UEnum::GetValueAsString(EUnit::Degrees).RightChop(EUnitNamespaceSize));
+			}
+		}
+
 		const bool bLastChild = SortedChildHandles.Num()-1 == ChildIndex;
 		// Make a widget for each property.  The vector component properties  will be displayed in the header
 
@@ -407,6 +418,21 @@ void FMathStructCustomization::ExtractNumericMetadata(TSharedRef<IPropertyHandle
 	MetadataOut.bSupportDynamicSliderMaxValue = SupportDynamicSliderMaxValueString.Len() > 0 && SupportDynamicSliderMaxValueString.ToBool();
 	MetadataOut.bSupportDynamicSliderMinValue = SupportDynamicSliderMinValueString.Len() > 0 && SupportDynamicSliderMinValueString.ToBool();
 	MetadataOut.bAllowSpinBox = bAllowSpin;
+
+	// By default allow widget to determine default interface
+	MetadataOut.TypeInterface = nullptr;
+
+	if (FStructProperty* StructProp = CastField<FStructProperty>(Property))
+	{
+		if (StructProp->Struct == TBaseStructure<FRotator>::Get())
+		{
+			// The units for degrees does not support floats
+			if constexpr (TIsFloatingPoint<NumericType>::Value && !std::is_same<float, NumericType>::value)
+			{
+				MetadataOut.TypeInterface = MakeShared<TNumericUnitTypeInterface<NumericType>>(EUnit::Degrees);
+			}
+		}
+	}
 }
 
 
@@ -447,7 +473,8 @@ TSharedRef<SWidget> FMathStructCustomization::MakeNumericWidget(
 			.Delta(Metadata.Delta)
 			// LinearDeltaSensitivity must be left unset if not provided, rather than being set to some default
 			.LinearDeltaSensitivity(Metadata.LinearDeltaSensitivity != 0 ? Metadata.LinearDeltaSensitivity : TAttribute<int32>())
-			.ToolTipText(this, &FMathStructCustomization::OnGetValueToolTip<NumericType>, WeakHandlePtr);
+			.ToolTipText(this, &FMathStructCustomization::OnGetValueToolTip<NumericType>, WeakHandlePtr)
+			.TypeInterface(Metadata.TypeInterface);
 }
 
 template <typename NumericType>
diff --git a/Engine/Source/Editor/DetailCustomizations/Public/Customizations/MathStructCustomizations.h b/Engine/Source/Editor/DetailCustomizations/Public/Customizations/MathStructCustomizations.h
--- a/Engine/Source/Editor/DetailCustomizations/Public/Customizations/MathStructCustomizations.h
+++ b/Engine/Source/Editor/DetailCustomizations/Public/Customizations/MathStructCustomizations.h
@@ -8,11 +8,13 @@
 #include "Delegates/Delegate.h"
 #include "IPropertyTypeCustomization.h"
 #include "Internationalization/Text.h"
+#include "Math/UnitConversion.h"
 #include "Misc/Optional.h"
 #include "PropertyHandle.h"
 #include "Styling/SlateTypes.h"
 #include "Templates/SharedPointer.h"
 #include "Types/SlateEnums.h"
+#include "Widgets/Input/NumericTypeInterface.h"
 #include "Widgets/SWidget.h"
 
 class FDetailWidgetRow;
@@ -69,6 +71,7 @@ public:
 		TOptional<NumericType> MaxValue;
 		TOptional<NumericType> SliderMinValue;
 		TOptional<NumericType> SliderMaxValue;
+		TSharedPtr<INumericTypeInterface<NumericType>> TypeInterface;
 		NumericType SliderExponent;
 		NumericType Delta;
 		int32 LinearDeltaSensitivity;
