Fix ProjectSettings Mac

diff --git a/Engine/Source/Editor/DetailCustomizations/Private/MacTargetSettingsDetails.cpp b/Engine/Source/Editor/DetailCustomizations/Private/MacTargetSettingsDetails.cpp
--- a/Engine/Source/Editor/DetailCustomizations/Private/MacTargetSettingsDetails.cpp
+++ b/Engine/Source/Editor/DetailCustomizations/Private/MacTargetSettingsDetails.cpp
@@ -358,7 +358,7 @@ TSharedRef<SWidget> FMacTargetSettingsDetails::OnGetShaderVersionContent()
 {
 	FMenuBuilder MenuBuilder(true, NULL);
 	
-	UEnum* Enum = FindObjectChecked<UEnum>(nullptr, TEXT("/Script/MacTargetPlatform.EMacMetalShaderStandard"), true);
+	UEnum* Enum = FindObjectChecked<UEnum>(nullptr, TEXT("/Script/MacTargetPlatformSettings.EMacMetalShaderStandard"), true);
 	
 	for (int32 i = 0; i < Enum->GetMaxEnumValue(); i++)
 	{
@@ -381,7 +381,7 @@ FText FMacTargetSettingsDetails::GetShaderVersionDesc() const
     int32 EnumValue;
     ShaderVersionPropertyHandle->GetValue(EnumValue);
 	
-	UEnum* Enum = FindObjectChecked<UEnum>(nullptr, TEXT("/Script/MacTargetPlatform.EMacMetalShaderStandard"), true);
+	UEnum* Enum = FindObjectChecked<UEnum>(nullptr, TEXT("/Script/MacTargetPlatformSettings.EMacMetalShaderStandard"), true);
 	
 	if (EnumValue < Enum->GetMaxEnumValue() && Enum->IsValidEnumValue(EnumValue))
 	{
