Use MakeShared instead of MakeSharable to reduce the number of allocations

diff --git a/Engine/Source/Editor/DetailCustomizations/Private/ActorDetails.cpp b/Engine/Source/Editor/DetailCustomizations/Private/ActorDetails.cpp
--- a/Engine/Source/Editor/DetailCustomizations/Private/ActorDetails.cpp
+++ b/Engine/Source/Editor/DetailCustomizations/Private/ActorDetails.cpp
@@ -476,8 +476,8 @@ void FActorDetails::AddTransformCategory( IDetailLayoutBuilder& DetailBuilder )
 	{
 		return;
 	}
-	
-	TSharedRef<FComponentTransformDetails> TransformDetails = MakeShareable( new FComponentTransformDetails( DetailBuilder.GetSelectedObjects(), SelectedActorInfo, DetailBuilder ) );
+
+	TSharedRef<FComponentTransformDetails> TransformDetails = MakeShared<FComponentTransformDetails>(DetailBuilder.GetSelectedObjects(), SelectedActorInfo, DetailBuilder);
 
 	IDetailCategoryBuilder& TransformCategory = DetailBuilder.EditCategory( "TransformCommon", LOCTEXT("TransformCommonCategory", "Transform"), ECategoryPriority::Transform );
 
