Use ITargetPlatformSettings  instead of ITargetPlatform for Windows, Mac and Linux TargetSettingsDetails

diff --git a/Engine/Source/Developer/SharedSettingsWidgets/Private/ShaderFormatsPropertyDetails.cpp b/Engine/Source/Developer/SharedSettingsWidgets/Private/ShaderFormatsPropertyDetails.cpp
--- a/Engine/Source/Developer/SharedSettingsWidgets/Private/ShaderFormatsPropertyDetails.cpp
+++ b/Engine/Source/Developer/SharedSettingsWidgets/Private/ShaderFormatsPropertyDetails.cpp
@@ -36,17 +36,17 @@ void FShaderFormatsPropertyDetails::SetOnUpdateShaderWarning(FSimpleDelegate con
 }
 
 void FShaderFormatsPropertyDetails::CreateTargetShaderFormatsPropertyView(
-	ITargetPlatform* TargetPlatform,
+	ITargetPlatformSettings* TargetPlatformSettings,
 	GetFriendlyNameFromRHINameFnc FriendlyNameFnc,
 	FilterShaderPlatformFnc* FilterShaderPlatformFunc,
 	ECategoryPriority::Type Priority)
 {
-	check(TargetPlatform);
+	check(TargetPlatformSettings);
 	DetailBuilder->HideProperty(ShaderFormatsPropertyHandle);
 	
 	// List of supported RHI's and selected targets
 	TArray<FName> ShaderFormats;
-	TargetPlatform->GetAllPossibleShaderFormats(ShaderFormats);
+	TargetPlatformSettings->GetAllPossibleShaderFormats(ShaderFormats);
 	
 	IDetailCategoryBuilder& TargetedRHICategoryBuilder = DetailBuilder->EditCategory(Title, FText::GetEmpty(), Priority);
 	
diff --git a/Engine/Source/Developer/SharedSettingsWidgets/Public/ShaderFormatsPropertyDetails.h b/Engine/Source/Developer/SharedSettingsWidgets/Public/ShaderFormatsPropertyDetails.h
--- a/Engine/Source/Developer/SharedSettingsWidgets/Public/ShaderFormatsPropertyDetails.h
+++ b/Engine/Source/Developer/SharedSettingsWidgets/Public/ShaderFormatsPropertyDetails.h
@@ -8,6 +8,7 @@
 class IDetailLayoutBuilder;
 class IPropertyHandle;
 class ITargetPlatform;
+class ITargetPlatformSettings;
 
 enum class ECheckBoxState : uint8;
 
@@ -36,6 +37,9 @@ public:
 	void SetOnUpdateShaderWarning(const FSimpleDelegate& Delegate);
 	
 	/** Create the UI to select which windows shader formats we are targeting */
+	void CreateTargetShaderFormatsPropertyView(ITargetPlatformSettings* TargetPlatform, GetFriendlyNameFromRHINameFnc* FriendlyNameFnc, FilterShaderPlatformFnc* FilterShaderPlatformFunc = nullptr, ECategoryPriority::Type InPriority = ECategoryPriority::Default);
+	
+	UE_DEPRECATED(5.5, "CreateTargetShaderFormatsPropertyView now takes a ITargetPlatformSettings* as first argument instead of ITargetPlatform*. Please change your callback function as this is a breaking change.")
 	void CreateTargetShaderFormatsPropertyView(ITargetPlatform* TargetPlatform, GetFriendlyNameFromRHINameFnc* FriendlyNameFnc, FilterShaderPlatformFnc* FilterShaderPlatformFunc = nullptr, ECategoryPriority::Type InPriority = ECategoryPriority::Default);
 
 	UE_DEPRECATED(5.1, "CreateTargetShaderFormatsPropertyView now gets RHI names via FName instead of FString. Please change your callback function as this is a breaking change.")
diff --git a/Engine/Source/Editor/DetailCustomizations/Private/LinuxTargetSettingsDetails.cpp b/Engine/Source/Editor/DetailCustomizations/Private/LinuxTargetSettingsDetails.cpp
--- a/Engine/Source/Editor/DetailCustomizations/Private/LinuxTargetSettingsDetails.cpp
+++ b/Engine/Source/Editor/DetailCustomizations/Private/LinuxTargetSettingsDetails.cpp
@@ -27,6 +27,7 @@
 #include "DetailCategoryBuilder.h"
 #include "Interfaces/ITargetPlatform.h"
 #include "Interfaces/ITargetPlatformModule.h"
+#include "Interfaces/ITargetPlatformSettingsModule.h"
 #include "SExternalImageReference.h"
 #include "RHIDefinitions.h"
 #include "RHIShaderFormatDefinitions.inl"
@@ -129,11 +130,9 @@ static FString GetLinuxSplashFilename(ELinuxImageScope::Type Scope, bool bIsEdit
 /* Helper function used to generate filenames for icons */
 static FString GetLinuxIconFilename(ELinuxImageScope::Type Scope)
 {
-	const FString& PlatformName = FModuleManager::GetModuleChecked<ITargetPlatformModule>("LinuxTargetPlatform").GetTargetPlatforms()[0]->PlatformName();
-
 	if (Scope == ELinuxImageScope::Engine)
 	{
-		FString Filename = FPaths::EngineDir() / FString(TEXT("Source/Runtime/Launch/Resources")) / PlatformName / FString("UnrealEngine.png");
+		FString Filename = FPaths::EngineDir() / FString(TEXT("Source/Runtime/Launch/Resources/Linux/UnrealEngine.png"));
 		return FPaths::ConvertRelativePathToFull(Filename);
 	}
 	else
@@ -141,7 +140,7 @@ static FString GetLinuxIconFilename(ELinuxImageScope::Type Scope)
 		FString Filename = FPaths::ProjectDir() / TEXT("Build/Linux/Application.png");
 		if(!FPaths::FileExists(Filename))
 		{
-			FString LegacyFilename = FPaths::GameSourceDir() / FString(FApp::GetProjectName()) / FString(TEXT("Resources")) / PlatformName / FString(FApp::GetProjectName()) + TEXT(".icns");
+			FString LegacyFilename = FPaths::GameSourceDir() / FString(FApp::GetProjectName()) / FString(TEXT("Resources/Linux")) / FString(FApp::GetProjectName()) + TEXT(".icns");
 			if(FPaths::FileExists(LegacyFilename))
 			{
 				Filename = LegacyFilename;
@@ -154,9 +153,9 @@ static FString GetLinuxIconFilename(ELinuxImageScope::Type Scope)
 void FLinuxTargetSettingsDetails::CustomizeDetails( IDetailLayoutBuilder& DetailBuilder )
 {
 	// Setup the supported/targeted RHI property view
-	ITargetPlatform* TargetPlatform = FModuleManager::GetModuleChecked<ITargetPlatformModule>("LinuxTargetPlatform").GetTargetPlatforms()[0];
+	ITargetPlatformSettings* TargetPlatformSettings = FModuleManager::GetModuleChecked<ITargetPlatformSettingsModule>("LinuxTargetPlatformSettings").GetTargetPlatformSettings()[0];
 	TargetShaderFormatsDetails = MakeShareable(new FShaderFormatsPropertyDetails(&DetailBuilder));
-	TargetShaderFormatsDetails->CreateTargetShaderFormatsPropertyView(TargetPlatform, &GetFriendlyNameFromLinuxShaderFormat);
+	TargetShaderFormatsDetails->CreateTargetShaderFormatsPropertyView(TargetPlatformSettings, &GetFriendlyNameFromLinuxShaderFormat);
 
 	// Next add the splash image customization
 	const FText EditorSplashDesc(LOCTEXT("EditorSplashLabel", "Editor Splash"));
diff --git a/Engine/Source/Editor/DetailCustomizations/Private/MacTargetSettingsDetails.cpp b/Engine/Source/Editor/DetailCustomizations/Private/MacTargetSettingsDetails.cpp
--- a/Engine/Source/Editor/DetailCustomizations/Private/MacTargetSettingsDetails.cpp
+++ b/Engine/Source/Editor/DetailCustomizations/Private/MacTargetSettingsDetails.cpp
@@ -17,6 +17,7 @@
 #include "DetailCategoryBuilder.h"
 #include "Interfaces/ITargetPlatform.h"
 #include "Interfaces/ITargetPlatformModule.h"
+#include "Interfaces/ITargetPlatformSettingsModule.h"
 #include "SExternalImageReference.h"
 #include "Widgets/Input/SNumericDropDown.h"
 #include "Dialogs/Dialogs.h"
@@ -81,11 +82,9 @@ static FString GetSplashFilename(EMacImageScope::Type Scope, bool bIsEditorSplas
 /* Helper function used to generate filenames for icons */
 static FString GetIconFilename(EMacImageScope::Type Scope)
 {
-	const FString& PlatformName = FModuleManager::GetModuleChecked<ITargetPlatformModule>("MacTargetPlatform").GetTargetPlatforms()[0]->PlatformName();
-
 	if (Scope == EMacImageScope::Engine)
 	{
-		FString Filename = FPaths::EngineDir() / FString(TEXT("Source/Runtime/Launch/Resources")) / PlatformName / FString("UnrealEngine.icns");
+		FString Filename = FPaths::EngineDir() / FString(TEXT("Source/Runtime/Launch/Resources/Mac/UnrealEngine.icns"));
 		return FPaths::ConvertRelativePathToFull(Filename);
 	}
 	else
@@ -93,7 +92,7 @@ static FString GetIconFilename(EMacImageScope::Type Scope)
 		FString Filename = FPaths::ProjectDir() / TEXT("Build/Mac/Application.icns");
 		if(!FPaths::FileExists(Filename))
 		{
-			FString LegacyFilename = FPaths::GameSourceDir() / FString(FApp::GetProjectName()) / FString(TEXT("Resources")) / PlatformName / FString(FApp::GetProjectName()) + TEXT(".icns");
+			FString LegacyFilename = FPaths::GameSourceDir() / FString(FApp::GetProjectName()) / FString(TEXT("Resources/Mac")) / FString(FApp::GetProjectName()) + TEXT(".icns");
 			if(FPaths::FileExists(LegacyFilename))
 			{
 				Filename = LegacyFilename;
@@ -172,12 +171,11 @@ void FMacTargetSettingsDetails::CustomizeDetails( IDetailLayoutBuilder& DetailBu
 {
 	FSimpleDelegate OnUpdateShaderStandardWarning = FSimpleDelegate::CreateSP(this, &FMacTargetSettingsDetails::UpdateShaderStandardWarning);
 	
-	ITargetPlatform* TargetPlatform = FModuleManager::GetModuleChecked<ITargetPlatformModule>("MacTargetPlatform").GetTargetPlatforms()[0];
-    
+	ITargetPlatformSettings* TargetPlatformSettings = FModuleManager::GetModuleChecked<ITargetPlatformSettingsModule>("MacTargetPlatformSettings").GetTargetPlatformSettings()[0];
 	// Setup the supported/targeted RHI property view
 	TargetShaderFormatsDetails = MakeShareable(new FShaderFormatsPropertyDetails(&DetailBuilder, TEXT("TargetedRHIs"), TEXT("Targeted RHIs")));
 	TargetShaderFormatsDetails->SetOnUpdateShaderWarning(OnUpdateShaderStandardWarning);
-	TargetShaderFormatsDetails->CreateTargetShaderFormatsPropertyView(TargetPlatform, &GetFriendlyNameFromRHINameMac);
+	TargetShaderFormatsDetails->CreateTargetShaderFormatsPropertyView(TargetPlatformSettings, &GetFriendlyNameFromRHINameMac);
 	
 	// Setup the shader version property view
     // Handle max. shader version a little specially.
diff --git a/Engine/Source/Editor/DetailCustomizations/Private/WindowsTargetSettingsDetails.cpp b/Engine/Source/Editor/DetailCustomizations/Private/WindowsTargetSettingsDetails.cpp
--- a/Engine/Source/Editor/DetailCustomizations/Private/WindowsTargetSettingsDetails.cpp
+++ b/Engine/Source/Editor/DetailCustomizations/Private/WindowsTargetSettingsDetails.cpp
@@ -27,6 +27,7 @@
 #include "DetailCategoryBuilder.h"
 #include "Interfaces/ITargetPlatform.h"
 #include "Interfaces/ITargetPlatformModule.h"
+#include "Interfaces/ITargetPlatformSettingsModule.h"
 #include "SExternalImageReference.h"
 #include "UnrealEngine.h"
 #include "RHIShaderFormatDefinitions.inl"
@@ -175,8 +176,6 @@ static FString GetWindowsSplashFilename(EWindowsImageScope::Type Scope, bool bIs
 /* Helper function used to generate filenames for icons */
 static FString GetWindowsIconFilename(EWindowsImageScope::Type Scope)
 {
-	const FString& PlatformName = FModuleManager::GetModuleChecked<ITargetPlatformModule>("WindowsTargetPlatform").GetTargetPlatforms()[0]->PlatformName();
-
 	if (Scope == EWindowsImageScope::Engine)
 	{
 		FString Filename = FPaths::EngineDir() / FString(TEXT("Build/Windows/Resources/Default.ico"));
@@ -187,7 +186,7 @@ static FString GetWindowsIconFilename(EWindowsImageScope::Type Scope)
 		FString Filename = FPaths::ProjectDir() / TEXT("Build/Windows/Application.ico");
 		if(!FPaths::FileExists(Filename))
 		{
-			FString LegacyFilename = FPaths::GameSourceDir() / FString(FApp::GetProjectName()) / FString(TEXT("Resources")) / PlatformName / FString(FApp::GetProjectName()) + TEXT(".ico");
+			FString LegacyFilename = FPaths::GameSourceDir() / FString(FApp::GetProjectName()) / FString(TEXT("Resources/Windows")) / FString(FApp::GetProjectName()) + TEXT(".ico");
 			if(FPaths::FileExists(LegacyFilename))
 			{
 				Filename = LegacyFilename;
@@ -200,16 +199,16 @@ static FString GetWindowsIconFilename(EWindowsImageScope::Type Scope)
 void FWindowsTargetSettingsDetails::CustomizeDetails( IDetailLayoutBuilder& DetailBuilder )
 {
 	// Setup the supported/targeted RHI property view
-	ITargetPlatform* TargetPlatform = FModuleManager::GetModuleChecked<ITargetPlatformModule>("WindowsTargetPlatform").GetTargetPlatforms()[0];
+	ITargetPlatformSettings* TargetPlatformSettings = FModuleManager::GetModuleChecked<ITargetPlatformSettingsModule>("WindowsTargetPlatformSettings").GetTargetPlatformSettings()[0];
 
 	D3D12TargetShaderFormatsDetails = MakeShareable(new FShaderFormatsPropertyDetails(&DetailBuilder, TEXT("D3D12TargetedShaderFormats"), TEXT("D3D12 Targeted Shader Formats")));
-	D3D12TargetShaderFormatsDetails->CreateTargetShaderFormatsPropertyView(TargetPlatform, &GetFriendlyNameForWindowsShaderPlatformCheckbox, &FilterShaderPlatform_D3D12, ECategoryPriority::Important);
+	D3D12TargetShaderFormatsDetails->CreateTargetShaderFormatsPropertyView(TargetPlatformSettings, &GetFriendlyNameForWindowsShaderPlatformCheckbox, &FilterShaderPlatform_D3D12, ECategoryPriority::Important);
 
 	D3D11TargetShaderFormatsDetails = MakeShareable(new FShaderFormatsPropertyDetails(&DetailBuilder, TEXT("D3D11TargetedShaderFormats"), TEXT("D3D11 Targeted Shader Formats")));
-	D3D11TargetShaderFormatsDetails->CreateTargetShaderFormatsPropertyView(TargetPlatform, &GetFriendlyNameForWindowsShaderPlatformCheckbox, &FilterShaderPlatform_D3D11, ECategoryPriority::Important);
+	D3D11TargetShaderFormatsDetails->CreateTargetShaderFormatsPropertyView(TargetPlatformSettings, &GetFriendlyNameForWindowsShaderPlatformCheckbox, &FilterShaderPlatform_D3D11, ECategoryPriority::Important);
 
 	VulkanTargetShaderFormatsDetails = MakeShareable(new FShaderFormatsPropertyDetails(&DetailBuilder, TEXT("VulkanTargetedShaderFormats"), TEXT("Vulkan Targeted Shader Formats")));
-	VulkanTargetShaderFormatsDetails->CreateTargetShaderFormatsPropertyView(TargetPlatform, &GetFriendlyNameForWindowsShaderPlatformCheckbox, &FilterShaderPlatform_Vulkan, ECategoryPriority::Important);
+	VulkanTargetShaderFormatsDetails->CreateTargetShaderFormatsPropertyView(TargetPlatformSettings, &GetFriendlyNameForWindowsShaderPlatformCheckbox, &FilterShaderPlatform_Vulkan, ECategoryPriority::Important);
 
 	// Next add the splash image customization
 	const FText EditorSplashDesc(LOCTEXT("EditorSplashLabel", "Editor Splash"));
