Fix new crash exposed by an old UEditorEngine::ReplaceActors bug

The original bug would leave CS Constructed Components in this case RootComponent in a garbaged state after undo. (This old bug is being fixed in another CL)

#rnx

diff --git a/Engine/Source/Editor/DetailCustomizations/Private/ComponentTransformDetails.cpp b/Engine/Source/Editor/DetailCustomizations/Private/ComponentTransformDetails.cpp
--- a/Engine/Source/Editor/DetailCustomizations/Private/ComponentTransformDetails.cpp
+++ b/Engine/Source/Editor/DetailCustomizations/Private/ComponentTransformDetails.cpp
@@ -654,7 +654,10 @@ TSharedPtr<IPropertyHandle> FComponentTransformDetails::FindOrCreatePropertyHand
 		CachedHandlesObjects.Append(SceneComponents);
 	}
 
-	PropertyHandles.Add(PropertyName, PropertyHandle);
+	if (PropertyHandle && PropertyHandle->IsValidHandle())
+	{
+		PropertyHandles.Add(PropertyName, PropertyHandle);
+	}
 	return PropertyHandle;
 }
 
