Add AllowedClasses and DisallowedClasses to SoftClassPathCustomization.

diff --git a/Engine/Source/Editor/DetailCustomizations/Private/SoftClassPathCustomization.cpp b/Engine/Source/Editor/DetailCustomizations/Private/SoftClassPathCustomization.cpp
--- a/Engine/Source/Editor/DetailCustomizations/Private/SoftClassPathCustomization.cpp
+++ b/Engine/Source/Editor/DetailCustomizations/Private/SoftClassPathCustomization.cpp
@@ -22,7 +22,9 @@ void FSoftClassPathCustomization::CustomizeHeader(TSharedRef<IPropertyHandle> In
 
 	const FString& MetaClassName = PropertyHandle->GetMetaData("MetaClass");
 	const FString& RequiredInterfaceName = PropertyHandle->GetMetaData("RequiredInterface"); // This was the old name, switch to MustImplement to synchronize with class property
-	const FString& MustImplementName = PropertyHandle->GetMetaData("MustImplement"); 
+	const FString& MustImplementName = PropertyHandle->GetMetaData("MustImplement");
+	TArray<const UClass*> AllowedClasses = PropertyCustomizationHelpers::GetClassesFromMetadataString(PropertyHandle->GetMetaData("AllowedClasses"));
+	TArray<const UClass*> DisallowedClasses = PropertyCustomizationHelpers::GetClassesFromMetadataString(PropertyHandle->GetMetaData("DisallowedClasses"));
 	const bool bAllowAbstract = PropertyHandle->HasMetaData("AllowAbstract");
 	const bool bIsBlueprintBaseOnly = PropertyHandle->HasMetaData("IsBlueprintBaseOnly") || PropertyHandle->HasMetaData("BlueprintBaseOnly");
 	const bool bAllowNone = !(PropertyHandle->GetMetaDataProperty()->PropertyFlags & CPF_NoClear);
@@ -50,6 +52,8 @@ void FSoftClassPathCustomization::CustomizeHeader(TSharedRef<IPropertyHandle> In
 		SNew(SClassPropertyEntryBox)
 			.MetaClass(MetaClass)
 			.RequiredInterface(RequiredInterface)
+			.AllowedClasses(MoveTemp(AllowedClasses))
+			.DisallowedClasses(MoveTemp(DisallowedClasses))
 			.AllowAbstract(bAllowAbstract)
 			.IsBlueprintBaseOnly(bIsBlueprintBaseOnly)
 			.AllowNone(bAllowNone)
diff --git a/Engine/Source/Editor/PropertyEditor/Private/PropertyCustomizationHelpers.cpp b/Engine/Source/Editor/PropertyEditor/Private/PropertyCustomizationHelpers.cpp
--- a/Engine/Source/Editor/PropertyEditor/Private/PropertyCustomizationHelpers.cpp
+++ b/Engine/Source/Editor/PropertyEditor/Private/PropertyCustomizationHelpers.cpp
@@ -658,6 +658,8 @@ void SClassPropertyEntryBox::Construct(const FArguments& InArgs)
 			SAssignNew(PropertyEditorClass, SPropertyEditorClass)
 				.MetaClass(InArgs._MetaClass)
 				.RequiredInterface(InArgs._RequiredInterface)
+				.AllowedClasses(InArgs._AllowedClasses)
+				.DisallowedClasses(InArgs._DisallowedClasses)
 				.AllowAbstract(InArgs._AllowAbstract)
 				.IsBlueprintBaseOnly(InArgs._IsBlueprintBaseOnly)
 				.AllowNone(InArgs._AllowNone)
diff --git a/Engine/Source/Editor/PropertyEditor/Private/UserInterface/PropertyEditor/SPropertyEditorClass.cpp b/Engine/Source/Editor/PropertyEditor/Private/UserInterface/PropertyEditor/SPropertyEditorClass.cpp
--- a/Engine/Source/Editor/PropertyEditor/Private/UserInterface/PropertyEditor/SPropertyEditorClass.cpp
+++ b/Engine/Source/Editor/PropertyEditor/Private/UserInterface/PropertyEditor/SPropertyEditorClass.cpp
@@ -213,8 +213,8 @@ void SPropertyEditorClass::Construct(const FArguments& InArgs, const TSharedPtr<
 		bShowViewOptions = InArgs._ShowViewOptions;
 		bShowTree = InArgs._ShowTree;
 		bShowDisplayNames = InArgs._ShowDisplayNames;
-		AllowedClassFilters.Empty();
-		DisallowedClassFilters.Empty();
+		AllowedClassFilters = InArgs._AllowedClasses;
+		DisallowedClassFilters = InArgs._DisallowedClasses;
 		SelectedClass = InArgs._SelectedClass;
 		OnSetClass = InArgs._OnSetClass;
 	}
diff --git a/Engine/Source/Editor/PropertyEditor/Private/UserInterface/PropertyEditor/SPropertyEditorClass.h b/Engine/Source/Editor/PropertyEditor/Private/UserInterface/PropertyEditor/SPropertyEditorClass.h
--- a/Engine/Source/Editor/PropertyEditor/Private/UserInterface/PropertyEditor/SPropertyEditorClass.h
+++ b/Engine/Source/Editor/PropertyEditor/Private/UserInterface/PropertyEditor/SPropertyEditorClass.h
@@ -42,6 +42,10 @@ public:
 		SLATE_ARGUMENT(const UClass*, MetaClass)
 		/** An interface that the selected class must implement (optional) */
 		SLATE_ARGUMENT(const UClass*, RequiredInterface)
+		/** Allowed class that the selected class must be a child-of. (optional) */
+		SLATE_ARGUMENT(TArray<const UClass*>, AllowedClasses)
+		/** Classes that the selected class cannot be a child-of. (optional) */
+		SLATE_ARGUMENT(TArray<const UClass*>, DisallowedClasses)
 		/** Whether or not abstract classes are allowed (optional) */
 		SLATE_ARGUMENT(bool, AllowAbstract)
 		/** Should only base blueprints be displayed? (optional) */
diff --git a/Engine/Source/Editor/PropertyEditor/Public/PropertyCustomizationHelpers.h b/Engine/Source/Editor/PropertyEditor/Public/PropertyCustomizationHelpers.h
--- a/Engine/Source/Editor/PropertyEditor/Public/PropertyCustomizationHelpers.h
+++ b/Engine/Source/Editor/PropertyEditor/Public/PropertyCustomizationHelpers.h
@@ -289,6 +289,10 @@ public:
 		SLATE_ARGUMENT(const UClass*, MetaClass)
 		/** An interface that the selected class must implement (optional) */
 		SLATE_ARGUMENT(const UClass*, RequiredInterface)
+		/** Allowed class that the selected class must be a child-of. (optional) */
+		SLATE_ARGUMENT(TArray<const UClass*>, AllowedClasses)
+		/** Classes that the selected class cannot be a child-of. (optional) */
+		SLATE_ARGUMENT(TArray<const UClass*>, DisallowedClasses)
 		/** Whether or not abstract classes are allowed (optional) */
 		SLATE_ARGUMENT(bool, AllowAbstract)
 		/** Should only base blueprints be displayed? (optional) */
