Fix crash when baking materials on corrupted actor with StaticMeshComponent
* Shielded the utility methods from bad input
* Disabled the Bake Material button in case the operation isn't possible.

diff --git a/Engine/Source/Developer/MeshMergeUtilities/Private/MeshMergeUtilities.cpp b/Engine/Source/Developer/MeshMergeUtilities/Private/MeshMergeUtilities.cpp
--- a/Engine/Source/Developer/MeshMergeUtilities/Private/MeshMergeUtilities.cpp
+++ b/Engine/Source/Developer/MeshMergeUtilities/Private/MeshMergeUtilities.cpp
@@ -459,6 +459,11 @@ void FMeshMergeUtilities::BakeMaterialsForComponent(TArray<TWeakObjectPtr<UObjec
 
 void FMeshMergeUtilities::BakeMaterialsForComponent(USkeletalMeshComponent* SkeletalMeshComponent) const
 {
+	if (!SkeletalMeshComponent || !SkeletalMeshComponent->GetSkeletalMeshAsset())
+	{
+		return;
+	}
+
 	// Retrieve settings object
 	UMaterialOptions* MaterialOptions = DuplicateObject(GetMutableDefault<UMaterialOptions>(), GetTransientPackage());
 	UAssetBakeOptions* AssetOptions = GetMutableDefault<UAssetBakeOptions>();
@@ -481,6 +486,11 @@ void FMeshMergeUtilities::BakeMaterialsForComponent(USkeletalMeshComponent* Skel
 
 void FMeshMergeUtilities::BakeMaterialsForComponent(UStaticMeshComponent* StaticMeshComponent) const
 {
+	if (!StaticMeshComponent || !StaticMeshComponent->GetStaticMesh())
+	{
+		return;
+	}
+
 	// Retrieve settings object
 	UMaterialOptions* MaterialOptions = DuplicateObject(GetMutableDefault<UMaterialOptions>(), GetTransientPackage());
 	UAssetBakeOptions* AssetOptions = GetMutableDefault<UAssetBakeOptions>();
@@ -502,6 +512,11 @@ void FMeshMergeUtilities::BakeMaterialsForComponent(UStaticMeshComponent* Static
 
 void FMeshMergeUtilities::BakeMaterialsForMesh(UStaticMesh* StaticMesh) const
 {
+	if (!StaticMesh)
+	{
+		return;
+	}
+
 	// Retrieve settings object
 	UMaterialOptions* MaterialOptions = DuplicateObject(GetMutableDefault<UMaterialOptions>(), GetTransientPackage());
 	UAssetBakeOptions* AssetOptions = GetMutableDefault<UAssetBakeOptions>();
diff --git a/Engine/Source/Editor/DetailCustomizations/Private/StaticMeshActorDetails.cpp b/Engine/Source/Editor/DetailCustomizations/Private/StaticMeshActorDetails.cpp
--- a/Engine/Source/Editor/DetailCustomizations/Private/StaticMeshActorDetails.cpp
+++ b/Engine/Source/Editor/DetailCustomizations/Private/StaticMeshActorDetails.cpp
@@ -3,6 +3,7 @@
 #include "StaticMeshActorDetails.h"
 
 #include "Containers/Array.h"
+#include "Components/StaticMeshComponent.h"
 #include "DetailCategoryBuilder.h"
 #include "DetailLayoutBuilder.h"
 #include "DetailWidgetRow.h"
@@ -100,6 +101,32 @@ void FStaticMeshActorDetails::CustomizeDetails( IDetailLayoutBuilder& DetailBuil
 		]
 	];
 
+
+	// Retrieve all currently selected static mesh components
+	auto RetrieveSelectedStaticMeshComponents = [](const TArray<TWeakObjectPtr<UObject>>& SelectedObjects)
+	{		
+		TArray<UStaticMeshComponent*> StaticMeshComponents;
+		for (TWeakObjectPtr<UObject> WeakObject : SelectedObjects)
+		{
+			if (WeakObject.IsValid())
+			{
+				UObject* Object = WeakObject.Get();
+				if (AStaticMeshActor* Actor = Cast<AStaticMeshActor>(Object))
+				{
+					if (UStaticMeshComponent* SMC = Actor->GetStaticMeshComponent())
+					{
+						if (SMC->GetStaticMesh())
+						{
+							StaticMeshComponents.Add(SMC);
+						}
+					}
+				}
+			}
+		}
+
+		return StaticMeshComponents;
+	};
+
 	// This allows baking out the materials for the given instance data	
 	TArray<TWeakObjectPtr<UObject>> Objects;
 	DetailBuilder.GetObjectsBeingCustomized(Objects);
@@ -113,24 +140,10 @@ void FStaticMeshActorDetails::CustomizeDetails( IDetailLayoutBuilder& DetailBuil
 		[
 			SNew(SButton)
 			.Text(LOCTEXT("BakeLabel", "Bake Materials"))
-			.OnClicked_Lambda([Objects]() -> FReply
+			.OnClicked_Lambda([RetrieveSelectedStaticMeshComponents, Objects]() -> FReply
 			{
 				const IMeshMergeUtilities& MeshMergeUtilities = FModuleManager::Get().LoadModuleChecked<IMeshMergeModule>("MeshMergeUtilities").GetUtilities();
-
-				// Retrieve all currently selected static mesh components
-				TArray<UStaticMeshComponent*> StaticMeshComponents;
-				for (TWeakObjectPtr<UObject> WeakObject : Objects)
-				{
-					if (WeakObject.IsValid())
-					{
-						UObject* Object = WeakObject.Get();
-						if (AStaticMeshActor* Actor = Cast<AStaticMeshActor>(Object))
-						{
-							StaticMeshComponents.Add(Actor->GetStaticMeshComponent());
-						}
-					}
-				}
-
+				TArray<UStaticMeshComponent*> StaticMeshComponents = RetrieveSelectedStaticMeshComponents(Objects);
 				for (UStaticMeshComponent* Component : StaticMeshComponents)
 				{
 					MeshMergeUtilities.BakeMaterialsForComponent(Component);
@@ -138,6 +151,10 @@ void FStaticMeshActorDetails::CustomizeDetails( IDetailLayoutBuilder& DetailBuil
 				
 				return FReply::Handled();
 			})
+			.IsEnabled_Lambda([RetrieveSelectedStaticMeshComponents, Objects]() -> bool
+			{
+				return !RetrieveSelectedStaticMeshComponents(Objects).IsEmpty();
+			})
 		]
 	];		
 	
