[Nanite-Skinning] Cleaned up NaniteSettings filtering

diff --git a/Engine/Source/Editor/DetailCustomizations/Private/DetailCustomizations.cpp b/Engine/Source/Editor/DetailCustomizations/Private/DetailCustomizations.cpp
--- a/Engine/Source/Editor/DetailCustomizations/Private/DetailCustomizations.cpp
+++ b/Engine/Source/Editor/DetailCustomizations/Private/DetailCustomizations.cpp
@@ -130,6 +130,7 @@
 #include "SceneCaptureDetails.h"
 #include "SceneComponentDetails.h"
 #include "SkeletalControlNodeDetails.h"
+#include "SkeletalMeshDetails.h"
 #include "SkeletalMeshComponentDetails.h"
 #include "SkeletalMeshLODSettingsDetails.h"
 #include "SkeletalMeshReductionSettingsDetails.h"
@@ -463,7 +464,7 @@ void FDetailCustomizationsModule::RegisterObjectCustomizations()
 	RegisterCustomClassLayout("MaterialExpressionTextureBase", FOnGetDetailCustomizationInstance::CreateStatic(&FMaterialExpressionTextureBaseDetails::MakeInstance));
 	RegisterCustomClassLayout("MaterialInstanceDynamic", FOnGetDetailCustomizationInstance::CreateStatic(&FMaterialInstanceDynamicDetails::MakeInstance));
 	RegisterCustomClassLayout("SkeletalMeshLODSettings", FOnGetDetailCustomizationInstance::CreateStatic(&FSkeletalMeshLODSettingsDetails::MakeInstance));
-
+	RegisterCustomClassLayout("SkeletalMesh", FOnGetDetailCustomizationInstance::CreateStatic(&FSkeletalMeshDetails::MakeInstance));
 	RegisterCustomClassLayout("Skeleton", FOnGetDetailCustomizationInstance::CreateStatic(&FSkeletonDetails::MakeInstance));
 
 	RegisterCustomClassLayout("MotionControllerComponent", FOnGetDetailCustomizationInstance::CreateStatic(&FMotionControllerDetails::MakeInstance));
diff --git a/Engine/Source/Editor/DetailCustomizations/Private/SkeletalMeshDetails.cpp b/Engine/Source/Editor/DetailCustomizations/Private/SkeletalMeshDetails.cpp
new file mode 100644
--- /dev/null
+++ b/Engine/Source/Editor/DetailCustomizations/Private/SkeletalMeshDetails.cpp
@@ -0,0 +1,33 @@
+// Copyright Epic Games, Inc. All Rights Reserved.
+
+#include "SkeletalMeshDetails.h"
+
+#include "Containers/Array.h"
+#include "DetailLayoutBuilder.h"
+#include "Engine/SkeletalMesh.h"
+#include "HAL/PlatformMath.h"
+#include "IMeshReductionManagerModule.h"
+#include "Misc/AssertionMacros.h"
+#include "Modules/ModuleManager.h"
+#include "PropertyEditorModule.h"
+#include "PropertyHandle.h"
+#include "UObject/NameTypes.h"
+
+#define LOCTEXT_NAMESPACE "SkeletalMeshDetails"
+
+void FSkeletalMeshDetails::CustomizeDetails(IDetailLayoutBuilder& LayoutBuilder)
+{
+	static const auto AllowSkinnedMeshes = IConsoleManager::Get().FindTConsoleVariableDataInt(TEXT("r.Nanite.AllowSkinnedMeshes"));
+	static const bool bAllowSkinnedMeshes = (AllowSkinnedMeshes && AllowSkinnedMeshes->GetValueOnAnyThread() != 0);
+	if (!bAllowSkinnedMeshes)
+	{
+		TSharedRef<IPropertyHandle> SettingsHandle = LayoutBuilder.GetProperty(FName("NaniteSettings"));
+		LayoutBuilder.HideProperty(SettingsHandle);
+	}
+}
+
+TSharedRef<IDetailCustomization> FSkeletalMeshDetails::MakeInstance()
+{
+	return MakeShareable(new FSkeletalMeshDetails);
+}
+#undef LOCTEXT_NAMESPACE
\ No newline at end of file
diff --git a/Engine/Source/Editor/DetailCustomizations/Private/SkeletalMeshDetails.h b/Engine/Source/Editor/DetailCustomizations/Private/SkeletalMeshDetails.h
new file mode 100644
--- /dev/null
+++ b/Engine/Source/Editor/DetailCustomizations/Private/SkeletalMeshDetails.h
@@ -0,0 +1,19 @@
+// Copyright Epic Games, Inc. All Rights Reserved.
+
+#pragma once
+
+#include "IDetailCustomization.h"
+#include "Templates/SharedPointer.h"
+
+class IDetailLayoutBuilder;
+
+class FSkeletalMeshDetails : public IDetailCustomization
+{
+public:
+	/** Begin IDetailCustomization interface */
+	virtual void CustomizeDetails(IDetailLayoutBuilder& LayoutBuilder) override;
+	/** End IDetailCustomization interface */
+
+	/** Makes a new instance of this detail layout class for a specific detail view requesting it */
+	static TSharedRef<IDetailCustomization> MakeInstance();
+};
diff --git a/Engine/Source/Runtime/Engine/Classes/Engine/SkeletalMesh.h b/Engine/Source/Runtime/Engine/Classes/Engine/SkeletalMesh.h
--- a/Engine/Source/Runtime/Engine/Classes/Engine/SkeletalMesh.h
+++ b/Engine/Source/Runtime/Engine/Classes/Engine/SkeletalMesh.h
@@ -930,10 +930,8 @@ public:
 	TArray<struct FBoneMirrorInfo> SkelMirrorTable;
 #endif
 
-#if 0 // TODO: Nanite-Skinning
 	/** Settings related to building Nanite data. */
-	UPROPERTY(EditAnywhere, Category = NaniteSettings)
-#endif
+	UPROPERTY(EditAnywhere, Category = Mesh)
 	FMeshNaniteSettings NaniteSettings;
 
 #if WITH_EDITORONLY_DATA
