Improve FNavAgentSelectorCustomization to let the user check or uncheck all supported agents in one click

diff --git a/Engine/Source/Editor/DetailCustomizations/Private/NavAgentSelectorCustomization.cpp b/Engine/Source/Editor/DetailCustomizations/Private/NavAgentSelectorCustomization.cpp
--- a/Engine/Source/Editor/DetailCustomizations/Private/NavAgentSelectorCustomization.cpp
+++ b/Engine/Source/Editor/DetailCustomizations/Private/NavAgentSelectorCustomization.cpp
@@ -23,10 +23,23 @@
 #include "UObject/UObjectGlobals.h"
 #include "UObject/UnrealType.h"
 #include "Widgets/DeclarativeSyntaxSupport.h"
+#include "Widgets/Input/SCheckBox.h"
 #include "Widgets/Text/STextBlock.h"
 
 #define LOCTEXT_NAMESPACE "FNavAgentSelectorCustomization"
 
+namespace NavAgentSelectorCustoPrivate
+{
+	const UNavigationSystemV1* GetNavSysCDO()
+	{
+		return (*GEngine->NavigationSystemClass != nullptr && GEngine->NavigationSystemClass->IsChildOf(UNavigationSystemV1::StaticClass()))
+			? GetDefault<UNavigationSystemV1>(GEngine->NavigationSystemClass)
+			: GetDefault<UNavigationSystemV1>();
+	}
+
+	static const FString AgentPrefix = "bSupportsAgent";
+}
+
 TSharedRef<IPropertyTypeCustomization> FNavAgentSelectorCustomization::MakeInstance()
 {
 	return MakeShareable(new FNavAgentSelectorCustomization);
@@ -45,9 +58,27 @@ void FNavAgentSelectorCustomization::CustomizeHeader(TSharedRef<IPropertyHandle>
 	.MaxDesiredWidth(400.0f)
 	.VAlign(VAlign_Center)
 	[
-		SNew(STextBlock)
-		.Text(this, &FNavAgentSelectorCustomization::GetSupportedDesc)
-		.Font(StructCustomizationUtils.GetRegularFont())
+		SNew(SHorizontalBox)
+
+		+ SHorizontalBox::Slot()
+		.VAlign(VAlign_Center)
+		.AutoWidth()
+		[
+			SNew(SCheckBox)
+			.OnCheckStateChanged(this, &FNavAgentSelectorCustomization::OnHeaderCheckStateChanged)
+			.IsChecked(this, &FNavAgentSelectorCustomization::IsHeaderChecked)
+		]
+
+		+ SHorizontalBox::Slot()
+		.VAlign(VAlign_Center)
+		.FillContentWidth(/*grow*/0.f, /*shrink*/1.f)
+		.Padding(FMargin(3, 0, 0, 0))
+		[
+			SNew(STextBlock)
+			.Text(this, &FNavAgentSelectorCustomization::GetSupportedDesc)
+			.OverflowPolicy(ETextOverflowPolicy::Ellipsis)
+			.Font(StructCustomizationUtils.GetRegularFont())
+		]
 	];
 }
 
@@ -56,11 +87,7 @@ void FNavAgentSelectorCustomization::CustomizeChildren(TSharedRef<class IPropert
 	uint32 NumChildren = 0;
 	StructPropertyHandle->GetNumChildren(NumChildren);
 
-	FString AgentPrefix("bSupportsAgent");
-	const UNavigationSystemV1* NavSysCDO = (*GEngine->NavigationSystemClass != nullptr && GEngine->NavigationSystemClass->IsChildOf(UNavigationSystemV1::StaticClass()))
-		? GetDefault<UNavigationSystemV1>(GEngine->NavigationSystemClass)
-		: GetDefault<UNavigationSystemV1>();
-
+	const UNavigationSystemV1* NavSysCDO = NavAgentSelectorCustoPrivate::GetNavSysCDO();
 	if (NavSysCDO == nullptr)
 	{
 		return;
@@ -71,12 +98,12 @@ void FNavAgentSelectorCustomization::CustomizeChildren(TSharedRef<class IPropert
 	for (uint32 Idx = 0; Idx < NumChildren; Idx++)
 	{
 		TSharedPtr<IPropertyHandle> PropHandle = StructPropertyHandle->GetChildHandle(Idx);
-		if (PropHandle->GetProperty() && PropHandle->GetProperty()->GetName().StartsWith(AgentPrefix))
+		if (PropHandle->GetProperty() && PropHandle->GetProperty()->GetName().StartsWith(NavAgentSelectorCustoPrivate::AgentPrefix))
 		{
 			PropHandle->SetOnPropertyValueChanged(FSimpleDelegate::CreateSP(this, &FNavAgentSelectorCustomization::OnAgentStateChanged));
 
 			int32 AgentIdx = -1;
-			TTypeFromString<int32>::FromString(AgentIdx, *(PropHandle->GetProperty()->GetName().Mid(AgentPrefix.Len()) ));
+			TTypeFromString<int32>::FromString(AgentIdx, *(PropHandle->GetProperty()->GetName().Mid(NavAgentSelectorCustoPrivate::AgentPrefix.Len()) ));
 
 			if (AgentIdx >= 0 && AgentIdx < NumAgents)
 			{
@@ -103,14 +130,106 @@ void FNavAgentSelectorCustomization::CustomizeChildren(TSharedRef<class IPropert
 
 void FNavAgentSelectorCustomization::OnAgentStateChanged()
 {
-	const UNavigationSystemV1* NavSysCDO = (*GEngine->NavigationSystemClass != nullptr && GEngine->NavigationSystemClass->IsChildOf(UNavigationSystemV1::StaticClass()))
-		? GetDefault<UNavigationSystemV1>(GEngine->NavigationSystemClass)
-		: GetDefault<UNavigationSystemV1>();
-
+	const UNavigationSystemV1* NavSysCDO = NavAgentSelectorCustoPrivate::GetNavSysCDO();
 	if (NavSysCDO == nullptr)
 	{
 		return;
 	}
+	int32 NumAgents, NumSupported, FirstSupportedIdx;
+	if (!ComputeSupportedAgentCount(NavSysCDO, NumAgents, NumSupported, FirstSupportedIdx))
+	{
+		return;
+	}
+
+	if (NumSupported == NumAgents)
+	{
+		SupportedDesc = LOCTEXT("AllAgents", "All");
+	}
+	else if (NumSupported == 0)
+	{
+		SupportedDesc = LOCTEXT("NoAgents", "None");
+	}
+	else if (NumSupported == 1)
+	{
+		SupportedDesc = FText::FromName(NavSysCDO->GetSupportedAgents()[FirstSupportedIdx].Name);
+	}
+	else
+	{
+		SupportedDesc = FText::Format(FText::FromString("{0}, ..."), FText::FromName(NavSysCDO->GetSupportedAgents()[FirstSupportedIdx].Name));
+	}
+}
+
+FText FNavAgentSelectorCustomization::GetSupportedDesc() const
+{
+	return SupportedDesc;
+}
+
+void FNavAgentSelectorCustomization::OnHeaderCheckStateChanged(ECheckBoxState InNewState)
+{
+	bool bNewValue = false;
+	switch (InNewState)
+	{
+		case ECheckBoxState::Checked:
+			bNewValue = true;
+			break;
+		case ECheckBoxState::Unchecked:
+			bNewValue = false;
+			break;
+		case ECheckBoxState::Undetermined:
+		default:
+			return;
+	}
+
+	const UNavigationSystemV1* NavSysCDO = NavAgentSelectorCustoPrivate::GetNavSysCDO();
+	if (!NavSysCDO)
+	{
+		return;
+	}
+
+	uint32 NumChildren = 0;
+	StructHandle->GetNumChildren(NumChildren);
+
+	for (uint32 Idx = 0; Idx < NumChildren; Idx++)
+	{
+		TSharedPtr<IPropertyHandle> PropHandle = StructHandle->GetChildHandle(Idx);
+		if (PropHandle->GetProperty() && PropHandle->GetProperty()->GetName().StartsWith(NavAgentSelectorCustoPrivate::AgentPrefix))
+		{
+			PropHandle->SetValue(bNewValue);
+		}
+	}
+
+	OnAgentStateChanged();
+}
+
+ECheckBoxState FNavAgentSelectorCustomization::IsHeaderChecked() const
+{
+	const UNavigationSystemV1* NavSysCDO = NavAgentSelectorCustoPrivate::GetNavSysCDO();
+	int32 NumAgents, NumSupported, FirstSupportedIdx;
+	if (!ComputeSupportedAgentCount(NavSysCDO, NumAgents, NumSupported, FirstSupportedIdx))
+	{
+		return ECheckBoxState::Unchecked;
+	}
+
+	if (NumSupported == NumAgents)
+	{
+		return ECheckBoxState::Checked;
+	}
+	else if (NumSupported == 0)
+	{
+		return ECheckBoxState::Unchecked;
+	}
+	else
+	{
+		return ECheckBoxState::Undetermined;
+	}
+}
+
+bool FNavAgentSelectorCustomization::ComputeSupportedAgentCount(const UNavigationSystemV1* NavSysCDO, int32& OutNumAgents, int32& OutNumSupported, int32& OutFirstSupportedIdx) const
+{
+	if (NavSysCDO == nullptr)
+	{
+		return false;
+	}
 
 	const int32 NumAgents = FMath::Min(NavSysCDO->GetSupportedAgents().Num(), 16);
 
@@ -118,25 +237,24 @@ void FNavAgentSelectorCustomization::OnAgentStateChanged()
 	StructHandle->GetNumChildren(NumChildren);
 
 	int32 NumSupported = 0;
-	int32 FirstSupportedIdx = -1;
+	int32 FirstSupportedIdx = INDEX_NONE;
 
-	FString AgentPrefix("bSupportsAgent");
 	for (uint32 Idx = 0; Idx < NumChildren; Idx++)
 	{
 		TSharedPtr<IPropertyHandle> PropHandle = StructHandle->GetChildHandle(Idx);
-		if (PropHandle->GetProperty() && PropHandle->GetProperty()->GetName().StartsWith(AgentPrefix))
+		if (PropHandle->GetProperty() && PropHandle->GetProperty()->GetName().StartsWith(NavAgentSelectorCustoPrivate::AgentPrefix))
 		{
 			bool bSupportsAgent = false;
 			FPropertyAccess::Result Result = PropHandle->GetValue(bSupportsAgent);
 			if (Result == FPropertyAccess::Success && bSupportsAgent)
 			{
 				int32 AgentIdx = -1;
-				TTypeFromString<int32>::FromString(AgentIdx, *(PropHandle->GetProperty()->GetName().Mid(AgentPrefix.Len())));
+				TTypeFromString<int32>::FromString(AgentIdx, *(PropHandle->GetProperty()->GetName().Mid(NavAgentSelectorCustoPrivate::AgentPrefix.Len())));
 
 				if (AgentIdx >= 0 && AgentIdx < NumAgents)
 				{
 					NumSupported++;
-					if (FirstSupportedIdx < 0)
+					if (FirstSupportedIdx == INDEX_NONE)
 					{
 						FirstSupportedIdx = AgentIdx;
 					}
@@ -145,27 +263,10 @@ void FNavAgentSelectorCustomization::OnAgentStateChanged()
 		}
 	}
 
-	if (NumSupported == NumAgents)
-	{
-		SupportedDesc = LOCTEXT("AllAgents", "all");
-	}
-	else if (NumSupported == 0)
-	{
-		SupportedDesc = LOCTEXT("NoAgents", "none");
-	}
-	else if (NumSupported == 1)
-	{
-		SupportedDesc = FText::FromName(NavSysCDO->GetSupportedAgents()[FirstSupportedIdx].Name);
-	}
-	else
-	{
-		SupportedDesc = FText::Format(FText::FromString("{0}, ..."), FText::FromName(NavSysCDO->GetSupportedAgents()[FirstSupportedIdx].Name));
-	}
-}
-
-FText FNavAgentSelectorCustomization::GetSupportedDesc() const
-{
-	return SupportedDesc;
+	OutNumAgents = NumAgents;
+	OutNumSupported = NumSupported;
+	OutFirstSupportedIdx = FirstSupportedIdx;
+	return true;
 }
 
 #undef LOCTEXT_NAMESPACE
diff --git a/Engine/Source/Editor/DetailCustomizations/Private/NavAgentSelectorCustomization.h b/Engine/Source/Editor/DetailCustomizations/Private/NavAgentSelectorCustomization.h
--- a/Engine/Source/Editor/DetailCustomizations/Private/NavAgentSelectorCustomization.h
+++ b/Engine/Source/Editor/DetailCustomizations/Private/NavAgentSelectorCustomization.h
@@ -7,6 +7,8 @@
 #include "Templates/SharedPointer.h"
 
 class IPropertyHandle;
+class UNavigationSystemV1;
+enum class ECheckBoxState : uint8;
 
 class FNavAgentSelectorCustomization : public IPropertyTypeCustomization
 {
@@ -17,8 +19,6 @@ public:
 	virtual void CustomizeHeader( TSharedRef<IPropertyHandle> StructPropertyHandle, class FDetailWidgetRow& HeaderRow, IPropertyTypeCustomizationUtils& StructCustomizationUtils ) override;
 	virtual void CustomizeChildren( TSharedRef<class IPropertyHandle> StructPropertyHandle, class IDetailChildrenBuilder& StructBuilder, IPropertyTypeCustomizationUtils& StructCustomizationUtils ) override;
 
-public:
-
 	/**
 	 * Creates a new instance.
 	 *
@@ -26,8 +26,12 @@ public:
 	 */
 	static TSharedRef<IPropertyTypeCustomization> MakeInstance();
 
+protected:
 	void OnAgentStateChanged();
 	FText GetSupportedDesc() const;
+	void OnHeaderCheckStateChanged(ECheckBoxState InNewState);
+	ECheckBoxState IsHeaderChecked() const;
+	bool ComputeSupportedAgentCount(const UNavigationSystemV1* NavSysCDO, int32& OutNumAgents, int32& OutNumSupported, int32& OutFirstSupportedIdx) const;
 
 	TSharedPtr<IPropertyHandle> StructHandle;
 	FText SupportedDesc;
