Drastically reduced severe hang (minutes-> seconds) when clicking on a landscape streaming proxy with "Show hidden properties while playing" enabled after a material compiles by using the proper RequestForceRefresh function, instead of a using a generic deferred action request in FCustomPrimitiveDataCustomization, because there can be a whole lot of those when displaying all properties of a landscape proxy (because of the several primitive components, the inter-component pointers, etc. : with only 4 landscape components per proxy, when all hidden properties are showed, that accounts for 300+ such customizations). The advantage is that all those customizations, instead of each forcing a refresh of the same DetailsView on the next tick (ATM, such a refresh, with all hidden properties displayed is 7s!) will only do it once.

diff --git a/Engine/Source/Editor/DetailCustomizations/Private/CustomPrimitiveDataCustomization.cpp b/Engine/Source/Editor/DetailCustomizations/Private/CustomPrimitiveDataCustomization.cpp
--- a/Engine/Source/Editor/DetailCustomizations/Private/CustomPrimitiveDataCustomization.cpp
+++ b/Engine/Source/Editor/DetailCustomizations/Private/CustomPrimitiveDataCustomization.cpp
@@ -626,23 +626,11 @@ void FCustomPrimitiveDataCustomization::PopulateParameterData(UPrimitiveComponen
 }
 
 void FCustomPrimitiveDataCustomization::RequestRefresh()
-{
-	if (!bDeferringRefresh && PropertyUtils.IsValid())
-	{
-		bDeferringRefresh = true;
-
-		PropertyUtils->EnqueueDeferredAction(FSimpleDelegate::CreateSP(this, &FCustomPrimitiveDataCustomization::OnDeferredRefresh));
-	}
-}
-	
-void FCustomPrimitiveDataCustomization::OnDeferredRefresh()
 {
 	if (PropertyUtils.IsValid())
 	{
-		PropertyUtils->ForceRefresh();
+		PropertyUtils->RequestForceRefresh();
 	}
-
-	bDeferringRefresh = false;
 }
 
 void FCustomPrimitiveDataCustomization::OnElementsModified(const FPropertyAccess::Result OldAccessResult, const uint32 OldNumElements)
diff --git a/Engine/Source/Editor/DetailCustomizations/Private/CustomPrimitiveDataCustomization.h b/Engine/Source/Editor/DetailCustomizations/Private/CustomPrimitiveDataCustomization.h
--- a/Engine/Source/Editor/DetailCustomizations/Private/CustomPrimitiveDataCustomization.h
+++ b/Engine/Source/Editor/DetailCustomizations/Private/CustomPrimitiveDataCustomization.h
@@ -58,6 +58,7 @@ private:
 		uint8 IndexOffset;
 	};
 
+	// TODO [jonathan.bard] Remove on non-release branch
 	bool bDeferringRefresh = false;
 
 	TSharedPtr<IPropertyUtilities> PropertyUtils;
@@ -98,7 +99,6 @@ private:
 	void PopulateParameterData(UPrimitiveComponent* PrimitiveComponent, int32& MaxPrimitiveDataIndex);
 
 	void RequestRefresh();
-	void OnDeferredRefresh();
 	void OnElementsModified(const enum FPropertyAccess::Result OldAccessResult, const uint32 OldNumElements);
 	void OnObjectPropertyChanged(UObject* Object, FPropertyChangedEvent& PropertyChangedEvent);
 	void OnMaterialCompiled(UMaterialInterface* Material);
