Control Rig: Fix crash with curve / eval curve node

#rnx

diff --git a/Engine/Source/Editor/DetailCustomizations/Private/CurveStructCustomization.cpp b/Engine/Source/Editor/DetailCustomizations/Private/CurveStructCustomization.cpp
--- a/Engine/Source/Editor/DetailCustomizations/Private/CurveStructCustomization.cpp
+++ b/Engine/Source/Editor/DetailCustomizations/Private/CurveStructCustomization.cpp
@@ -175,6 +175,13 @@ void FCurveStructCustomization::CustomizeHeader( TSharedRef<IPropertyHandle> InS
 
 void FCurveStructCustomization::CustomizeChildren( TSharedRef<IPropertyHandle> InStructPropertyHandle, IDetailChildrenBuilder& StructBuilder, IPropertyTypeCustomizationUtils& StructCustomizationUtils )
 {
+	// this customization can be called without an instance now in the details
+	// panel, so we need to make sure to initialize members here as well.
+	if(!StructPropertyHandle.IsValid())
+	{
+		StructPropertyHandle = InStructPropertyHandle;
+	}
+	
 	uint32 NumChildren = 0;
 	StructPropertyHandle->GetNumChildren(NumChildren);
 
