Emit warning for code using the old #include "PerPlatformProperties.h" instead of #include "UObject/PerPlatformProperties.h"
Update all existing references

diff --git a/Engine/Plugins/Animation/ACLPlugin/Source/ACLPlugin/Classes/AnimBoneCompressionCodec_ACL.h b/Engine/Plugins/Animation/ACLPlugin/Source/ACLPlugin/Classes/AnimBoneCompressionCodec_ACL.h
--- a/Engine/Plugins/Animation/ACLPlugin/Source/ACLPlugin/Classes/AnimBoneCompressionCodec_ACL.h
+++ b/Engine/Plugins/Animation/ACLPlugin/Source/ACLPlugin/Classes/AnimBoneCompressionCodec_ACL.h
@@ -4,7 +4,7 @@
 // Copyright 2018 Nicholas Frechette. All Rights Reserved.
 
 #include "CoreMinimal.h"
-#include "PerPlatformProperties.h"
+#include "UObject/PerPlatformProperties.h"
 #include "UObject/ObjectMacros.h"
 #include "AnimBoneCompressionCodec_ACLBase.h"
 #include "AnimBoneCompressionCodec_ACL.generated.h"
diff --git a/Engine/Plugins/Animation/ACLPlugin/Source/ACLPlugin/Classes/AnimationCompressionLibraryDatabase.h b/Engine/Plugins/Animation/ACLPlugin/Source/ACLPlugin/Classes/AnimationCompressionLibraryDatabase.h
--- a/Engine/Plugins/Animation/ACLPlugin/Source/ACLPlugin/Classes/AnimationCompressionLibraryDatabase.h
+++ b/Engine/Plugins/Animation/ACLPlugin/Source/ACLPlugin/Classes/AnimationCompressionLibraryDatabase.h
@@ -12,7 +12,7 @@ THIRD_PARTY_INCLUDES_END
 
 #include "CoreMinimal.h"
 #include "Containers/Ticker.h"
-#include "PerPlatformProperties.h"
+#include "UObject/PerPlatformProperties.h"
 #include "Engine/LatentActionManager.h"
 #include "Serialization/BulkData.h"
 #include "UObject/ObjectMacros.h"
diff --git a/Engine/Plugins/Animation/ACLPlugin/Source/ACLPlugin/Private/ACLImpl.cpp b/Engine/Plugins/Animation/ACLPlugin/Source/ACLPlugin/Private/ACLImpl.cpp
--- a/Engine/Plugins/Animation/ACLPlugin/Source/ACLPlugin/Private/ACLImpl.cpp
+++ b/Engine/Plugins/Animation/ACLPlugin/Source/ACLPlugin/Private/ACLImpl.cpp
@@ -6,7 +6,7 @@
 #if WITH_EDITOR
 #include "AnimationCompression.h"
 #include "AnimationUtils.h"
-#include "PerPlatformProperties.h"
+#include "UObject/PerPlatformProperties.h"
 #include "Animation/AnimCompressionTypes.h"
 #include "Interfaces/ITargetPlatform.h"
 #include "PlatformInfo.h"
diff --git a/Engine/Plugins/Developer/AnimationSharing/Source/AnimationSharing/Public/AnimationSharingTypes.h b/Engine/Plugins/Developer/AnimationSharing/Source/AnimationSharing/Public/AnimationSharingTypes.h
--- a/Engine/Plugins/Developer/AnimationSharing/Source/AnimationSharing/Public/AnimationSharingTypes.h
+++ b/Engine/Plugins/Developer/AnimationSharing/Source/AnimationSharing/Public/AnimationSharingTypes.h
@@ -2,7 +2,7 @@
 
 #pragma once
 
-#include "PerPlatformProperties.h"
+#include "UObject/PerPlatformProperties.h"
 #include "Templates/SubclassOf.h"
 
 #include "AnimationSharingTypes.generated.h"
diff --git a/Engine/Plugins/Runtime/HairStrands/Source/HairStrandsCore/Public/GroomAsset.h b/Engine/Plugins/Runtime/HairStrands/Source/HairStrandsCore/Public/GroomAsset.h
--- a/Engine/Plugins/Runtime/HairStrands/Source/HairStrandsCore/Public/GroomAsset.h
+++ b/Engine/Plugins/Runtime/HairStrands/Source/HairStrandsCore/Public/GroomAsset.h
@@ -19,7 +19,7 @@
 #include "HairStrandsInterface.h"
 #include "Engine/SkeletalMesh.h"
 #include "Interfaces/Interface_AssetUserData.h"
-#include "PerPlatformProperties.h"
+#include "UObject/PerPlatformProperties.h"
 #include "UObject/StrongObjectPtr.h"
 
 #include "GroomAsset.generated.h"
diff --git a/Engine/Plugins/Runtime/HairStrands/Source/HairStrandsCore/Public/GroomAssetInterpolation.h b/Engine/Plugins/Runtime/HairStrands/Source/HairStrandsCore/Public/GroomAssetInterpolation.h
--- a/Engine/Plugins/Runtime/HairStrands/Source/HairStrandsCore/Public/GroomAssetInterpolation.h
+++ b/Engine/Plugins/Runtime/HairStrands/Source/HairStrandsCore/Public/GroomAssetInterpolation.h
@@ -4,7 +4,7 @@
 
 #include "UObject/Object.h"
 #include "UObject/ObjectMacros.h"
-#include "PerPlatformProperties.h"
+#include "UObject/PerPlatformProperties.h"
 #include "GroomAssetInterpolation.generated.h"
 
 UENUM()
diff --git a/Engine/Plugins/Runtime/Metasound/Source/MetasoundEngine/Private/MetasoundBuilderSubsystem.cpp b/Engine/Plugins/Runtime/Metasound/Source/MetasoundEngine/Private/MetasoundBuilderSubsystem.cpp
--- a/Engine/Plugins/Runtime/Metasound/Source/MetasoundEngine/Private/MetasoundBuilderSubsystem.cpp
+++ b/Engine/Plugins/Runtime/Metasound/Source/MetasoundEngine/Private/MetasoundBuilderSubsystem.cpp
@@ -27,7 +27,7 @@
 #include "MetasoundUObjectRegistry.h"
 #include "MetasoundVertex.h"
 #include "NodeTemplates/MetasoundFrontendNodeTemplateInput.h"
-#include "PerPlatformProperties.h"
+#include "UObject/PerPlatformProperties.h"
 
 #include UE_INLINE_GENERATED_CPP_BY_NAME(MetasoundBuilderSubsystem)
 
diff --git a/Engine/Plugins/Runtime/Metasound/Source/MetasoundEngine/Public/MetasoundSettings.h b/Engine/Plugins/Runtime/Metasound/Source/MetasoundEngine/Public/MetasoundSettings.h
--- a/Engine/Plugins/Runtime/Metasound/Source/MetasoundEngine/Public/MetasoundSettings.h
+++ b/Engine/Plugins/Runtime/Metasound/Source/MetasoundEngine/Public/MetasoundSettings.h
@@ -7,7 +7,7 @@
 #include "UObject/NoExportTypes.h"
 #include "UObject/Object.h"
 #include "UObject/ObjectMacros.h"
-#include "PerPlatformProperties.h"
+#include "UObject/PerPlatformProperties.h"
 
 #include "MetasoundSettings.generated.h"
 
diff --git a/Engine/Source/Editor/DetailCustomizations/Private/PerPlatformPropertyCustomization.cpp b/Engine/Source/Editor/DetailCustomizations/Private/PerPlatformPropertyCustomization.cpp
--- a/Engine/Source/Editor/DetailCustomizations/Private/PerPlatformPropertyCustomization.cpp
+++ b/Engine/Source/Editor/DetailCustomizations/Private/PerPlatformPropertyCustomization.cpp
@@ -15,7 +15,7 @@
 #include "Internationalization/Internationalization.h"
 #include "Misc/AssertionMacros.h"
 #include "Misc/DataDrivenPlatformInfoRegistry.h"
-#include "PerPlatformProperties.h"
+#include "UObject/PerPlatformProperties.h"
 #include "PropertyEditorModule.h"
 #include "PropertyHandle.h"
 #include "SPerPlatformPropertiesWidget.h"
diff --git a/Engine/Source/Editor/DetailCustomizations/Public/PerPlatformPropertyCustomization.h b/Engine/Source/Editor/DetailCustomizations/Public/PerPlatformPropertyCustomization.h
--- a/Engine/Source/Editor/DetailCustomizations/Public/PerPlatformPropertyCustomization.h
+++ b/Engine/Source/Editor/DetailCustomizations/Public/PerPlatformPropertyCustomization.h
@@ -9,7 +9,7 @@
 #include "IPropertyTypeCustomization.h"
 #include "Internationalization/Text.h"
 #include "Misc/Attribute.h"
-#include "PerPlatformProperties.h"
+#include "UObject/PerPlatformProperties.h"
 #include "Templates/SharedPointer.h"
 #include "UObject/NameTypes.h"
 #include "UnrealClient.h"
diff --git a/Engine/Source/Editor/UnrealEd/Classes/Editor/PropertyEditorTestObject.h b/Engine/Source/Editor/UnrealEd/Classes/Editor/PropertyEditorTestObject.h
--- a/Engine/Source/Editor/UnrealEd/Classes/Editor/PropertyEditorTestObject.h
+++ b/Engine/Source/Editor/UnrealEd/Classes/Editor/PropertyEditorTestObject.h
@@ -14,7 +14,7 @@
 #include "Engine/EngineTypes.h"
 
 #include "Curves/RichCurve.h"
-#include "PerPlatformProperties.h"
+#include "UObject/PerPlatformProperties.h"
 
 #include "PropertyEditorTestObject.generated.h"
 
diff --git a/Engine/Source/Editor/UnrealEd/Private/StaticMeshEdit.cpp b/Engine/Source/Editor/UnrealEd/Private/StaticMeshEdit.cpp
--- a/Engine/Source/Editor/UnrealEd/Private/StaticMeshEdit.cpp
+++ b/Engine/Source/Editor/UnrealEd/Private/StaticMeshEdit.cpp
@@ -29,7 +29,7 @@
 #include "MaterialDomain.h"
 #include "Materials/MaterialInterface.h"
 #include "Materials/Material.h"
-#include "PerPlatformProperties.h"
+#include "UObject/PerPlatformProperties.h"
 #include "Interfaces/ITargetPlatformManagerModule.h"
 #include "Interfaces/ITargetPlatform.h"
 #include "Settings/EditorExperimentalSettings.h"
diff --git a/Engine/Source/Editor/UnrealEd/Public/ImportUtils/StaticMeshImportUtils.h b/Engine/Source/Editor/UnrealEd/Public/ImportUtils/StaticMeshImportUtils.h
--- a/Engine/Source/Editor/UnrealEd/Public/ImportUtils/StaticMeshImportUtils.h
+++ b/Engine/Source/Editor/UnrealEd/Public/ImportUtils/StaticMeshImportUtils.h
@@ -10,7 +10,7 @@
 #include "Engine/EngineTypes.h"
 #include "Engine/StaticMesh.h"
 #include "MeshReductionSettings.h"
-#include "PerPlatformProperties.h"
+#include "UObject/PerPlatformProperties.h"
 #include "PerQualityLevelProperties.h"
 
 struct FMeshDescription;
diff --git a/Engine/Source/Runtime/Engine/Classes/Animation/AnimSequence.h b/Engine/Source/Runtime/Engine/Classes/Animation/AnimSequence.h
--- a/Engine/Source/Runtime/Engine/Classes/Animation/AnimSequence.h
+++ b/Engine/Source/Runtime/Engine/Classes/Animation/AnimSequence.h
@@ -20,7 +20,7 @@
 #include "Animation/CustomAttributes.h"
 #include "Animation/AnimData/AnimDataNotifications.h"
 #include "Animation/AttributeCurve.h"
-#include "PerPlatformProperties.h"
+#include "UObject/PerPlatformProperties.h"
 #include "IO/IoHash.h"
 
 #if WITH_EDITOR
diff --git a/Engine/Source/Runtime/Engine/Classes/Animation/VariableFrameStrippingSettings.h b/Engine/Source/Runtime/Engine/Classes/Animation/VariableFrameStrippingSettings.h
--- a/Engine/Source/Runtime/Engine/Classes/Animation/VariableFrameStrippingSettings.h
+++ b/Engine/Source/Runtime/Engine/Classes/Animation/VariableFrameStrippingSettings.h
@@ -2,7 +2,7 @@
 #pragma once
 
 #include "CoreMinimal.h"
-#include "PerPlatformProperties.h"
+#include "UObject/PerPlatformProperties.h"
 #include "VariableFrameStrippingSettings.generated.h"
 
 #if WITH_EDITORONLY_DATA
diff --git a/Engine/Source/Runtime/Engine/Classes/Components/RuntimeVirtualTextureComponent.h b/Engine/Source/Runtime/Engine/Classes/Components/RuntimeVirtualTextureComponent.h
--- a/Engine/Source/Runtime/Engine/Classes/Components/RuntimeVirtualTextureComponent.h
+++ b/Engine/Source/Runtime/Engine/Classes/Components/RuntimeVirtualTextureComponent.h
@@ -4,7 +4,7 @@
 
 #include "CoreMinimal.h"
 #include "Engine/TextureDefines.h"
-#include "PerPlatformProperties.h"
+#include "UObject/PerPlatformProperties.h"
 #include "RenderCommandFence.h"
 #include "SceneComponent.h"
 #include "SceneTypes.h"
diff --git a/Engine/Source/Runtime/Engine/Classes/Engine/RendererSettings.h b/Engine/Source/Runtime/Engine/Classes/Engine/RendererSettings.h
--- a/Engine/Source/Runtime/Engine/Classes/Engine/RendererSettings.h
+++ b/Engine/Source/Runtime/Engine/Classes/Engine/RendererSettings.h
@@ -8,7 +8,7 @@
 #include "Engine/Scene.h"
 #include "Engine/DeveloperSettings.h"
 #include "PixelFormat.h"
-#include "PerPlatformProperties.h"
+#include "UObject/PerPlatformProperties.h"
 #include "LegacyScreenPercentageDriver.h"
 
 #include "RendererSettings.generated.h"
diff --git a/Engine/Source/Runtime/Engine/Classes/Engine/SkeletalMesh.h b/Engine/Source/Runtime/Engine/Classes/Engine/SkeletalMesh.h
--- a/Engine/Source/Runtime/Engine/Classes/Engine/SkeletalMesh.h
+++ b/Engine/Source/Runtime/Engine/Classes/Engine/SkeletalMesh.h
@@ -15,7 +15,7 @@
 #include "Interfaces/Interface_AssetUserData.h"
 #include "Interfaces/Interface_CollisionDataProvider.h"
 #include "Misc/EnumClassFlags.h"
-#include "PerPlatformProperties.h"
+#include "UObject/PerPlatformProperties.h"
 #include "PerQualityLevelProperties.h"
 #include "ReferenceSkeleton.h"
 #include "RenderCommandFence.h"
diff --git a/Engine/Source/Runtime/Engine/Classes/Engine/SkeletalMeshLODSettings.h b/Engine/Source/Runtime/Engine/Classes/Engine/SkeletalMeshLODSettings.h
--- a/Engine/Source/Runtime/Engine/Classes/Engine/SkeletalMeshLODSettings.h
+++ b/Engine/Source/Runtime/Engine/Classes/Engine/SkeletalMeshLODSettings.h
@@ -8,7 +8,7 @@
 #include "Engine/EngineTypes.h"
 #include "SkeletalMeshReductionSettings.h"
 #include "DataAsset.h"
-#include "PerPlatformProperties.h"
+#include "UObject/PerPlatformProperties.h"
 #include "PerQualityLevelProperties.h"
 #include "SkeletalMeshLODSettings.generated.h"
 
diff --git a/Engine/Source/Runtime/Engine/Classes/Engine/SkinnedAsset.h b/Engine/Source/Runtime/Engine/Classes/Engine/SkinnedAsset.h
--- a/Engine/Source/Runtime/Engine/Classes/Engine/SkinnedAsset.h
+++ b/Engine/Source/Runtime/Engine/Classes/Engine/SkinnedAsset.h
@@ -11,7 +11,7 @@
 #include "Engine/StreamableRenderAsset.h"
 #include "Interfaces/Interface_AsyncCompilation.h"
 #include "ReferenceSkeleton.h"
-#include "PerPlatformProperties.h"
+#include "UObject/PerPlatformProperties.h"
 #include "PSOPrecache.h"
 #include "SkinnedAsset.generated.h"
 
diff --git a/Engine/Source/Runtime/Engine/Classes/Engine/SkinnedAssetCommon.h b/Engine/Source/Runtime/Engine/Classes/Engine/SkinnedAssetCommon.h
--- a/Engine/Source/Runtime/Engine/Classes/Engine/SkinnedAssetCommon.h
+++ b/Engine/Source/Runtime/Engine/Classes/Engine/SkinnedAssetCommon.h
@@ -10,7 +10,7 @@
 #include "Components.h"
 #include "CoreMinimal.h"
 #include "Engine/EngineTypes.h"
-#include "PerPlatformProperties.h"
+#include "UObject/PerPlatformProperties.h"
 #include "SkeletalMeshReductionSettings.h"
 #include "Animation/SkeletalMeshVertexAttribute.h"
 
diff --git a/Engine/Source/Runtime/Engine/Classes/Engine/StaticMesh.h b/Engine/Source/Runtime/Engine/Classes/Engine/StaticMesh.h
--- a/Engine/Source/Runtime/Engine/Classes/Engine/StaticMesh.h
+++ b/Engine/Source/Runtime/Engine/Classes/Engine/StaticMesh.h
@@ -21,7 +21,7 @@
 #include "Engine/StreamableRenderAsset.h"
 #include "Templates/UniquePtr.h"
 #include "StaticMeshSourceData.h"
-#include "PerPlatformProperties.h"
+#include "UObject/PerPlatformProperties.h"
 #include "MeshTypes.h"
 #include "PerQualityLevelProperties.h"
 
diff --git a/Engine/Source/Runtime/Engine/Classes/Engine/StaticMeshSourceData.h b/Engine/Source/Runtime/Engine/Classes/Engine/StaticMeshSourceData.h
--- a/Engine/Source/Runtime/Engine/Classes/Engine/StaticMeshSourceData.h
+++ b/Engine/Source/Runtime/Engine/Classes/Engine/StaticMeshSourceData.h
@@ -7,7 +7,7 @@
 #include "Engine/EngineTypes.h"
 #include "MeshReductionSettings.h"
 #include "UObject/ObjectMacros.h"
-#include "PerPlatformProperties.h"
+#include "UObject/PerPlatformProperties.h"
 
 #include "StaticMeshSourceData.generated.h"
 
diff --git a/Engine/Source/Runtime/Engine/Classes/Engine/Texture.h b/Engine/Source/Runtime/Engine/Classes/Engine/Texture.h
--- a/Engine/Source/Runtime/Engine/Classes/Engine/Texture.h
+++ b/Engine/Source/Runtime/Engine/Classes/Engine/Texture.h
@@ -15,7 +15,7 @@
 #include "Engine/TextureDefines.h"
 #include "MaterialValueType.h"
 #include "Engine/StreamableRenderAsset.h"
-#include "PerPlatformProperties.h"
+#include "UObject/PerPlatformProperties.h"
 #include "ImageCore.h"
 #if WITH_EDITORONLY_DATA
 #include "Misc/TVariant.h"
diff --git a/Engine/Source/Runtime/Engine/Classes/Sound/SoundClass.h b/Engine/Source/Runtime/Engine/Classes/Sound/SoundClass.h
--- a/Engine/Source/Runtime/Engine/Classes/Sound/SoundClass.h
+++ b/Engine/Source/Runtime/Engine/Classes/Sound/SoundClass.h
@@ -16,7 +16,7 @@
 #include "EdGraph/EdGraph.h"
 #endif // WITH_EDITOR
 #include "SoundWaveLoadingBehavior.h"
-#include "PerPlatformProperties.h"
+#include "UObject/PerPlatformProperties.h"
 
 #include "SoundClass.generated.h"
 
diff --git a/Engine/Source/Runtime/Engine/Classes/Sound/SoundWave.h b/Engine/Source/Runtime/Engine/Classes/Sound/SoundWave.h
--- a/Engine/Source/Runtime/Engine/Classes/Sound/SoundWave.h
+++ b/Engine/Source/Runtime/Engine/Classes/Sound/SoundWave.h
@@ -24,7 +24,7 @@
 #include "UObject/ObjectKey.h"
 #include "AudioMixerTypes.h"
 #include "AudioCompressionSettings.h"
-#include "PerPlatformProperties.h"
+#include "UObject/PerPlatformProperties.h"
 #include "ContentStreaming.h"
 #include "IAudioProxyInitializer.h"
 #include "IWaveformTransformation.h"
diff --git a/Engine/Source/Runtime/Engine/Classes/Sound/SoundWaveLoadingBehavior.h b/Engine/Source/Runtime/Engine/Classes/Sound/SoundWaveLoadingBehavior.h
--- a/Engine/Source/Runtime/Engine/Classes/Sound/SoundWaveLoadingBehavior.h
+++ b/Engine/Source/Runtime/Engine/Classes/Sound/SoundWaveLoadingBehavior.h
@@ -5,7 +5,7 @@
 #include "UObject/ObjectMacros.h"
 
 #if WITH_EDITOR
-#include "PerPlatformProperties.h"
+#include "UObject/PerPlatformProperties.h"
 #endif //WITH_EDITOR
 
 #include "SoundWaveLoadingBehavior.generated.h"
diff --git a/Engine/Source/Runtime/Engine/Public/Animation/SkeletalMeshVertexAttribute.h b/Engine/Source/Runtime/Engine/Public/Animation/SkeletalMeshVertexAttribute.h
--- a/Engine/Source/Runtime/Engine/Public/Animation/SkeletalMeshVertexAttribute.h
+++ b/Engine/Source/Runtime/Engine/Public/Animation/SkeletalMeshVertexAttribute.h
@@ -2,7 +2,7 @@
 
 #pragma once
 
-#include "PerPlatformProperties.h"
+#include "UObject/PerPlatformProperties.h"
 
 #include "SkeletalMeshVertexAttribute.generated.h"
 
diff --git a/Engine/Source/Runtime/Engine/Public/Animation/SkinWeightProfile.h b/Engine/Source/Runtime/Engine/Public/Animation/SkinWeightProfile.h
--- a/Engine/Source/Runtime/Engine/Public/Animation/SkinWeightProfile.h
+++ b/Engine/Source/Runtime/Engine/Public/Animation/SkinWeightProfile.h
@@ -3,7 +3,7 @@
 #pragma once
 
 #include "Rendering/SkinWeightVertexBuffer.h"
-#include "PerPlatformProperties.h"
+#include "UObject/PerPlatformProperties.h"
 #include "UObject/NameTypes.h"
 #include "Misc/CoreStats.h"
 #include "RenderingThread.h"
diff --git a/Engine/Source/Runtime/Engine/Public/PerPlatformProperties.h b/Engine/Source/Runtime/Engine/Public/PerPlatformProperties.h
--- a/Engine/Source/Runtime/Engine/Public/PerPlatformProperties.h
+++ b/Engine/Source/Runtime/Engine/Public/PerPlatformProperties.h
@@ -6,4 +6,6 @@ Redirect PerPlatformProperties.h include to UObject
 
 #pragma once
 
+#pragma message("WARNING: Do not #include PerPlatformProperties.h. Please use #include UObject/PerPlatformProperties.h instead. Deprecated in 5.5")
+
 #include "UObject/PerPlatformProperties.h"
diff --git a/Engine/Source/Runtime/Engine/Public/StaticMeshResources.h b/Engine/Source/Runtime/Engine/Public/StaticMeshResources.h
--- a/Engine/Source/Runtime/Engine/Public/StaticMeshResources.h
+++ b/Engine/Source/Runtime/Engine/Public/StaticMeshResources.h
@@ -35,7 +35,7 @@
 #include "Templates/UniquePtr.h"
 #include "Serialization/BulkData.h"
 #include "WeightedRandomSampler.h"
-#include "PerPlatformProperties.h"
+#include "UObject/PerPlatformProperties.h"
 #include "RayTracingGeometry.h"
 #if WITH_EDITORONLY_DATA
 #include "Interface_CollisionDataProviderCore.h"
diff --git a/Engine/Source/Runtime/Landscape/Classes/LandscapeComponent.h b/Engine/Source/Runtime/Landscape/Classes/LandscapeComponent.h
--- a/Engine/Source/Runtime/Landscape/Classes/LandscapeComponent.h
+++ b/Engine/Source/Runtime/Landscape/Classes/LandscapeComponent.h
@@ -8,7 +8,7 @@
 #include "Misc/Guid.h"
 #include "Engine/TextureStreamingTypes.h"
 #include "Components/PrimitiveComponent.h"
-#include "PerPlatformProperties.h"
+#include "UObject/PerPlatformProperties.h"
 #include "Serialization/BulkData.h"
 #include "LandscapePhysicalMaterial.h"
 #include "LandscapeInfo.h"
diff --git a/Engine/Source/Runtime/Landscape/Classes/LandscapeGrassType.h b/Engine/Source/Runtime/Landscape/Classes/LandscapeGrassType.h
--- a/Engine/Source/Runtime/Landscape/Classes/LandscapeGrassType.h
+++ b/Engine/Source/Runtime/Landscape/Classes/LandscapeGrassType.h
@@ -6,7 +6,7 @@
 #include "UObject/ObjectMacros.h"
 #include "UObject/Object.h"
 #include "Engine/EngineTypes.h"
-#include "PerPlatformProperties.h"
+#include "UObject/PerPlatformProperties.h"
 #include "PerQualityLevelProperties.h"
 #include "SceneTypes.h"
 #include "LandscapeGrassType.generated.h"
diff --git a/Engine/Source/Runtime/Landscape/Classes/LandscapeProxy.h b/Engine/Source/Runtime/Landscape/Classes/LandscapeProxy.h
--- a/Engine/Source/Runtime/Landscape/Classes/LandscapeProxy.h
+++ b/Engine/Source/Runtime/Landscape/Classes/LandscapeProxy.h
@@ -10,7 +10,7 @@
 #include "PhysicsEngine/BodyInstance.h"
 #include "Async/AsyncWork.h"
 #include "Engine/Texture.h"
-#include "PerPlatformProperties.h"
+#include "UObject/PerPlatformProperties.h"
 #include "LandscapeNaniteComponent.h"
 #include "LandscapeWeightmapUsage.h"
 #include "LandscapeHeightfieldCollisionComponent.h"
