Fixed montage segment details ticking the scene in the widget instead of via the viewport client. This was prone to crashes when force ticking slate widgets.

[RN] minor

diff --git a/Engine/Source/Editor/DetailCustomizations/Private/AnimMontageSegmentDetails.cpp b/Engine/Source/Editor/DetailCustomizations/Private/AnimMontageSegmentDetails.cpp
--- a/Engine/Source/Editor/DetailCustomizations/Private/AnimMontageSegmentDetails.cpp
+++ b/Engine/Source/Editor/DetailCustomizations/Private/AnimMontageSegmentDetails.cpp
@@ -91,6 +91,17 @@ FAnimationSegmentViewportClient::FAnimationSegmentViewportClient(FAdvancedPrevie
 }
 
 
+void FAnimationSegmentViewportClient::Tick(float DeltaSeconds)
+{
+	FEditorViewportClient::Tick(DeltaSeconds);
+	
+	// Tick the preview scene world.
+	if (PreviewScene)
+	{
+		PreviewScene->GetWorld()->Tick(LEVELTICK_All, DeltaSeconds);
+	}
+}
+
 void FAnimationSegmentViewportClient::UpdateLighting()
 {
 	const USkeletalMeshEditorSettings* Options = GetDefault<USkeletalMeshEditorSettings>();
@@ -708,8 +719,6 @@ void SAnimationSegmentViewport::Tick( const FGeometry& AllottedGeometry, const d
 		{
 			Description->SetText(LOCTEXT("Default", "Default"));
 		}
-
-		Component->GetScene()->GetWorld()->Tick(LEVELTICK_All, InDeltaTime);
 	}
 	else
 	{
diff --git a/Engine/Source/Editor/DetailCustomizations/Private/AnimMontageSegmentDetails.h b/Engine/Source/Editor/DetailCustomizations/Private/AnimMontageSegmentDetails.h
--- a/Engine/Source/Editor/DetailCustomizations/Private/AnimMontageSegmentDetails.h
+++ b/Engine/Source/Editor/DetailCustomizations/Private/AnimMontageSegmentDetails.h
@@ -201,7 +201,7 @@ public:
 	// FlEditorViewportClient interface
 	virtual FSceneInterface* GetScene() const override;
 	virtual FLinearColor GetBackgroundColor() const override { return FLinearColor::Black; }
-
+	virtual void Tick(float DeltaSeconds) override;
 	// End of FEditorViewportClient
 
 	void UpdateLighting();
