Add capability to suppress the automatic Development Status warnings added to details panel for registered class types

diff --git a/Engine/Source/Editor/DetailCustomizations/Private/DetailCustomizations.cpp b/Engine/Source/Editor/DetailCustomizations/Private/DetailCustomizations.cpp
--- a/Engine/Source/Editor/DetailCustomizations/Private/DetailCustomizations.cpp
+++ b/Engine/Source/Editor/DetailCustomizations/Private/DetailCustomizations.cpp
@@ -214,6 +214,30 @@ void FDetailCustomizationsModule::ShutdownModule()
 	}
 }
 
+void FDetailCustomizationsModule::RegisterDevelopmentStatusWarningSupression(FName ClassName)
+{
+	SuppressedDevelopmentStatusWarnings.Add(ClassName);
+}
+
+void FDetailCustomizationsModule::UnregisterDevelopmentStatusWarningSupression(FName ClassName)
+{
+	SuppressedDevelopmentStatusWarnings.Remove(ClassName);
+}
+
+bool FDetailCustomizationsModule::IsDevelopmentStatusWarningSupressed(const UClass* Class) const
+{
+	const UClass* ClassCursor = Class;
+	while(ClassCursor)
+	{
+		if (SuppressedDevelopmentStatusWarnings.Contains(ClassCursor->GetFName()))
+		{
+			return true;
+		}
+		ClassCursor = ClassCursor->GetSuperClass();
+	}
+	return false;
+}
+
 /** Helper that will flag this struct name as supporting the UIMin and UIMax meta data types */
 #define REGISTER_UIMINMAX_CUSTOMIZATION( StructName, CallbackFunc ) \
 			RangeVisibilityUtils::StructsSupportingRangeVisibility.Add( StructName );		\
diff --git a/Engine/Source/Editor/DetailCustomizations/Private/ObjectDetails.cpp b/Engine/Source/Editor/DetailCustomizations/Private/ObjectDetails.cpp
--- a/Engine/Source/Editor/DetailCustomizations/Private/ObjectDetails.cpp
+++ b/Engine/Source/Editor/DetailCustomizations/Private/ObjectDetails.cpp
@@ -6,6 +6,7 @@
 #include "Containers/UnrealString.h"
 #include "Delegates/Delegate.h"
 #include "DetailCategoryBuilder.h"
+#include "DetailCustomizations.h"
 #include "DetailLayoutBuilder.h"
 #include "DetailWidgetRow.h"
 #include "EdGraphSchema_K2.h"
@@ -62,6 +63,12 @@ void FObjectDetails::AddExperimentalWarningCategory(IDetailLayoutBuilder& Detail
 	FString MostDerivedDevelopmentClassName;
 	FObjectEditorUtils::GetClassDevelopmentStatus(DetailBuilder.GetBaseClass(), bBaseClassIsExperimental, bBaseClassIsEarlyAccess, MostDerivedDevelopmentClassName);
 
+	FDetailCustomizationsModule& DetailCustomizationsModule = FModuleManager::Get().GetModuleChecked<FDetailCustomizationsModule>(TEXT("DetailCustomizations"));
+	if (DetailCustomizationsModule.IsDevelopmentStatusWarningSupressed(DetailBuilder.GetBaseClass()))
+	{
+		return;
+	}
+
 	if (bBaseClassIsExperimental || bBaseClassIsEarlyAccess)
 	{
 		const FName CategoryName(TEXT("Warning"));
diff --git a/Engine/Source/Editor/DetailCustomizations/Public/DetailCustomizations.h b/Engine/Source/Editor/DetailCustomizations/Public/DetailCustomizations.h
--- a/Engine/Source/Editor/DetailCustomizations/Public/DetailCustomizations.h
+++ b/Engine/Source/Editor/DetailCustomizations/Public/DetailCustomizations.h
@@ -31,6 +31,17 @@ public:
 		return true;
 	}
 
+	/**
+	 * Suppresses the DevelopmentStatus warnings for the given class or any derived versions of this class
+	 */
+	void DETAILCUSTOMIZATIONS_API RegisterDevelopmentStatusWarningSupression(FName ClassName);
+	/**
+	 * Removes suppression of DevelopmentStatus warnings for the given class or any derived versions of this class
+	 */
+	void DETAILCUSTOMIZATIONS_API UnregisterDevelopmentStatusWarningSupression(FName ClassName);
+
+	bool IsDevelopmentStatusWarningSupressed(const UClass* Class) const;
+
 private:
 	void RegisterPropertyTypeCustomizations();
 	void RegisterObjectCustomizations();
@@ -55,4 +66,5 @@ private:
 	/** List of registered class that we must unregister when the module shuts down */
 	TSet< FName > RegisteredClassNames;
 	TSet< FName > RegisteredPropertyTypes;
+	TSet< FName > SuppressedDevelopmentStatusWarnings;
 };
