Added 'Disable Content Bundles' flag in World Partition.
When checked, this flag will :
- Skip the loading/cooking of Content Bundles
- Disable the Content Bundles Outliner
- Remove the Content Bundle column in the World Outliner
- Hide the ContentBundleGuid property on the Actor

diff --git a/Engine/Source/Editor/DetailCustomizations/Private/ActorDetails.cpp b/Engine/Source/Editor/DetailCustomizations/Private/ActorDetails.cpp
--- a/Engine/Source/Editor/DetailCustomizations/Private/ActorDetails.cpp
+++ b/Engine/Source/Editor/DetailCustomizations/Private/ActorDetails.cpp
@@ -55,6 +55,7 @@
 #include "Misc/MessageDialog.h"
 #include "ScopedTransaction.h"
 #include "WorldPartition/WorldPartitionSubsystem.h"
+#include "WorldPartition/WorldPartition.h"
 #include "Algo/AnyOf.h"
 
 #include "Kismet2/BlueprintEditorUtils.h"
@@ -638,22 +639,26 @@ void FActorDetails::AddActorCategory( IDetailLayoutBuilder& DetailBuilder, const
 
 			if (Actor->GetContentBundleGuid().IsValid())
 			{
-				const FText ActorContentBundleGuidText = FText::FromString(Actor->GetContentBundleGuid().ToString());
-				ActorCategory.AddCustomRow( LOCTEXT("ContentBundleGuid", "ContentBundleGuid") )
-					.NameContent()
-					[
-						SNew(STextBlock)
-						.Text(LOCTEXT("ContentBundleGuid2", "Content Bundle Guid"))
-						.ToolTipText(LOCTEXT("ActorContentBundleGuid_ToolTip", "Actor Content BundleGuid"))
-						.Font(IDetailLayoutBuilder::GetDetailFont())
-					]
-					.ValueContent()
-					[
-						SNew(STextBlock)
-							.Text(ActorContentBundleGuidText)
+				UWorldPartition* WorldPartition = Actor->GetWorld() ? Actor->GetWorld()->GetWorldPartition() : nullptr;
+				if (WorldPartition && WorldPartition->IsContentBundleEnabled())
+				{
+					const FText ActorContentBundleGuidText = FText::FromString(Actor->GetContentBundleGuid().ToString());
+					ActorCategory.AddCustomRow( LOCTEXT("ContentBundleGuid", "ContentBundleGuid") )
+						.NameContent()
+						[
+							SNew(STextBlock)
+							.Text(LOCTEXT("ContentBundleGuid2", "Content Bundle Guid"))
+							.ToolTipText(LOCTEXT("ActorContentBundleGuid_ToolTip", "Actor Content BundleGuid"))
 							.Font(IDetailLayoutBuilder::GetDetailFont())
-							.IsEnabled(false)
-					];
+						]
+						.ValueContent()
+						[
+							SNew(STextBlock)
+								.Text(ActorContentBundleGuidText)
+								.Font(IDetailLayoutBuilder::GetDetailFont())
+								.IsEnabled(false)
+						];
+				}
 			}
 		}
 	};
diff --git a/Engine/Source/Editor/SceneOutliner/Private/SceneOutlinerModule.cpp b/Engine/Source/Editor/SceneOutliner/Private/SceneOutlinerModule.cpp
--- a/Engine/Source/Editor/SceneOutliner/Private/SceneOutlinerModule.cpp
+++ b/Engine/Source/Editor/SceneOutliner/Private/SceneOutlinerModule.cpp
@@ -569,9 +569,18 @@ void FSceneOutlinerModule::CreateActorInfoColumns(FSceneOutlinerInitializationOp
 
 	AddTextInfoColumn(FSceneOutlinerBuiltInColumnTypes::Mobility(), FSceneOutlinerBuiltInColumnTypes::Mobility_Localized(), MobilityInfoText);
 	AddTextInfoColumn(FSceneOutlinerBuiltInColumnTypes::Level(), LevelColumnName, LevelInfoText);
-	AddTextInfoColumn(FSceneOutlinerBuiltInColumnTypes::Layer(), FSceneOutlinerBuiltInColumnTypes::Layer_Localized(), LayerInfoText);
-	AddTextInfoColumn(FSceneOutlinerBuiltInColumnTypes::DataLayer(), FSceneOutlinerBuiltInColumnTypes::DataLayer_Localized(), DataLayerInfoText);
-	AddTextInfoColumn(FSceneOutlinerBuiltInColumnTypes::ContentBundle(), FSceneOutlinerBuiltInColumnTypes::ContentBundle_Localized(), ContentBundleInfoText);
+	if (UWorldPartition* WorldPartition = WorldPtr ? WorldPtr->GetWorldPartition() : nullptr)
+	{
+		AddTextInfoColumn(FSceneOutlinerBuiltInColumnTypes::DataLayer(), FSceneOutlinerBuiltInColumnTypes::DataLayer_Localized(), DataLayerInfoText);
+		if (WorldPartition->IsContentBundleEnabled())
+		{
+			AddTextInfoColumn(FSceneOutlinerBuiltInColumnTypes::ContentBundle(), FSceneOutlinerBuiltInColumnTypes::ContentBundle_Localized(), ContentBundleInfoText);
+		}
+	}
+	else
+	{
+		AddTextInfoColumn(FSceneOutlinerBuiltInColumnTypes::Layer(), FSceneOutlinerBuiltInColumnTypes::Layer_Localized(), LayerInfoText);
+	}
 	AddTextInfoColumn(FSceneOutlinerBuiltInColumnTypes::SubPackage(), FSceneOutlinerBuiltInColumnTypes::SubPackage_Localized(), SubPackageInfoText);
 	AddTextInfoColumn(FSceneOutlinerBuiltInColumnTypes::Socket(), FSceneOutlinerBuiltInColumnTypes::Socket_Localized(), SocketInfoText);
 	AddTextInfoColumn(FSceneOutlinerBuiltInColumnTypes::IDName(), FSceneOutlinerBuiltInColumnTypes::IDName_Localized(), InternalNameInfoText);
diff --git a/Engine/Source/Editor/WorldPartitionEditor/Private/WorldPartition/ContentBundle/SContentBundleBrowser.cpp b/Engine/Source/Editor/WorldPartitionEditor/Private/WorldPartition/ContentBundle/SContentBundleBrowser.cpp
--- a/Engine/Source/Editor/WorldPartitionEditor/Private/WorldPartition/ContentBundle/SContentBundleBrowser.cpp
+++ b/Engine/Source/Editor/WorldPartitionEditor/Private/WorldPartition/ContentBundle/SContentBundleBrowser.cpp
@@ -9,6 +9,7 @@
 #include "WorldPartition/ContentBundle/Outliner/ContentBundleStatusColumn.h"
 #include "WorldPartition/ContentBundle/Outliner/ContentBundleClientColumn.h"
 #include "WorldPartition/ContentBundle/Outliner/ContentBundleActorCountColumn.h"
+#include "WorldPartition/WorldPartition.h"
 
 namespace ContentBundleOutlinerPrivate
 {
@@ -43,7 +44,11 @@ void SContentBundleBrowser::Construct(const FArguments& InArgs)
 	ChildSlot
 	[
 		SAssignNew(ContentAreaBox, SVerticalBox)
-		.IsEnabled_Lambda([]() { return GWorld ? UWorld::IsPartitionedWorld(GWorld) : false; })
+		.IsEnabled_Lambda([]()
+		{
+			UWorldPartition* WorldPartition = GWorld ? GWorld->GetWorldPartition() : nullptr;
+			return WorldPartition && WorldPartition->IsContentBundleEnabled();
+		})
 	];
 
 	ContentAreaBox->AddSlot()
diff --git a/Engine/Source/Runtime/Engine/Private/WorldPartition/ContentBundle/ContentBundleWorldSubsystem.cpp b/Engine/Source/Runtime/Engine/Private/WorldPartition/ContentBundle/ContentBundleWorldSubsystem.cpp
--- a/Engine/Source/Runtime/Engine/Private/WorldPartition/ContentBundle/ContentBundleWorldSubsystem.cpp
+++ b/Engine/Source/Runtime/Engine/Private/WorldPartition/ContentBundle/ContentBundleWorldSubsystem.cpp
@@ -204,6 +204,11 @@ void UContentBundleManager::OnWorldPartitionInitialized(UWorldPartition* InWorld
 		return;
 	}
 
+	if (!InWorldPartition->IsContentBundleEnabled())
+	{
+		return;
+	}
+
 	check(GetContentBundleContainer(InWorldPartition->GetTypedOuter<UWorld>()) == nullptr);
 	TUniquePtr<FContentBundleContainer>& ContentBundleContainer = ContentBundleContainers.Emplace_GetRef(MakeUnique<FContentBundleContainer>(InWorldPartition->GetTypedOuter<UWorld>()));
 	ContentBundleContainer->Initialize();
@@ -217,7 +222,12 @@ void UContentBundleManager::OnWorldPartitionUninitialized(UWorldPartition* InWor
 	}
 
 	uint32 ContainerIndex = GetContentBundleContainerIndex(InWorldPartition->GetTypedOuter<UWorld>());
-	check(ContainerIndex != INDEX_NONE);
+	if (ContainerIndex == INDEX_NONE)
+	{
+		return;
+	}
+
+	check(ContentBundleContainers.IsValidIndex(ContainerIndex));
 	ContentBundleContainers[ContainerIndex]->Deinitialize();
 	ContentBundleContainers.RemoveAtSwap(ContainerIndex);
 }
diff --git a/Engine/Source/Runtime/Engine/Private/WorldPartition/WorldPartition.cpp b/Engine/Source/Runtime/Engine/Private/WorldPartition/WorldPartition.cpp
--- a/Engine/Source/Runtime/Engine/Private/WorldPartition/WorldPartition.cpp
+++ b/Engine/Source/Runtime/Engine/Private/WorldPartition/WorldPartition.cpp
@@ -339,6 +339,7 @@ UWorldPartition::UWorldPartition(const FObjectInitializer& ObjectInitializer)
 	, Replay(nullptr)
 {
 	bEnableStreaming = true;
+	bDisableContentBundles = false;
 	ServerStreamingMode = EWorldPartitionServerStreamingMode::ProjectDefault;
 	ServerStreamingOutMode = EWorldPartitionServerStreamingOutMode::ProjectDefault;
 	DataLayersLogicOperator = EWorldPartitionDataLayersLogicOperator::Or;
diff --git a/Engine/Source/Runtime/Engine/Public/WorldPartition/WorldPartition.h b/Engine/Source/Runtime/Engine/Public/WorldPartition/WorldPartition.h
--- a/Engine/Source/Runtime/Engine/Public/WorldPartition/WorldPartition.h
+++ b/Engine/Source/Runtime/Engine/Public/WorldPartition/WorldPartition.h
@@ -458,6 +458,7 @@ public:
 	ENGINE_API bool CanStream() const;
 	ENGINE_API bool IsServer() const;
 	ENGINE_API bool IsServerStreamingEnabled() const;
+	ENGINE_API bool IsContentBundleEnabled() const { return !bDisableContentBundles; }
 	ENGINE_API bool IsServerStreamingOutEnabled() const;
 	ENGINE_API bool UseMakingVisibleTransactionRequests() const;
 	ENGINE_API bool UseMakingInvisibleTransactionRequests() const;
@@ -573,6 +574,10 @@ private:
 	uint8 bAllowShowingHLODsInEditor : 1;
 #endif
 
+	/** if set to true, this removes any content bundles from this world and also removes content bundle editing */
+	UPROPERTY(EditAnywhere, Category = WorldPartitionSetup, AdvancedDisplay)
+	uint8 bDisableContentBundles : 1;
+
 	TObjectPtr<UWorld> World;
 
 #if WITH_EDITOR
