FDataTableCategoryCustomization now respects the RowType meta data

#rnx

diff --git a/Engine/Source/Editor/DetailCustomizations/Private/DataTableCategoryCustomization.h b/Engine/Source/Editor/DetailCustomizations/Private/DataTableCategoryCustomization.h
--- a/Engine/Source/Editor/DetailCustomizations/Private/DataTableCategoryCustomization.h
+++ b/Engine/Source/Editor/DetailCustomizations/Private/DataTableCategoryCustomization.h
@@ -11,6 +11,7 @@
 #include "IDetailChildrenBuilder.h"
 #include "Engine/DataTable.h"
 #include "DetailWidgetRow.h"
+#include "PropertyCustomizationHelpers.h"
 
 #define LOCTEXT_NAMESPACE "FDataTableCategoryCustomizationLayout"
 
@@ -30,6 +31,13 @@ public:
 	{
 		this->StructPropertyHandle = InStructPropertyHandle;
 
+		if (StructPropertyHandle->HasMetaData(TEXT("RowType")))
+		{
+			const FString& RowType = StructPropertyHandle->GetMetaData(TEXT("RowType"));
+			RowTypeFilter = FName(*RowType);
+			RowFilterStruct = UClass::TryFindTypeSlow<UScriptStruct>(RowType);
+		}
+		
 		HeaderRow
 			.NameContent()
 			[
@@ -48,8 +56,23 @@ public:
 			&& ColumnNamePropertyHandle->IsValidHandle()
 			&& RowContentsPropertyHandle->IsValidHandle())
 		{
-			/** Edit the data table uobject as normal */
-			StructBuilder.AddProperty(DataTablePropertyHandle.ToSharedRef());
+			/** Construct a asset picker widget with a custom filter */
+			StructBuilder.AddCustomRow(LOCTEXT("DataTable_TableName", "Data Table"))
+				.NameContent()
+				[
+					SNew(STextBlock)
+					.Text(LOCTEXT("DataTable_TableName", "Data Table"))
+					.Font(StructCustomizationUtils.GetRegularFont())
+				]
+				.ValueContent()
+				.MaxDesiredWidth(0.0f) // don't constrain the combo button width
+				[
+					SNew(SObjectPropertyEntryBox)
+					.PropertyHandle(DataTablePropertyHandle)
+					.AllowedClass(UDataTable::StaticClass())
+					.OnShouldFilterAsset(this, &FDataTableCategoryCustomizationLayout::ShouldFilterAsset)
+				];
+			
 			FSimpleDelegate OnDataTableChangedDelegate = FSimpleDelegate::CreateSP(this, &FDataTableCategoryCustomizationLayout::OnDataTableChanged);
 			DataTablePropertyHandle->SetOnPropertyValueChanged(OnDataTableChangedDelegate);
 
@@ -301,6 +324,32 @@ private:
 			RowContentsPropertyHandle->SetValue(NewValue);
 		}
 	}
+	
+	/** Returns true if we should hide asset due to wrong type */
+	bool ShouldFilterAsset(const FAssetData& AssetData)
+	{
+		if (!RowTypeFilter.IsNone())
+		{
+			static const FName RowStructureTagName("RowStructure");
+			FString RowStructure;
+			if (AssetData.GetTagValue<FString>(RowStructureTagName, RowStructure))
+			{
+				if (RowStructure == RowTypeFilter.ToString())
+				{
+					return false;
+				}
+
+				// This is slow, but at the moment we don't have an alternative to the short struct name search
+				UScriptStruct* RowStruct = UClass::TryFindTypeSlow<UScriptStruct>(RowStructure);
+				if (RowStruct && RowFilterStruct && RowStruct->IsChildOf(RowFilterStruct))
+				{
+					return false;
+				}
+			}
+			return true;
+		}
+		return false;
+	}
 
 	/** The column combobox object */
 	TSharedPtr<SComboBox<TSharedPtr<FString> > > ColumnNameComboBox;
@@ -315,6 +364,9 @@ private:
 	TArray<TSharedPtr<FString> > RowContents;
 	/** A cached copy of strings to populate the column combo box */
 	TArray<TSharedPtr<FString> > ColumnNames;
+	/** The MetaData derived filter for the row type */
+	FName RowTypeFilter;
+	UScriptStruct* RowFilterStruct = nullptr;
 };
 
 #undef LOCTEXT_NAMESPACE
