Fixed missing transaction when overriding landscape properties

diff --git a/Engine/Source/Editor/DetailCustomizations/Private/LandscapeProxyUIDetails.cpp b/Engine/Source/Editor/DetailCustomizations/Private/LandscapeProxyUIDetails.cpp
--- a/Engine/Source/Editor/DetailCustomizations/Private/LandscapeProxyUIDetails.cpp
+++ b/Engine/Source/Editor/DetailCustomizations/Private/LandscapeProxyUIDetails.cpp
@@ -281,7 +281,7 @@ void FLandscapeProxyUIDetails::CustomizeDetails( IDetailLayoutBuilder& DetailBui
 				{
 					// Extend the tool tip to indicate this property is inherited
 					FText ToolTipText = NanitePositionPrecisionHandle->GetToolTipText();
-					DetailRow->ToolTip(FText::Format(NSLOCTEXT("Landscape", "InheritedProperty", "{0} This property is inherited from the parent Landscape proxy."), ToolTipText));
+					DetailRow->ToolTip(FText::Format(LOCTEXT("InheritedProperty", "{0} This property is inherited from the parent Landscape proxy."), ToolTipText));
 			
 					// Disable the property editing
 					DetailRow->IsEnabled(false);
@@ -436,7 +436,7 @@ void FLandscapeProxyUIDetails::CustomizeDetails( IDetailLayoutBuilder& DetailBui
 					{
 						// Extend the tool tip to indicate this property is inherited
 						FText ToolTipText = PropertyHandle->GetToolTipText();
-						DetailRow->ToolTip(FText::Format(NSLOCTEXT("Landscape", "InheritedProperty", "{0} This property is inherited from the parent Landscape proxy."), ToolTipText));
+						DetailRow->ToolTip(FText::Format(LOCTEXT("InheritedProperty", "{0} This property is inherited from the parent Landscape proxy."), ToolTipText));
 
 						// Disable the property editing
 						DetailRow->IsEnabled(false);
@@ -449,7 +449,7 @@ void FLandscapeProxyUIDetails::CustomizeDetails( IDetailLayoutBuilder& DetailBui
 
 				if (PropertyHandle->IsValidHandle())
 				{
-					const FText TooltipText = NSLOCTEXT("Landscape", "OverriddenProperty", "Check this box to override the parent landscape's property.");
+					const FText TooltipText = LOCTEXT("OverriddenProperty", "Check this box to override the parent landscape's property.");
 					FName PropertyName = Property->GetFName();
 					IDetailPropertyRow* DetailRow = DetailBuilder.EditDefaultProperty(PropertyHandle);
 					TSharedPtr<SWidget> NameWidget = nullptr;
@@ -505,6 +505,8 @@ void FLandscapeProxyUIDetails::CustomizeDetails( IDetailLayoutBuilder& DetailBui
 
 								const bool bChecked = NewState == ECheckBoxState::Checked;
 
+								const FScopedTransaction Transaction(LOCTEXT("SetSharedPropertyOverride", "Change Property Override"));
+
 								for (const TWeakObjectPtr<ALandscapeStreamingProxy> StreamingProxy : EditingStreamingProxies)
 								{
 									if (!StreamingProxy.IsValid())
