Pass the EPropertyValueSetFlags to the SetPerObjectValues function of the FMathStructCustomization as to avoid unecessary interactive updates.

diff --git a/Engine/Source/Editor/DetailCustomizations/Private/MathStructCustomizations.cpp b/Engine/Source/Editor/DetailCustomizations/Private/MathStructCustomizations.cpp
--- a/Engine/Source/Editor/DetailCustomizations/Private/MathStructCustomizations.cpp
+++ b/Engine/Source/Editor/DetailCustomizations/Private/MathStructCustomizations.cpp
@@ -680,7 +680,7 @@ void FMathStructCustomization::SetValue(NumericType NewValue, EPropertyValueSetF
 								ObjectChildValues[ChildOutputIndex] = TTypeToString<NumericType>::ToSanitizedString(ChildNewValue);
 							}
 
-							ChildHandle->SetPerObjectValues(ObjectChildValues);
+							ChildHandle->SetPerObjectValues(ObjectChildValues, Flags);
 						}
 					}
 				}
