Two utility classes for the DetailsView class

1. DetailsViewCategoryHiderCustomization - Used to hide entire categories of a Details View.

2. DetailsViewWarningOrErrorCustomization - Used to insert an SWarningOrErrorBox into a Details View at the top of a specified category.

diff --git a/Engine/Source/Editor/DetailCustomizations/Private/DetailsViewCategoryHiderCustomization.cpp b/Engine/Source/Editor/DetailCustomizations/Private/DetailsViewCategoryHiderCustomization.cpp
new file mode 100644
--- /dev/null
+++ b/Engine/Source/Editor/DetailCustomizations/Private/DetailsViewCategoryHiderCustomization.cpp
@@ -0,0 +1,32 @@
+// Copyright Epic Games, Inc. All Rights Reserved.
+#include "DetailsViewCategoryHiderCustomization.h"
+
+#include "DetailLayoutBuilder.h"
+
+TSharedRef<IDetailCustomization> FDetailsViewCategoryHiderCustomization::MakeInstance(TArray<FName>&& InCategoriesToHide)
+{
+	return MakeShared<FDetailsViewCategoryHiderCustomization>(MoveTemp(InCategoriesToHide));
+}
+
+TSharedRef<IDetailCustomization> FDetailsViewCategoryHiderCustomization::MakeInstance(const TArrayView<FName>& InCategoriesToHide)
+{
+	return MakeShared<FDetailsViewCategoryHiderCustomization>(InCategoriesToHide);
+}
+
+FDetailsViewCategoryHiderCustomization::FDetailsViewCategoryHiderCustomization(TArray<FName>&& InCategoriesToHide) 
+	: CategoriesToHide(MoveTemp(InCategoriesToHide))
+{
+}
+
+FDetailsViewCategoryHiderCustomization::FDetailsViewCategoryHiderCustomization(const TArrayView<FName>& InCategoriesToHide)
+	: CategoriesToHide(InCategoriesToHide)
+{
+}
+
+void FDetailsViewCategoryHiderCustomization::CustomizeDetails(class IDetailLayoutBuilder& DetailLayoutBuilder)
+{
+	for (const FName& CategoryName : CategoriesToHide)
+	{
+		DetailLayoutBuilder.HideCategory(CategoryName);
+	}
+}
diff --git a/Engine/Source/Editor/DetailCustomizations/Private/DetailsViewWarningOrErrorCustomization.cpp b/Engine/Source/Editor/DetailCustomizations/Private/DetailsViewWarningOrErrorCustomization.cpp
new file mode 100644
--- /dev/null
+++ b/Engine/Source/Editor/DetailCustomizations/Private/DetailsViewWarningOrErrorCustomization.cpp
@@ -0,0 +1,53 @@
+// Copyright Epic Games, Inc. All Rights Reserved.
+#include "DetailsViewWarningOrErrorCustomization.h"
+
+#include "DetailWidgetRow.h"
+#include "Widgets/Text/STextBlock.h"
+
+TSharedRef<IDetailCustomization> FDetailsViewWarningOrErrorCustomization::MakeInstance(
+	const FName& InCategoryForInsertion,
+	const FName& InRowTag,
+	const FText& InWarningOrErrorLabel,
+	const EMessageStyle InMessageStyle/* = EMessageStyle::Warning*/,
+	ECategoryPriority::Type InCategoryPriority/* = ECategoryPriority::Uncommon*/)
+{
+	return MakeShared<FDetailsViewWarningOrErrorCustomization>(InCategoryForInsertion, InRowTag, InWarningOrErrorLabel, InMessageStyle, InCategoryPriority);
+}
+
+FDetailsViewWarningOrErrorCustomization::FDetailsViewWarningOrErrorCustomization(
+	const FName& InCategoryForInsertion, 
+	const FName& InRowTag, 
+	const FText& InWarningOrErrorLabel, 
+	const EMessageStyle InMessageStyle, 
+	const ECategoryPriority::Type InCategoryPriority) :
+		CategoryForInsertion(InCategoryForInsertion),
+		RowTag(InRowTag),
+		WarningOrErrorLabel(InWarningOrErrorLabel),
+		MessageStyle(InMessageStyle),
+		CategoryPriority(InCategoryPriority)
+{
+}
+
+void FDetailsViewWarningOrErrorCustomization::CustomizeDetails(class IDetailLayoutBuilder& DetailLayoutBuilder)
+{
+	IDetailCategoryBuilder& TargetCategory = DetailLayoutBuilder.EditCategory(CategoryForInsertion, FText::GetEmpty(), CategoryPriority);
+
+	FDetailWidgetRow& WarningOrErrorRow = TargetCategory
+		.AddCustomRow(FText::GetEmpty(), false)
+		.RowTag(RowTag)
+		[
+			SNew(SVerticalBox)
+				+ SVerticalBox::Slot()
+				.AutoHeight()
+				.HAlign(HAlign_Left)
+				[
+					SNew(SWarningOrErrorBox)
+						.MessageStyle(MessageStyle)
+						[
+							SNew(STextBlock)
+								.Text(WarningOrErrorLabel)
+								.Justification(ETextJustify::Left)
+						]
+				]
+		];
+}
diff --git a/Engine/Source/Editor/DetailCustomizations/Public/DetailsViewCategoryHiderCustomization.h b/Engine/Source/Editor/DetailCustomizations/Public/DetailsViewCategoryHiderCustomization.h
new file mode 100644
--- /dev/null
+++ b/Engine/Source/Editor/DetailCustomizations/Public/DetailsViewCategoryHiderCustomization.h
@@ -0,0 +1,22 @@
+// Copyright Epic Games, Inc. All Rights Reserved.
+#include "IDetailCustomization.h"
+
+/** Used to hide categories in a Details View*/
+class FDetailsViewCategoryHiderCustomization : public IDetailCustomization
+{
+public:
+	static DETAILCUSTOMIZATIONS_API TSharedRef<IDetailCustomization> MakeInstance(TArray<FName>&& InCategoriesToHide);
+	static DETAILCUSTOMIZATIONS_API TSharedRef<IDetailCustomization> MakeInstance(const TArrayView<FName>& InCategoriesToHide);
+
+private:
+	FDetailsViewCategoryHiderCustomization(TArray<FName>&& InCategoriesToHide);
+	FDetailsViewCategoryHiderCustomization(const TArrayView<FName>& InCategoriesToHide);
+
+	virtual void CustomizeDetails(class IDetailLayoutBuilder& DetailLayoutBuilder) override;
+
+private:
+	TArray<FName> CategoriesToHide;
+
+	template <typename ObjectType, ESPMode Mode>
+	friend class SharedPointerInternals::TIntrusiveReferenceController;
+};
diff --git a/Engine/Source/Editor/DetailCustomizations/Public/DetailsViewWarningOrErrorCustomization.h b/Engine/Source/Editor/DetailCustomizations/Public/DetailsViewWarningOrErrorCustomization.h
new file mode 100644
--- /dev/null
+++ b/Engine/Source/Editor/DetailCustomizations/Public/DetailsViewWarningOrErrorCustomization.h
@@ -0,0 +1,36 @@
+// Copyright Epic Games, Inc. All Rights Reserved.
+#include "DetailLayoutBuilder.h"
+#include "IDetailCustomization.h"
+#include "SWarningOrErrorBox.h"
+
+/** Used to insert an SWarningOrErrorBox into a Details View at the top of a specified category */
+class FDetailsViewWarningOrErrorCustomization : public IDetailCustomization
+{
+public:
+	static DETAILCUSTOMIZATIONS_API TSharedRef<IDetailCustomization> MakeInstance(
+		const FName& InCategoryForInsertion,
+		const FName& InRowTag,
+		const FText& InWarningOrErrorLabel,
+		const EMessageStyle InMessageStyle = EMessageStyle::Warning,
+		ECategoryPriority::Type InCategoryPriority = ECategoryPriority::Uncommon);
+
+private:
+	FDetailsViewWarningOrErrorCustomization(
+		const FName& InCategoryForInsertion, 
+		const FName& InRowTag, 
+		const FText& InWarningOrErrorLabel, 
+		const EMessageStyle InMessageStyle, 
+		const ECategoryPriority::Type InCategoryPriority);
+
+	virtual void CustomizeDetails(IDetailLayoutBuilder& DetailLayoutBuilder) override;
+	
+private:
+	FName CategoryForInsertion;
+	FName RowTag;
+	FText WarningOrErrorLabel;
+	EMessageStyle MessageStyle;
+	ECategoryPriority::Type CategoryPriority;
+
+	template <typename ObjectType, ESPMode Mode>
+	friend class SharedPointerInternals::TIntrusiveReferenceController;
+};
