Removed redundant STextBlock. Fixes word wrapping when resizing the parent container.

diff --git a/Engine/Source/Editor/DetailCustomizations/Private/DetailsViewWarningOrErrorCustomization.cpp b/Engine/Source/Editor/DetailCustomizations/Private/DetailsViewWarningOrErrorCustomization.cpp
--- a/Engine/Source/Editor/DetailCustomizations/Private/DetailsViewWarningOrErrorCustomization.cpp
+++ b/Engine/Source/Editor/DetailCustomizations/Private/DetailsViewWarningOrErrorCustomization.cpp
@@ -39,15 +39,10 @@ void FDetailsViewWarningOrErrorCustomization::CustomizeDetails(class IDetailLayo
 			SNew(SVerticalBox)
 				+ SVerticalBox::Slot()
 				.AutoHeight()
-				.HAlign(HAlign_Left)
 				[
 					SNew(SWarningOrErrorBox)
 						.MessageStyle(MessageStyle)
-						[
-							SNew(STextBlock)
-								.Text(WarningOrErrorLabel)
-								.Justification(ETextJustify::Left)
-						]
+						.Message(WarningOrErrorLabel)
 				]
 		];
 }
