SceneCapture: Added tonemapper and post-process materials to visible scene capture show flags.

This allows users to avoid double application of color grading operations (from PPV color grading or PPMs) on scene captured render targets rendered back into the level.

diff --git a/Engine/Source/Editor/DetailCustomizations/Private/SceneCaptureDetails.cpp b/Engine/Source/Editor/DetailCustomizations/Private/SceneCaptureDetails.cpp
--- a/Engine/Source/Editor/DetailCustomizations/Private/SceneCaptureDetails.cpp
+++ b/Engine/Source/Editor/DetailCustomizations/Private/SceneCaptureDetails.cpp
@@ -120,6 +120,8 @@ void FSceneCaptureDetails::CustomizeDetails( IDetailLayoutBuilder& DetailLayout
 	ShowFlagsToAllowForCaptures.Add(FEngineShowFlags::EShowFlag::SF_Game);
 	ShowFlagsToAllowForCaptures.Add(FEngineShowFlags::EShowFlag::SF_ToneCurve); 
 	ShowFlagsToAllowForCaptures.Add(FEngineShowFlags::EShowFlag::SF_PathTracing);
+	ShowFlagsToAllowForCaptures.Add(FEngineShowFlags::EShowFlag::SF_Tonemapper);
+	ShowFlagsToAllowForCaptures.Add(FEngineShowFlags::EShowFlag::SF_PostProcessMaterial);
 
 	// Create array of flag name strings for each group
 	TArray< TArray<FString> > ShowFlagsByGroup;
