[DetailsCustomizations] RowStruct curve editor validates Owner before spawn (crash fix)

diff --git a/Engine/Source/Editor/DetailCustomizations/Private/CurveColorCustomization.cpp b/Engine/Source/Editor/DetailCustomizations/Private/CurveColorCustomization.cpp
--- a/Engine/Source/Editor/DetailCustomizations/Private/CurveColorCustomization.cpp
+++ b/Engine/Source/Editor/DetailCustomizations/Private/CurveColorCustomization.cpp
@@ -462,37 +462,41 @@ FReply FCurveColorCustomization::OnCurvePreviewDoubleClick(const FGeometry& InMy
 		{
 			DestroyPopOutWindow();
 
-			// Determine the position of the window so that it will spawn near the mouse, but not go off the screen.
-			const FVector2D CursorPos = FSlateApplication::Get().GetCursorPos();
-			FSlateRect Anchor(CursorPos.X, CursorPos.Y, CursorPos.X, CursorPos.Y);
-
-			FVector2D AdjustedSummonLocation = FSlateApplication::Get().CalculatePopupWindowPosition( Anchor, FCurveColorCustomization::DEFAULT_WINDOW_SIZE, true, FVector2D::ZeroVector, Orient_Horizontal );
-
-			TSharedPtr<SWindow> Window = SNew(SWindow)
-				.Title( FText::Format( LOCTEXT("WindowHeader", "{0} - Internal Color Curve Editor"), StructPropertyHandle->GetPropertyDisplayName()) )
-				.ClientSize( FCurveColorCustomization::DEFAULT_WINDOW_SIZE )
-				.ScreenPosition(AdjustedSummonLocation)
-				.AutoCenter(EAutoCenter::None)
-				.SupportsMaximize(false)
-				.SupportsMinimize(false)
-				.SizingRule( ESizingRule::FixedSize );
-
-			// init the mini curve editor widget
-			TSharedRef<SMiniCurveEditor> MiniCurveEditor =
-				SNew(SMiniCurveEditor)
-				.CurveOwner(this)
-				.OwnerObject(Owner)
-				.ParentWindow(Window);
-
-			Window->SetContent( MiniCurveEditor );
-
-			// Find the window of the parent widget
-			FWidgetPath WidgetPath;
-			FSlateApplication::Get().GeneratePathToWidgetChecked( CurveWidget.ToSharedRef(), WidgetPath );
-			Window = FSlateApplication::Get().AddWindowAsNativeChild( Window.ToSharedRef(), WidgetPath.GetWindow() );
-
-			//hold on to the window created for external use...
-			CurveEditorWindow = Window;
+			// Check that the owner exists before spawning the window
+			if (Owner)
+			{
+				// Determine the position of the window so that it will spawn near the mouse, but not go off the screen.
+				const FVector2D CursorPos = FSlateApplication::Get().GetCursorPos();
+				FSlateRect Anchor(CursorPos.X, CursorPos.Y, CursorPos.X, CursorPos.Y);
+
+				FVector2D AdjustedSummonLocation = FSlateApplication::Get().CalculatePopupWindowPosition( Anchor, FCurveColorCustomization::DEFAULT_WINDOW_SIZE, true, FVector2D::ZeroVector, Orient_Horizontal );
+
+				TSharedPtr<SWindow> Window = SNew(SWindow)
+					.Title( FText::Format( LOCTEXT("WindowHeader", "{0} - Internal Color Curve Editor"), StructPropertyHandle->GetPropertyDisplayName()) )
+					.ClientSize( FCurveColorCustomization::DEFAULT_WINDOW_SIZE )
+					.ScreenPosition(AdjustedSummonLocation)
+					.AutoCenter(EAutoCenter::None)
+					.SupportsMaximize(false)
+					.SupportsMinimize(false)
+					.SizingRule( ESizingRule::FixedSize );
+
+				// init the mini curve editor widget
+				TSharedRef<SMiniCurveEditor> MiniCurveEditor =
+					SNew(SMiniCurveEditor)
+					.CurveOwner(this)
+					.OwnerObject(Owner)
+					.ParentWindow(Window);
+
+				Window->SetContent( MiniCurveEditor );
+
+				// Find the window of the parent widget
+				FWidgetPath WidgetPath;
+				FSlateApplication::Get().GeneratePathToWidgetChecked( CurveWidget.ToSharedRef(), WidgetPath );
+				Window = FSlateApplication::Get().AddWindowAsNativeChild( Window.ToSharedRef(), WidgetPath.GetWindow() );
+
+				//hold on to the window created for external use...
+				CurveEditorWindow = Window;
+			}
 		}
 	}
 	return FReply::Handled();
diff --git a/Engine/Source/Editor/DetailCustomizations/Private/CurveStructCustomization.cpp b/Engine/Source/Editor/DetailCustomizations/Private/CurveStructCustomization.cpp
--- a/Engine/Source/Editor/DetailCustomizations/Private/CurveStructCustomization.cpp
+++ b/Engine/Source/Editor/DetailCustomizations/Private/CurveStructCustomization.cpp
@@ -459,37 +459,41 @@ FReply FCurveStructCustomization::OnCurvePreviewDoubleClick( const FGeometry& In
 		{
 			DestroyPopOutWindow();
 
-			// Determine the position of the window so that it will spawn near the mouse, but not go off the screen.
-			const FVector2D CursorPos = FSlateApplication::Get().GetCursorPos();
-			FSlateRect Anchor(CursorPos.X, CursorPos.Y, CursorPos.X, CursorPos.Y);
-
-			FVector2D AdjustedSummonLocation = FSlateApplication::Get().CalculatePopupWindowPosition( Anchor, FCurveStructCustomization::DEFAULT_WINDOW_SIZE, true, FVector2D::ZeroVector, Orient_Horizontal );
-
-			TSharedPtr<SWindow> Window = SNew(SWindow)
-				.Title( FText::Format( LOCTEXT("WindowHeader", "{0} - Internal Curve Editor"), StructPropertyHandle->GetPropertyDisplayName()) )
-				.ClientSize( FCurveStructCustomization::DEFAULT_WINDOW_SIZE )
-				.ScreenPosition(AdjustedSummonLocation)
-				.AutoCenter(EAutoCenter::None)
-				.SupportsMaximize(false)
-				.SupportsMinimize(false)
-				.SizingRule( ESizingRule::FixedSize );
-
-			// init the mini curve editor widget
-			TSharedRef<SMiniCurveEditor> MiniCurveEditor =
-				SNew(SMiniCurveEditor)
-				.CurveOwner(this)
-				.OwnerObject(Owner)
-				.ParentWindow(Window);
-
-			Window->SetContent( MiniCurveEditor );
-
-			// Find the window of the parent widget
-			FWidgetPath WidgetPath;
-			FSlateApplication::Get().GeneratePathToWidgetChecked( CurveWidget.ToSharedRef(), WidgetPath );
-			Window = FSlateApplication::Get().AddWindowAsNativeChild( Window.ToSharedRef(), WidgetPath.GetWindow() );
-
-			//hold on to the window created for external use...
-			CurveEditorWindow = Window;
+			// Check that the owner exists before spawning a window
+			if (Owner)
+			{
+				// Determine the position of the window so that it will spawn near the mouse, but not go off the screen.
+				const FVector2D CursorPos = FSlateApplication::Get().GetCursorPos();
+				FSlateRect Anchor(CursorPos.X, CursorPos.Y, CursorPos.X, CursorPos.Y);
+
+				FVector2D AdjustedSummonLocation = FSlateApplication::Get().CalculatePopupWindowPosition( Anchor, FCurveStructCustomization::DEFAULT_WINDOW_SIZE, true, FVector2D::ZeroVector, Orient_Horizontal );
+
+				TSharedPtr<SWindow> Window = SNew(SWindow)
+					.Title( FText::Format( LOCTEXT("WindowHeader", "{0} - Internal Curve Editor"), StructPropertyHandle->GetPropertyDisplayName()) )
+					.ClientSize( FCurveStructCustomization::DEFAULT_WINDOW_SIZE )
+					.ScreenPosition(AdjustedSummonLocation)
+					.AutoCenter(EAutoCenter::None)
+					.SupportsMaximize(false)
+					.SupportsMinimize(false)
+					.SizingRule( ESizingRule::FixedSize );
+
+				// init the mini curve editor widget
+				TSharedRef<SMiniCurveEditor> MiniCurveEditor =
+					SNew(SMiniCurveEditor)
+					.CurveOwner(this)
+					.OwnerObject(Owner)
+					.ParentWindow(Window);
+
+				Window->SetContent( MiniCurveEditor );
+
+				// Find the window of the parent widget
+				FWidgetPath WidgetPath;
+				FSlateApplication::Get().GeneratePathToWidgetChecked( CurveWidget.ToSharedRef(), WidgetPath );
+				Window = FSlateApplication::Get().AddWindowAsNativeChild( Window.ToSharedRef(), WidgetPath.GetWindow() );
+
+				//hold on to the window created for external use...
+				CurveEditorWindow = Window;
+			}
 		}
 	}
 	return FReply::Handled();
diff --git a/Engine/Source/Editor/DetailCustomizations/Private/CurveVectorCustomization.cpp b/Engine/Source/Editor/DetailCustomizations/Private/CurveVectorCustomization.cpp
--- a/Engine/Source/Editor/DetailCustomizations/Private/CurveVectorCustomization.cpp
+++ b/Engine/Source/Editor/DetailCustomizations/Private/CurveVectorCustomization.cpp
@@ -446,37 +446,41 @@ FReply FCurveVectorCustomization::OnCurvePreviewDoubleClick(const FGeometry& InM
 		{
 			DestroyPopOutWindow();
 
-			// Determine the position of the window so that it will spawn near the mouse, but not go off the screen.
-			const FVector2D CursorPos = FSlateApplication::Get().GetCursorPos();
-			FSlateRect Anchor(CursorPos.X, CursorPos.Y, CursorPos.X, CursorPos.Y);
-
-			FVector2D AdjustedSummonLocation = FSlateApplication::Get().CalculatePopupWindowPosition(Anchor, FCurveVectorCustomization::DEFAULT_WINDOW_SIZE, true, FVector2D::ZeroVector, Orient_Horizontal);
-
-			TSharedPtr<SWindow> Window = SNew(SWindow)
-				.Title(FText::Format(LOCTEXT("WindowHeader", "{0} - Internal Vector Curve Editor"), StructPropertyHandle->GetPropertyDisplayName()))
-				.ClientSize(FCurveVectorCustomization::DEFAULT_WINDOW_SIZE)
-				.ScreenPosition(AdjustedSummonLocation)
-				.AutoCenter(EAutoCenter::None)
-				.SupportsMaximize(false)
-				.SupportsMinimize(false)
-				.SizingRule( ESizingRule::FixedSize );
-
-			// init the mini curve editor widget
-			TSharedRef<SMiniCurveEditor> MiniCurveEditor =
-				SNew(SMiniCurveEditor)
-				.CurveOwner(this)
-				.OwnerObject(Owner)
-				.ParentWindow(Window);
-
-			Window->SetContent( MiniCurveEditor );
-
-			// Find the window of the parent widget
-			FWidgetPath WidgetPath;
-			FSlateApplication::Get().GeneratePathToWidgetChecked( CurveWidget.ToSharedRef(), WidgetPath );
-			Window = FSlateApplication::Get().AddWindowAsNativeChild( Window.ToSharedRef(), WidgetPath.GetWindow() );
-
-			//hold on to the window created for external use...
-			CurveEditorWindow = Window;
+			// Check that the owner exists before spawning the window
+			if (Owner)
+			{
+				// Determine the position of the window so that it will spawn near the mouse, but not go off the screen.
+				const FVector2D CursorPos = FSlateApplication::Get().GetCursorPos();
+				FSlateRect Anchor(CursorPos.X, CursorPos.Y, CursorPos.X, CursorPos.Y);
+
+				FVector2D AdjustedSummonLocation = FSlateApplication::Get().CalculatePopupWindowPosition(Anchor, FCurveVectorCustomization::DEFAULT_WINDOW_SIZE, true, FVector2D::ZeroVector, Orient_Horizontal);
+
+				TSharedPtr<SWindow> Window = SNew(SWindow)
+					.Title(FText::Format(LOCTEXT("WindowHeader", "{0} - Internal Vector Curve Editor"), StructPropertyHandle->GetPropertyDisplayName()))
+					.ClientSize(FCurveVectorCustomization::DEFAULT_WINDOW_SIZE)
+					.ScreenPosition(AdjustedSummonLocation)
+					.AutoCenter(EAutoCenter::None)
+					.SupportsMaximize(false)
+					.SupportsMinimize(false)
+					.SizingRule( ESizingRule::FixedSize );
+
+				// init the mini curve editor widget
+				TSharedRef<SMiniCurveEditor> MiniCurveEditor =
+					SNew(SMiniCurveEditor)
+					.CurveOwner(this)
+					.OwnerObject(Owner)
+					.ParentWindow(Window);
+
+				Window->SetContent( MiniCurveEditor );
+
+				// Find the window of the parent widget
+				FWidgetPath WidgetPath;
+				FSlateApplication::Get().GeneratePathToWidgetChecked( CurveWidget.ToSharedRef(), WidgetPath );
+				Window = FSlateApplication::Get().AddWindowAsNativeChild( Window.ToSharedRef(), WidgetPath.GetWindow() );
+
+				//hold on to the window created for external use...
+				CurveEditorWindow = Window;
+			}
 		}
 	}
 	return FReply::Handled();
