Fixes crash in animation sequence editor when viewing old animation with bone track names out of sync with data model.

diff --git a/Engine/Source/Editor/DetailCustomizations/Private/AnimSequenceDetails.cpp b/Engine/Source/Editor/DetailCustomizations/Private/AnimSequenceDetails.cpp
--- a/Engine/Source/Editor/DetailCustomizations/Private/AnimSequenceDetails.cpp
+++ b/Engine/Source/Editor/DetailCustomizations/Private/AnimSequenceDetails.cpp
@@ -304,7 +304,15 @@ void FAnimSequenceDetails::GenerateAnimationTrackNameArrayElementWidget(TSharedR
 			.AutoWidth()
 			[
 				SNew(SEditableTextBox)
-					.Text(FText::FromName(AnimationTrackNamesList[ArrayIndex]))
+					.Text_Lambda([this, ArrayIndex]()
+					{
+						if (AnimationTrackNamesList.IsValidIndex(ArrayIndex))
+						{
+							return FText::FromName(AnimationTrackNamesList[ArrayIndex]);
+						}
+						
+						return FText::GetEmpty();
+					})
 					.IsReadOnly(true)
 			]
 	];
