Tweaked the details customization for Timecode objects to update the read-only status via lambda instead of once upon construction (so the customization works with InlineEditConditionToggle properly).

diff --git a/Engine/Source/Editor/DetailCustomizations/Private/TimecodeDetailsCustomization.cpp b/Engine/Source/Editor/DetailCustomizations/Private/TimecodeDetailsCustomization.cpp
--- a/Engine/Source/Editor/DetailCustomizations/Private/TimecodeDetailsCustomization.cpp
+++ b/Engine/Source/Editor/DetailCustomizations/Private/TimecodeDetailsCustomization.cpp
@@ -46,7 +46,10 @@ void FTimecodeDetailsCustomization::CustomizeChildren(TSharedRef<IPropertyHandle
 			.SelectAllTextWhenFocused(true)
 			.RevertTextOnEscape(true)
 			.Font(IDetailLayoutBuilder::GetDetailFont())
-			.IsReadOnly(!PropertyHandle->IsEditable())
+			.IsReadOnly_Lambda([this]()
+			{
+				return TimecodeProperty.IsValid() ? !TimecodeProperty->IsEditable() : false;
+			})
 		];
 }
 
