Before adding CallInEditor buttons for properties in the details panel, check the permissions system to make sure they should be added. This fixes an issue where properties could be erroneously included or excluded from the details panel based on another property's permissions.

#rnx

diff --git a/Engine/Source/Editor/DetailCustomizations/Private/ObjectDetails.cpp b/Engine/Source/Editor/DetailCustomizations/Private/ObjectDetails.cpp
--- a/Engine/Source/Editor/DetailCustomizations/Private/ObjectDetails.cpp
+++ b/Engine/Source/Editor/DetailCustomizations/Private/ObjectDetails.cpp
@@ -24,6 +24,7 @@
 #include "Misc/Attribute.h"
 #include "Misc/CString.h"
 #include "ObjectEditorUtils.h"
+#include "PropertyPermissionList.h"
 #include "Reflection/FunctionUtils.h"
 #include "Settings/BlueprintEditorProjectSettings.h"
 #include "SWarningOrErrorBox.h"
@@ -146,9 +147,12 @@ void FObjectDetails::AddCallInEditorMethods(IDetailLayoutBuilder& DetailBuilder)
 			if (bAllowFunction)
 			{
 				const FName FunctionName = TestFunction->GetFName();
-				if (!CallInEditorFunctions.FindByPredicate([&FunctionName](const UFunction* Func) { return Func->GetFName() == FunctionName; }))
+				if (FPropertyEditorPermissionList::Get().DoesPropertyPassFilter(TestFunction->GetOwnerClass(), FunctionName))
 				{
-					CallInEditorFunctions.Add(*FunctionIter);
+					if (!CallInEditorFunctions.FindByPredicate([&FunctionName](const UFunction* Func) { return Func->GetFName() == FunctionName; }))
+					{
+						CallInEditorFunctions.Add(*FunctionIter);
+					}
 				}
 			}
 		}
