FDateTime structs in editor: Fix IsEnabled to be updated when the datetime's IsEditConst condition changes.
#rnx

diff --git a/Engine/Source/Editor/DetailCustomizations/Private/DateTimeStructCustomization.cpp b/Engine/Source/Editor/DetailCustomizations/Private/DateTimeStructCustomization.cpp
--- a/Engine/Source/Editor/DetailCustomizations/Private/DateTimeStructCustomization.cpp
+++ b/Engine/Source/Editor/DetailCustomizations/Private/DateTimeStructCustomization.cpp
@@ -49,7 +49,7 @@ void FDateTimeStructCustomization::CustomizeHeader(TSharedRef<IPropertyHandle> S
 		[
 			SNew(SEditableTextBox)
 				.ClearKeyboardFocusOnCommit(false)
-				.IsEnabled(!PropertyHandle->IsEditConst())
+				.IsEnabled(StructPropertyHandle, &IPropertyHandle::IsEditable)
 				.ForegroundColor(this, &FDateTimeStructCustomization::HandleTextBoxForegroundColor)
 				.OnTextChanged(this, &FDateTimeStructCustomization::HandleTextBoxTextChanged)
 				.OnTextCommitted(this, &FDateTimeStructCustomization::HandleTextBoxTextCommited)
