[Details Panel] Optimization to GuidStructCustomization
* Defer creation of menu widgets until they are required by user input
* For arrays of structs with Guids, this the creation of the menu widgets can constitute a sizable portion of processing.

diff --git a/Engine/Source/Editor/DetailCustomizations/Private/GuidStructCustomization.cpp b/Engine/Source/Editor/DetailCustomizations/Private/GuidStructCustomization.cpp
--- a/Engine/Source/Editor/DetailCustomizations/Private/GuidStructCustomization.cpp
+++ b/Engine/Source/Editor/DetailCustomizations/Private/GuidStructCustomization.cpp
@@ -50,22 +50,21 @@ void FGuidStructCustomization::CustomizeHeader( TSharedRef<class IPropertyHandle
 	}
 	else
 	{
-		FMenuBuilder QuickSetMenuBuilder(true, nullptr);
-		{
-			FUIAction GenerateAction(FExecuteAction::CreateSP(this, &FGuidStructCustomization::HandleGuidActionClicked, EPropertyEditorGuidActions::Generate));
-			QuickSetMenuBuilder.AddMenuEntry(LOCTEXT("GenerateAction", "Generate"), LOCTEXT("GenerateActionHint", "Generate a new random globally unique identifier (GUID)."), FSlateIcon(), GenerateAction);
-
-			FUIAction InvalidateAction(FExecuteAction::CreateSP(this, &FGuidStructCustomization::HandleGuidActionClicked, EPropertyEditorGuidActions::Invalidate));
-			QuickSetMenuBuilder.AddMenuEntry(LOCTEXT("InvalidateAction", "Invalidate"), LOCTEXT("InvalidateActionHint", "Set an invalid globally unique identifier (GUID)."), FSlateIcon(), InvalidateAction);
-		}
-
 		QuickSetSlotContent = 
 			SNew(SComboButton)
 			.ComboButtonStyle(FAppStyle::Get(), "SimpleComboButton")
-			.MenuContent()
-			[
-				QuickSetMenuBuilder.MakeWidget()
-			];
+			.OnGetMenuContent_Lambda([this]()
+			{
+				FMenuBuilder QuickSetMenuBuilder(true, nullptr);
+				{
+					FUIAction GenerateAction(FExecuteAction::CreateSP(this, &FGuidStructCustomization::HandleGuidActionClicked, EPropertyEditorGuidActions::Generate));
+					QuickSetMenuBuilder.AddMenuEntry(LOCTEXT("GenerateAction", "Generate"), LOCTEXT("GenerateActionHint", "Generate a new random globally unique identifier (GUID)."), FSlateIcon(), GenerateAction);
+
+					FUIAction InvalidateAction(FExecuteAction::CreateSP(this, &FGuidStructCustomization::HandleGuidActionClicked, EPropertyEditorGuidActions::Invalidate));
+					QuickSetMenuBuilder.AddMenuEntry(LOCTEXT("InvalidateAction", "Invalidate"), LOCTEXT("InvalidateActionHint", "Set an invalid globally unique identifier (GUID)."), FSlateIcon(), InvalidateAction);
+				}
+				return QuickSetMenuBuilder.MakeWidget();
+			});
 	}
 
 	// create struct header
