Fix unintended copy of FSelectedActorInfo in MakeTransformDetails
Since of the operand in the ternary statement was a temporary object, we held a const reference to a temporary copy rather than a reference to the object pointed by SelectedActorInfoPtr

diff --git a/Engine/Source/Editor/DetailCustomizations/Private/SceneComponentDetails.cpp b/Engine/Source/Editor/DetailCustomizations/Private/SceneComponentDetails.cpp
--- a/Engine/Source/Editor/DetailCustomizations/Private/SceneComponentDetails.cpp
+++ b/Engine/Source/Editor/DetailCustomizations/Private/SceneComponentDetails.cpp
@@ -349,7 +349,9 @@ void FSceneComponentDetails::MakeTransformDetails( IDetailLayoutBuilder& DetailB
 			}
 		}
 
-		const FSelectedActorInfo& SelectedActorInfo = SelectedActorInfoPtr ? *SelectedActorInfoPtr : FSelectedActorInfo();
+		static const FSelectedActorInfo EmptyActorInfo;
+
+		const FSelectedActorInfo& SelectedActorInfo = SelectedActorInfoPtr ? *SelectedActorInfoPtr : EmptyActorInfo;
 		TSharedRef<FComponentTransformDetails> TransformDetails = MakeShareable( new FComponentTransformDetails( bIsEditingBlueprintDefaults ? SceneComponentObjects : DetailBuilder.GetSelectedObjects(), SelectedActorInfo, DetailBuilder ) );
 
 		if(!bShouldShowTransform)
