Fix the padding on the legacy input deprecation warning icon to be aligned correctly.

#rnx

diff --git a/Engine/Source/Editor/DetailCustomizations/Private/InputSettingsDetails.cpp b/Engine/Source/Editor/DetailCustomizations/Private/InputSettingsDetails.cpp
--- a/Engine/Source/Editor/DetailCustomizations/Private/InputSettingsDetails.cpp
+++ b/Engine/Source/Editor/DetailCustomizations/Private/InputSettingsDetails.cpp
@@ -664,7 +664,8 @@ void FInputSettingsDetails::CustomizeDetails(class IDetailLayoutBuilder& DetailB
 				SNew(SHorizontalBox)
 				+ SHorizontalBox::Slot()
 				.AutoWidth()
-				.Padding(10.0f, 10.0f)
+				.Padding(10.0f, 0.0f)
+				.VAlign(VAlign_Bottom)
 				[
 					SNew(SImage)
 					.Image(WarningBrush)
