[DetailCustomizations] Curve Editor crash fix

#rnx

diff --git a/Engine/Source/Editor/DetailCustomizations/Private/CurveColorCustomization.cpp b/Engine/Source/Editor/DetailCustomizations/Private/CurveColorCustomization.cpp
--- a/Engine/Source/Editor/DetailCustomizations/Private/CurveColorCustomization.cpp
+++ b/Engine/Source/Editor/DetailCustomizations/Private/CurveColorCustomization.cpp
@@ -462,41 +462,37 @@ FReply FCurveColorCustomization::OnCurvePreviewDoubleClick(const FGeometry& InMy
 		{
 			DestroyPopOutWindow();
 
-			// Check that the owner exists before spawning the window
-			if (Owner)
-			{
-				// Determine the position of the window so that it will spawn near the mouse, but not go off the screen.
-				const FVector2D CursorPos = FSlateApplication::Get().GetCursorPos();
-				FSlateRect Anchor(CursorPos.X, CursorPos.Y, CursorPos.X, CursorPos.Y);
-
-				FVector2D AdjustedSummonLocation = FSlateApplication::Get().CalculatePopupWindowPosition( Anchor, FCurveColorCustomization::DEFAULT_WINDOW_SIZE, true, FVector2D::ZeroVector, Orient_Horizontal );
-
-				TSharedPtr<SWindow> Window = SNew(SWindow)
-					.Title( FText::Format( LOCTEXT("WindowHeader", "{0} - Internal Color Curve Editor"), StructPropertyHandle->GetPropertyDisplayName()) )
-					.ClientSize( FCurveColorCustomization::DEFAULT_WINDOW_SIZE )
-					.ScreenPosition(AdjustedSummonLocation)
-					.AutoCenter(EAutoCenter::None)
-					.SupportsMaximize(false)
-					.SupportsMinimize(false)
-					.SizingRule( ESizingRule::FixedSize );
-
-				// init the mini curve editor widget
-				TSharedRef<SMiniCurveEditor> MiniCurveEditor =
-					SNew(SMiniCurveEditor)
-					.CurveOwner(this)
-					.OwnerObject(Owner)
-					.ParentWindow(Window);
-
-				Window->SetContent( MiniCurveEditor );
-
-				// Find the window of the parent widget
-				FWidgetPath WidgetPath;
-				FSlateApplication::Get().GeneratePathToWidgetChecked( CurveWidget.ToSharedRef(), WidgetPath );
-				Window = FSlateApplication::Get().AddWindowAsNativeChild( Window.ToSharedRef(), WidgetPath.GetWindow() );
-
-				//hold on to the window created for external use...
-				CurveEditorWindow = Window;
-			}
+			// Determine the position of the window so that it will spawn near the mouse, but not go off the screen.
+			const FVector2D CursorPos = FSlateApplication::Get().GetCursorPos();
+			FSlateRect Anchor(CursorPos.X, CursorPos.Y, CursorPos.X, CursorPos.Y);
+
+			FVector2D AdjustedSummonLocation = FSlateApplication::Get().CalculatePopupWindowPosition( Anchor, FCurveColorCustomization::DEFAULT_WINDOW_SIZE, true, FVector2D::ZeroVector, Orient_Horizontal );
+
+			TSharedPtr<SWindow> Window = SNew(SWindow)
+				.Title( FText::Format( LOCTEXT("WindowHeader", "{0} - Internal Color Curve Editor"), StructPropertyHandle->GetPropertyDisplayName()) )
+				.ClientSize( FCurveColorCustomization::DEFAULT_WINDOW_SIZE )
+				.ScreenPosition(AdjustedSummonLocation)
+				.AutoCenter(EAutoCenter::None)
+				.SupportsMaximize(false)
+				.SupportsMinimize(false)
+				.SizingRule( ESizingRule::FixedSize );
+
+			// init the mini curve editor widget
+			TSharedRef<SMiniCurveEditor> MiniCurveEditor =
+				SNew(SMiniCurveEditor)
+				.CurveOwner(this)
+				.OwnerObject(Owner)
+				.ParentWindow(Window);
+
+			Window->SetContent( MiniCurveEditor );
+
+			// Find the window of the parent widget
+			FWidgetPath WidgetPath;
+			FSlateApplication::Get().GeneratePathToWidgetChecked( CurveWidget.ToSharedRef(), WidgetPath );
+			Window = FSlateApplication::Get().AddWindowAsNativeChild( Window.ToSharedRef(), WidgetPath.GetWindow() );
+
+			//hold on to the window created for external use...
+			CurveEditorWindow = Window;
 		}
 	}
 	return FReply::Handled();
diff --git a/Engine/Source/Editor/DetailCustomizations/Private/CurveStructCustomization.cpp b/Engine/Source/Editor/DetailCustomizations/Private/CurveStructCustomization.cpp
--- a/Engine/Source/Editor/DetailCustomizations/Private/CurveStructCustomization.cpp
+++ b/Engine/Source/Editor/DetailCustomizations/Private/CurveStructCustomization.cpp
@@ -258,14 +258,26 @@ void FCurveStructCustomization::CustomizeChildren( TSharedRef<IPropertyHandle> I
 TArray<FRichCurveEditInfoConst> FCurveStructCustomization::GetCurves() const
 {
 	TArray<FRichCurveEditInfoConst> Curves;
-	Curves.Add(FRichCurveEditInfoConst(&RuntimeCurve->EditorCurveData));
+	FRichCurveEditInfoConst CurveEditInfo = FRichCurveEditInfoConst(&RuntimeCurve->EditorCurveData);
+	if (StructPropertyHandle.IsValid())
+	{
+		// Use the DisplayName of the property in place of the CurveName (which is None after creation above)
+		CurveEditInfo.CurveName = FName(StructPropertyHandle->GetPropertyDisplayName().ToString());
+	}
+	Curves.Add(CurveEditInfo);
 	return Curves;
 }
 
 TArray<FRichCurveEditInfo> FCurveStructCustomization::GetCurves()
 {
 	TArray<FRichCurveEditInfo> Curves;
-	Curves.Add(FRichCurveEditInfo(&RuntimeCurve->EditorCurveData));
+	FRichCurveEditInfo CurveEditInfo = FRichCurveEditInfo(&RuntimeCurve->EditorCurveData);
+	if (StructPropertyHandle.IsValid())
+	{
+		// Use the DisplayName of the property in place of the CurveName (which is None after creation above)
+		CurveEditInfo.CurveName = FName(StructPropertyHandle->GetPropertyDisplayName().ToString());
+	}
+	Curves.Add(CurveEditInfo);
 	return Curves;
 }
 
@@ -459,41 +471,37 @@ FReply FCurveStructCustomization::OnCurvePreviewDoubleClick( const FGeometry& In
 		{
 			DestroyPopOutWindow();
 
-			// Check that the owner exists before spawning a window
-			if (Owner)
-			{
-				// Determine the position of the window so that it will spawn near the mouse, but not go off the screen.
-				const FVector2D CursorPos = FSlateApplication::Get().GetCursorPos();
-				FSlateRect Anchor(CursorPos.X, CursorPos.Y, CursorPos.X, CursorPos.Y);
-
-				FVector2D AdjustedSummonLocation = FSlateApplication::Get().CalculatePopupWindowPosition( Anchor, FCurveStructCustomization::DEFAULT_WINDOW_SIZE, true, FVector2D::ZeroVector, Orient_Horizontal );
-
-				TSharedPtr<SWindow> Window = SNew(SWindow)
-					.Title( FText::Format( LOCTEXT("WindowHeader", "{0} - Internal Curve Editor"), StructPropertyHandle->GetPropertyDisplayName()) )
-					.ClientSize( FCurveStructCustomization::DEFAULT_WINDOW_SIZE )
-					.ScreenPosition(AdjustedSummonLocation)
-					.AutoCenter(EAutoCenter::None)
-					.SupportsMaximize(false)
-					.SupportsMinimize(false)
-					.SizingRule( ESizingRule::FixedSize );
-
-				// init the mini curve editor widget
-				TSharedRef<SMiniCurveEditor> MiniCurveEditor =
-					SNew(SMiniCurveEditor)
-					.CurveOwner(this)
-					.OwnerObject(Owner)
-					.ParentWindow(Window);
-
-				Window->SetContent( MiniCurveEditor );
-
-				// Find the window of the parent widget
-				FWidgetPath WidgetPath;
-				FSlateApplication::Get().GeneratePathToWidgetChecked( CurveWidget.ToSharedRef(), WidgetPath );
-				Window = FSlateApplication::Get().AddWindowAsNativeChild( Window.ToSharedRef(), WidgetPath.GetWindow() );
-
-				//hold on to the window created for external use...
-				CurveEditorWindow = Window;
-			}
+			// Determine the position of the window so that it will spawn near the mouse, but not go off the screen.
+			const FVector2D CursorPos = FSlateApplication::Get().GetCursorPos();
+			FSlateRect Anchor(CursorPos.X, CursorPos.Y, CursorPos.X, CursorPos.Y);
+
+			FVector2D AdjustedSummonLocation = FSlateApplication::Get().CalculatePopupWindowPosition( Anchor, FCurveStructCustomization::DEFAULT_WINDOW_SIZE, true, FVector2D::ZeroVector, Orient_Horizontal );
+
+			TSharedPtr<SWindow> Window = SNew(SWindow)
+				.Title( FText::Format( LOCTEXT("WindowHeader", "{0} - Internal Curve Editor"), StructPropertyHandle->GetPropertyDisplayName()) )
+				.ClientSize( FCurveStructCustomization::DEFAULT_WINDOW_SIZE )
+				.ScreenPosition(AdjustedSummonLocation)
+				.AutoCenter(EAutoCenter::None)
+				.SupportsMaximize(false)
+				.SupportsMinimize(false)
+				.SizingRule( ESizingRule::FixedSize );
+
+			// init the mini curve editor widget
+			TSharedRef<SMiniCurveEditor> MiniCurveEditor =
+				SNew(SMiniCurveEditor)
+				.CurveOwner(this)
+				.OwnerObject(Owner)
+				.ParentWindow(Window);
+
+			Window->SetContent( MiniCurveEditor );
+
+			// Find the window of the parent widget
+			FWidgetPath WidgetPath;
+			FSlateApplication::Get().GeneratePathToWidgetChecked( CurveWidget.ToSharedRef(), WidgetPath );
+			Window = FSlateApplication::Get().AddWindowAsNativeChild( Window.ToSharedRef(), WidgetPath.GetWindow() );
+
+			//hold on to the window created for external use...
+			CurveEditorWindow = Window;
 		}
 	}
 	return FReply::Handled();
diff --git a/Engine/Source/Editor/DetailCustomizations/Private/CurveVectorCustomization.cpp b/Engine/Source/Editor/DetailCustomizations/Private/CurveVectorCustomization.cpp
--- a/Engine/Source/Editor/DetailCustomizations/Private/CurveVectorCustomization.cpp
+++ b/Engine/Source/Editor/DetailCustomizations/Private/CurveVectorCustomization.cpp
@@ -446,41 +446,37 @@ FReply FCurveVectorCustomization::OnCurvePreviewDoubleClick(const FGeometry& InM
 		{
 			DestroyPopOutWindow();
 
-			// Check that the owner exists before spawning the window
-			if (Owner)
-			{
-				// Determine the position of the window so that it will spawn near the mouse, but not go off the screen.
-				const FVector2D CursorPos = FSlateApplication::Get().GetCursorPos();
-				FSlateRect Anchor(CursorPos.X, CursorPos.Y, CursorPos.X, CursorPos.Y);
-
-				FVector2D AdjustedSummonLocation = FSlateApplication::Get().CalculatePopupWindowPosition(Anchor, FCurveVectorCustomization::DEFAULT_WINDOW_SIZE, true, FVector2D::ZeroVector, Orient_Horizontal);
-
-				TSharedPtr<SWindow> Window = SNew(SWindow)
-					.Title(FText::Format(LOCTEXT("WindowHeader", "{0} - Internal Vector Curve Editor"), StructPropertyHandle->GetPropertyDisplayName()))
-					.ClientSize(FCurveVectorCustomization::DEFAULT_WINDOW_SIZE)
-					.ScreenPosition(AdjustedSummonLocation)
-					.AutoCenter(EAutoCenter::None)
-					.SupportsMaximize(false)
-					.SupportsMinimize(false)
-					.SizingRule( ESizingRule::FixedSize );
-
-				// init the mini curve editor widget
-				TSharedRef<SMiniCurveEditor> MiniCurveEditor =
-					SNew(SMiniCurveEditor)
-					.CurveOwner(this)
-					.OwnerObject(Owner)
-					.ParentWindow(Window);
-
-				Window->SetContent( MiniCurveEditor );
-
-				// Find the window of the parent widget
-				FWidgetPath WidgetPath;
-				FSlateApplication::Get().GeneratePathToWidgetChecked( CurveWidget.ToSharedRef(), WidgetPath );
-				Window = FSlateApplication::Get().AddWindowAsNativeChild( Window.ToSharedRef(), WidgetPath.GetWindow() );
-
-				//hold on to the window created for external use...
-				CurveEditorWindow = Window;
-			}
+			// Determine the position of the window so that it will spawn near the mouse, but not go off the screen.
+			const FVector2D CursorPos = FSlateApplication::Get().GetCursorPos();
+			FSlateRect Anchor(CursorPos.X, CursorPos.Y, CursorPos.X, CursorPos.Y);
+
+			FVector2D AdjustedSummonLocation = FSlateApplication::Get().CalculatePopupWindowPosition(Anchor, FCurveVectorCustomization::DEFAULT_WINDOW_SIZE, true, FVector2D::ZeroVector, Orient_Horizontal);
+
+			TSharedPtr<SWindow> Window = SNew(SWindow)
+				.Title(FText::Format(LOCTEXT("WindowHeader", "{0} - Internal Vector Curve Editor"), StructPropertyHandle->GetPropertyDisplayName()))
+				.ClientSize(FCurveVectorCustomization::DEFAULT_WINDOW_SIZE)
+				.ScreenPosition(AdjustedSummonLocation)
+				.AutoCenter(EAutoCenter::None)
+				.SupportsMaximize(false)
+				.SupportsMinimize(false)
+				.SizingRule( ESizingRule::FixedSize );
+
+			// init the mini curve editor widget
+			TSharedRef<SMiniCurveEditor> MiniCurveEditor =
+				SNew(SMiniCurveEditor)
+				.CurveOwner(this)
+				.OwnerObject(Owner)
+				.ParentWindow(Window);
+
+			Window->SetContent( MiniCurveEditor );
+
+			// Find the window of the parent widget
+			FWidgetPath WidgetPath;
+			FSlateApplication::Get().GeneratePathToWidgetChecked( CurveWidget.ToSharedRef(), WidgetPath );
+			Window = FSlateApplication::Get().AddWindowAsNativeChild( Window.ToSharedRef(), WidgetPath.GetWindow() );
+
+			//hold on to the window created for external use...
+			CurveEditorWindow = Window;
 		}
 	}
 	return FReply::Handled();
diff --git a/Engine/Source/Editor/UnrealEd/Private/MiniCurveEditor.cpp b/Engine/Source/Editor/UnrealEd/Private/MiniCurveEditor.cpp
--- a/Engine/Source/Editor/UnrealEd/Private/MiniCurveEditor.cpp
+++ b/Engine/Source/Editor/UnrealEd/Private/MiniCurveEditor.cpp
@@ -1,6 +1,5 @@
 // Copyright Epic Games, Inc. All Rights Reserved.
 
-
 #include "MiniCurveEditor.h"
 #include "Widgets/Layout/SBorder.h"
 #include "Styling/AppStyle.h"
@@ -8,7 +7,6 @@
 #include "Subsystems/AssetEditorSubsystem.h"
 #include "Editor.h"
 
-
 void SMiniCurveEditor::Construct(const FArguments& InArgs)
 {
 	ViewMinInput=0.f;
@@ -35,12 +33,21 @@ void SMiniCurveEditor::Construct(const FArguments& InArgs)
 
 	WidgetWindow = InArgs._ParentWindow;
 
-	GEditor->GetEditorSubsystem<UAssetEditorSubsystem>()->NotifyAssetOpened(InArgs._OwnerObject, this);
+	bHasOwner = InArgs._OwnerObject != nullptr;
+
+	// Only notify if the Owner is present - it's not when spawned from the DataTable editor, for example
+	if (bHasOwner)
+	{
+		GEditor->GetEditorSubsystem<UAssetEditorSubsystem>()->NotifyAssetOpened(InArgs._OwnerObject, this);
+	}
 }
 
 SMiniCurveEditor::~SMiniCurveEditor()
 {
-	GEditor->GetEditorSubsystem<UAssetEditorSubsystem>()->NotifyEditorClosed(this);
+	if (bHasOwner)
+	{
+		GEditor->GetEditorSubsystem<UAssetEditorSubsystem>()->NotifyEditorClosed(this);
+	}
 }
 
 float SMiniCurveEditor::GetTimelineLength() const
@@ -48,7 +55,6 @@ float SMiniCurveEditor::GetTimelineLength() const
 	return 0.f;
 }
 
-
 void SMiniCurveEditor::SetInputViewRange(float InViewMinInput, float InViewMaxInput)
 {
 	ViewMaxInput = InViewMaxInput;
diff --git a/Engine/Source/Editor/UnrealEd/Public/MiniCurveEditor.h b/Engine/Source/Editor/UnrealEd/Public/MiniCurveEditor.h
--- a/Engine/Source/Editor/UnrealEd/Public/MiniCurveEditor.h
+++ b/Engine/Source/Editor/UnrealEd/Public/MiniCurveEditor.h
@@ -2,31 +2,30 @@
 
 #pragma once
 
-#include "CoreMinimal.h"
 #include "Widgets/DeclarativeSyntaxSupport.h"
 #include "Widgets/SCompoundWidget.h"
 #include "Framework/Docking/TabManager.h"
 #include "Subsystems/AssetEditorSubsystem.h"
 
-
 class FCurveOwnerInterface;
 class SCurveEditor;
 
-class SMiniCurveEditor :  public SCompoundWidget,public IAssetEditorInstance
+class SMiniCurveEditor
+	: public SCompoundWidget
+	, public IAssetEditorInstance
 {
 public:
 	SLATE_BEGIN_ARGS( SMiniCurveEditor )
 		: _CurveOwner(nullptr)
 		, _OwnerObject(nullptr)
 		{}
-
-	SLATE_ARGUMENT( FCurveOwnerInterface*, CurveOwner )
-	SLATE_ARGUMENT( UObject*, OwnerObject )
-	SLATE_ARGUMENT( TWeakPtr<SWindow>, ParentWindow )
+		SLATE_ARGUMENT( FCurveOwnerInterface*, CurveOwner )
+		SLATE_ARGUMENT( UObject*, OwnerObject )
+		SLATE_ARGUMENT( TWeakPtr<SWindow>, ParentWindow )
 	SLATE_END_ARGS()
 
 	UNREALED_API void Construct(const FArguments& InArgs);
-	UNREALED_API virtual ~SMiniCurveEditor();
+	UNREALED_API virtual ~SMiniCurveEditor() override;
 
 	// IAssetEditorInstance interface
 	UNREALED_API virtual FName GetEditorName() const override;
@@ -40,20 +39,22 @@ public:
 	UNREALED_API virtual void RemoveEditingAsset(UObject* Asset) override;
 
 private:
+	/** Sometimes false when this is a transient or popout window, indicating it shouldn't be tracked as an asset editor. */
+	bool bHasOwner = true;
 
-	float ViewMinInput;
-	float ViewMaxInput;
+	float ViewMinInput = 0.0f;
+	float ViewMaxInput = 1.0f;
 
 	TSharedPtr<class SCurveEditor> TrackWidget;
 
 	float GetViewMinInput() const { return ViewMinInput; }
 	float GetViewMaxInput() const { return ViewMaxInput; }
+
 	/** Return length of timeline */
 	UNREALED_API float GetTimelineLength() const;
 
 	UNREALED_API void SetInputViewRange(float InViewMinInput, float InViewMaxInput);
 
-
 protected:
 	TWeakPtr<SWindow> WidgetWindow;
 };
