[DetailCustomizations] CallInEditor fix

diff --git a/Engine/Source/Editor/DetailCustomizations/Private/ObjectDetails.cpp b/Engine/Source/Editor/DetailCustomizations/Private/ObjectDetails.cpp
--- a/Engine/Source/Editor/DetailCustomizations/Private/ObjectDetails.cpp
+++ b/Engine/Source/Editor/DetailCustomizations/Private/ObjectDetails.cpp
@@ -92,6 +92,9 @@ void FObjectDetails::AddExperimentalWarningCategory(IDetailLayoutBuilder& Detail
 
 void FObjectDetails::AddCallInEditorMethods(IDetailLayoutBuilder& DetailBuilder)
 {
+	DetailBuilder.GetObjectsBeingCustomized(/*out*/ SelectedObjectsList);
+	SelectedObjectsList.RemoveAllSwap([](TWeakObjectPtr<UObject> ObjPtr) { UObject* Obj = ObjPtr.Get(); return (Obj == nullptr) || Obj->HasAnyFlags(RF_ArchetypeObject); });
+
 	TArray<UFunction*> CallInEditorFunctions;
 	PropertyCustomizationHelpers::GetCallInEditorFunctionsForClass(
 		DetailBuilder.GetBaseClass(),
