Check row data before re-interpret casting to FFrameRate

diff --git a/Engine/Source/Editor/DetailCustomizations/Private/FrameRateCustomization.cpp b/Engine/Source/Editor/DetailCustomizations/Private/FrameRateCustomization.cpp
--- a/Engine/Source/Editor/DetailCustomizations/Private/FrameRateCustomization.cpp
+++ b/Engine/Source/Editor/DetailCustomizations/Private/FrameRateCustomization.cpp
@@ -82,9 +82,11 @@ void FFrameRateCustomization::SetFrameRate(FFrameRate NewFrameRate)
 {
 	if (FStructProperty* StructProperty = CastField<FStructProperty>(StructPropertyHandle->GetProperty()))
 	{
+		static FFrameRate DefaultRate(24,1);
+
 		TArray<void*> RawData;
 		StructPropertyHandle->AccessRawData(RawData);
-		FFrameRate* PreviousFrameRate = reinterpret_cast<FFrameRate*>(RawData[0]);
+		FFrameRate* PreviousFrameRate = RawData.Num() == 1 ? reinterpret_cast<FFrameRate*>(RawData[0]) : &DefaultRate;
 
 		FString TextValue;
 		StructProperty->Struct->ExportText(TextValue, &NewFrameRate, PreviousFrameRate, nullptr, EPropertyPortFlags::PPF_None, nullptr);
@@ -125,4 +127,4 @@ bool FFrameRateCustomization::HasMultipleValues() const
 	return false;
 }
 
-#undef LOCTEXT_NAMESPACE
\ No newline at end of file
+#undef LOCTEXT_NAMESPACE
