[ColorGrading] [ColorGradingEditor] Update color grading UI with new styling that can scale with available space

#virtualized

diff --git a/Engine/Plugins/Editor/ColorGrading/Source/ColorGradingEditor/ColorGradingEditor.build.cs b/Engine/Plugins/Editor/ColorGrading/Source/ColorGradingEditor/ColorGradingEditor.build.cs
--- a/Engine/Plugins/Editor/ColorGrading/Source/ColorGradingEditor/ColorGradingEditor.build.cs
+++ b/Engine/Plugins/Editor/ColorGrading/Source/ColorGradingEditor/ColorGradingEditor.build.cs
@@ -10,6 +10,7 @@ public class ColorGradingEditor : ModuleRules
 		PublicDependencyModuleNames.AddRange(
 			new string[]
 			{
+				"AdvancedWidgets",
 				"SceneOutliner",
 				"ObjectMixerEditor"
 			}
diff --git a/Engine/Plugins/Editor/ColorGrading/Source/ColorGradingEditor/Private/SColorGradingColorWheel.cpp b/Engine/Plugins/Editor/ColorGrading/Source/ColorGradingEditor/Private/SColorGradingColorWheel.cpp
--- a/Engine/Plugins/Editor/ColorGrading/Source/ColorGradingEditor/Private/SColorGradingColorWheel.cpp
+++ b/Engine/Plugins/Editor/ColorGrading/Source/ColorGradingEditor/Private/SColorGradingColorWheel.cpp
@@ -4,42 +4,80 @@
 
 #include "Customizations/MathStructCustomizations.h"
 #include "DetailLayoutBuilder.h"
-#include "Editor.h"
 #include "PropertyHandle.h"
 #include "ScopedTransaction.h"
+#include "Util/ColorGradingUtil.h"
 #include "Widgets/SBoxPanel.h"
+#include "Widgets/ColorGrading/SColorGradingComponentViewer.h"
 #include "Widgets/Colors/SComplexGradient.h"
 #include "Widgets/Layout/SBox.h"
 
+#if WITH_EDITOR
+#include "Editor.h"
+#endif
+
 #define LOCTEXT_NAMESPACE "ColorGradingEditor"
 
+SColorGradingColorWheel::SColorGradingColorWheel()
+{
+#if WITH_EDITOR
+	if (GEditor)
+	{
+		GEditor->RegisterForUndo(this);
+	}
+#endif
+}
+
+SColorGradingColorWheel::~SColorGradingColorWheel()
+{
+#if WITH_EDITOR
+	if (GEditor)
+	{
+		GEditor->UnregisterForUndo(this);
+	}
+#endif
+}
+
 void SColorGradingColorWheel::Construct(const FArguments& InArgs)
 {
-	Orientation = InArgs._Orientation;
 	ColorDisplayMode = InArgs._ColorDisplayMode;
 
+	ColorPickerBox = SNew(SBox);
+	
+	ColorSlidersBox = SNew(SBox)
+		.HAlign(HAlign_Fill)
+		.MaxDesiredWidth(400.f)
+		.MinDesiredWidth(400.f);
+
 	ChildSlot
 	[
-		SNew(SVerticalBox)
-
-		+ SVerticalBox::Slot()
-		.AutoHeight()
-		.HAlign(HAlign_Fill)
+		SNew(SBox)
+		.Padding(16.f, 8.f)
 		[
-			SAssignNew(HeaderBox, SBox)
-		]
+			SNew(SVerticalBox)
 
-		+ SVerticalBox::Slot()
-		.FillHeight(1.0f)
-		[
-			SAssignNew(ColorWheelBox, SBox)
-		]
-	];
+			+ SVerticalBox::Slot()
+			.AutoHeight()
+			.HAlign(HAlign_Fill)
+			[
+				SAssignNew(HeaderBox, SBox)
+			]
 
-	ColorPickerBox = SNew(SBox);
-	ColorSlidersBox = SNew(SBox);
+			+ SVerticalBox::Slot()
+			.FillHeight(1.0f)
+			.HAlign(HAlign_Fill)
+			[
+				ColorPickerBox.ToSharedRef()
+			]
 
-	SetOrientation(Orientation);
+			+ SVerticalBox::Slot()
+			.AutoHeight()
+			.HAlign(HAlign_Center)
+			[
+				ColorSlidersBox.ToSharedRef()
+			]
+		]
+	];
 
 	if (InArgs._HeaderContent.IsValid())
 	{
@@ -63,6 +101,8 @@ void SColorGradingColorWheel::SetColorPropertyHandle(TSharedPtr<IPropertyHandle>
 		ComponentSliderDynamicMaxValue = ColorPropertyMetadata->MaxValue;
 	}
 
+	RecalculateHSVColor();
+
 	// Since many of the color picker slate properties are not attributes, we need to recreate the widget every time the property handle changes
 	// to ensure the picker is configured correctly for the color property
 	if (ColorPickerBox.IsValid())
@@ -84,57 +124,23 @@ void SColorGradingColorWheel::SetHeaderContent(const TSharedRef<SWidget>& Header
 	}
 }
 
-void SColorGradingColorWheel::SetOrientation(EOrientation NewOrientation)
+#if WITH_EDITOR
+void SColorGradingColorWheel::PostUndo(bool bSuccess)
 {
-	Orientation = NewOrientation;
-
-	if (ColorWheelBox.IsValid())
-	{
-		TSharedPtr<SWidget> OrientedBox;
-
-		if (Orientation == EOrientation::Orient_Vertical)
-		{
-			OrientedBox = SNew(SVerticalBox)
-
-			+ SVerticalBox::Slot()
-			.AutoHeight()
-			[
-				ColorPickerBox.ToSharedRef()
-			]
-
-			+ SVerticalBox::Slot()
-			.AutoHeight()
-			[
-				ColorSlidersBox.ToSharedRef()
-			];
-		}
-		else
-		{
-			OrientedBox = SNew(SHorizontalBox)
-
-			+ SHorizontalBox::Slot()
-			.AutoWidth()
-			[
-				ColorPickerBox.ToSharedRef()
-			]
-
-			+ SHorizontalBox::Slot()
-			.FillWidth(1.0f)
-			[
-				ColorSlidersBox.ToSharedRef()
-			];
-		}
+	RecalculateHSVColor();
+}
 
-		ColorWheelBox->SetContent(OrientedBox.ToSharedRef());
-		ColorWheelBox->SetHAlign(Orientation == EOrientation::Orient_Vertical ? HAlign_Center : HAlign_Fill);
-	}
+void SColorGradingColorWheel::PostRedo(bool bSuccess)
+{
+	RecalculateHSVColor();
 }
+#endif
 
 TSharedRef<SWidget> SColorGradingColorWheel::SColorGradingColorWheel::CreateColorGradingPicker()
 {
 	if (ColorPropertyHandle.IsValid() && ColorPropertyHandle.Pin()->IsValidHandle())
 	{
-		return SNew(SColorGradingPicker)
+		return SNew(UE::ColorGrading::SColorGradingPicker)
 			.ValueMin(GetMetadataMinValue())
 			.ValueMax(GetMetadataMaxValue())
 			.SliderValueMin(GetMetadataSliderMinValue())
@@ -167,19 +173,24 @@ TSharedRef<SWidget> SColorGradingColorWheel::CreateColorComponentSliders()
 		uint32 NumComponents = 4;
 		for (uint32 ComponentIndex = 0; ComponentIndex < NumComponents; ++ComponentIndex)
 		{
-			TAttribute<FText> TextGetter = TAttribute<FText>::Create(TAttribute<FText>::FGetter::CreateSP(this, &SColorGradingColorWheel::GetComponentLabelText, ComponentIndex));
-			TSharedRef<SWidget> LabelWidget = SNumericEntryBox<float>::BuildLabel(TextGetter, FLinearColor::White, FLinearColor(0.2f, 0.2f, 0.2f));
-
-			TSharedRef<SNumericEntryBox<float>> NumericEntryBox = SNew(SNumericEntryBox<float>)
-				.SpinBoxStyle(&FCoreStyle::Get().GetWidgetStyle<FSpinBoxStyle>("NumericEntrySpinBox_Dark"))
-				.EditableTextBoxStyle(&FCoreStyle::Get().GetWidgetStyle<FEditableTextBoxStyle>("DarkEditableTextBox"))
-				.Font(IDetailLayoutBuilder::GetDetailFont())
-				.UndeterminedString(NSLOCTEXT("PropertyEditor", "MultipleValues", "Multiple Values"))
+			TAttribute<UE::ColorGrading::EColorGradingComponent> ComponentGetter = TAttribute<UE::ColorGrading::EColorGradingComponent>::Create(
+				TAttribute<UE::ColorGrading::EColorGradingComponent>::FGetter::CreateSP(this, &SColorGradingColorWheel::GetComponent, ComponentIndex)
+			);
+
+			ColorSlidersVerticalBox->AddSlot()
+			.Padding(FMargin(0.0f, 8.0f, 0.0f, 0.0f))
+			.AutoHeight()
+			.HAlign(HAlign_Fill)
+			.VAlign(VAlign_Fill)
+			[
+				SNew(UE::ColorGrading::SColorGradingComponentViewer)
+				.Component(ComponentGetter)
+				.ColorGradingMode(ColorPropertyMetadata->ColorGradingMode)
 				.Value(this, &SColorGradingColorWheel::GetComponentValue, ComponentIndex)
 				.OnValueChanged(this, &SColorGradingColorWheel::SetComponentValue, ComponentIndex)
 				.OnBeginSliderMovement(this, &SColorGradingColorWheel::BeginUsingComponentSlider, ComponentIndex)
 				.OnEndSliderMovement(this, &SColorGradingColorWheel::EndUsingComponentSlider, ComponentIndex)
-				.AllowSpin(ColorPropertyHandle.Pin()->GetNumOuterObjects() == 1)
+				.OnQueryCurrentColor(this, &SColorGradingColorWheel::GetColor)
 				.ShiftMultiplier(ColorPropertyMetadata->ShiftMultiplier)
 				.CtrlMultiplier(ColorPropertyMetadata->CtrlMultiplier)
 				.SupportDynamicSliderMinValue(this, &SColorGradingColorWheel::ComponentSupportsDynamicSliderValue, ColorPropertyMetadata->bSupportDynamicSliderMinValue, ComponentIndex)
@@ -193,37 +204,8 @@ TSharedRef<SWidget> SColorGradingColorWheel::CreateColorComponentSliders()
 				.SliderExponent(ColorPropertyMetadata->SliderExponent)
 				.SliderExponentNeutralValue(ColorPropertyMetadata->SliderMinValue.GetValue() + (ColorPropertyMetadata->SliderMaxValue.GetValue() - ColorPropertyMetadata->SliderMinValue.GetValue()) / 2.0f)
 				.Delta(this, &SColorGradingColorWheel::GetComponentSliderDeltaValue, ColorPropertyMetadata->Delta, ComponentIndex)
-				.ToolTipText(this, &SColorGradingColorWheel::GetComponentToolTipText, ComponentIndex)
-				.LabelPadding(FMargin(0))
 				.IsEnabled(this, &SColorGradingColorWheel::IsPropertyEnabled)
-				.Label()
-				[
-					LabelWidget
-				];
-
-			ColorSlidersVerticalBox->AddSlot()
-				.Padding(FMargin(0.0f, 2.0f, 3.0f, 0.0f))
-				.AutoHeight()
-				.HAlign(HAlign_Fill)
-				.VAlign(VAlign_Fill)
-				[
-					NumericEntryBox
-				];
-
-			ColorSlidersVerticalBox->AddSlot()
-				.Padding(FMargin(15.0f, 0.0f, 3.0f, 2.0f))
-				.AutoHeight()
-				.HAlign(HAlign_Fill)
-				.VAlign(VAlign_Bottom)
-				[
-					SNew(SBox)
-					.HeightOverride(6.0f)
-					[
-						SNew(SComplexGradient)
-						.GradientColors(TAttribute<TArray<FLinearColor>>::Create(TAttribute<TArray<FLinearColor>>::FGetter::CreateSP(this, &SColorGradingColorWheel::GetGradientColor, ComponentIndex)))
-						.Visibility(this, &SColorGradingColorWheel::GetGradientVisibility)
-					]
-				];
+			];
 		}
 
 		return ColorSlidersVerticalBox;
@@ -245,23 +227,23 @@ SColorGradingColorWheel::FColorPropertyMetadata SColorGradingColorWheel::GetColo
 		{
 			if (ColorGradingModeString.Compare(TEXT("saturation")) == 0)
 			{
-				PropertyMetadata.ColorGradingMode = EColorGradingModes::Saturation;
+				PropertyMetadata.ColorGradingMode = UE::ColorGrading::EColorGradingModes::Saturation;
 			}
 			else if (ColorGradingModeString.Compare(TEXT("contrast")) == 0)
 			{
-				PropertyMetadata.ColorGradingMode = EColorGradingModes::Contrast;
+				PropertyMetadata.ColorGradingMode = UE::ColorGrading::EColorGradingModes::Contrast;
 			}
 			else if (ColorGradingModeString.Compare(TEXT("gamma")) == 0)
 			{
-				PropertyMetadata.ColorGradingMode = EColorGradingModes::Gamma;
+				PropertyMetadata.ColorGradingMode = UE::ColorGrading::EColorGradingModes::Gamma;
 			}
 			else if (ColorGradingModeString.Compare(TEXT("gain")) == 0)
 			{
-				PropertyMetadata.ColorGradingMode = EColorGradingModes::Gain;
+				PropertyMetadata.ColorGradingMode = UE::ColorGrading::EColorGradingModes::Gain;
 			}
 			else if (ColorGradingModeString.Compare(TEXT("offset")) == 0)
 			{
-				PropertyMetadata.ColorGradingMode = EColorGradingModes::Offset;
+				PropertyMetadata.ColorGradingMode = UE::ColorGrading::EColorGradingModes::Offset;
 			}
 		}
 
@@ -401,64 +383,73 @@ void SColorGradingColorWheel::TransactColorValue()
 	}
 }
 
+void SColorGradingColorWheel::RecalculateHSVColor()
+{
+	if (ColorPropertyHandle.IsValid() && ColorPropertyHandle.Pin().IsValid())
+	{
+		FVector4 VectorValue;
+		if (ColorPropertyHandle.Pin()->GetValue(VectorValue) == FPropertyAccess::Success)
+		{
+			CurrentHSVColor = FLinearColor(VectorValue).LinearRGBToHSV();
+		}
+	}
+}
+
 void SColorGradingColorWheel::BeginUsingColorPickerSlider()
 {
 	bIsUsingColorPickerSlider = true;
-	GEditor->BeginTransaction(LOCTEXT("ColorWheel_TransactionName", "Color Grading Main Value"));
+
+#if WITH_EDITOR
+	if (GEditor)
+	{
+		GEditor->BeginTransaction(LOCTEXT("ColorWheel_TransactionName", "Color Grading Main Value"));
+	}
+#endif
 }
 
 void SColorGradingColorWheel::EndUsingColorPickerSlider()
 {
 	bIsUsingColorPickerSlider = false;
-	GEditor->EndTransaction();
+
+#if WITH_EDITOR
+	if (GEditor)
+	{
+		GEditor->EndTransaction();
+	}
+#endif
 }
 
 void SColorGradingColorWheel::BeginUsingComponentSlider(uint32 ComponentIndex)
 {
 	bIsUsingComponentSlider = true;
-	GEditor->BeginTransaction(LOCTEXT("ColorWheel_TransactionName", "Color Grading Main Value"));
+
+#if WITH_EDITOR
+	if (GEditor)
+	{
+		GEditor->BeginTransaction(LOCTEXT("ColorWheel_TransactionName", "Color Grading Main Value"));
+	}
+#endif
 }
 
 void SColorGradingColorWheel::EndUsingComponentSlider(float NewValue, uint32 ComponentIndex)
 {
 	bIsUsingComponentSlider = false;
 	SetComponentValue(NewValue, ComponentIndex);
-	GEditor->EndTransaction();
-}
 
-FText SColorGradingColorWheel::GetComponentLabelText(uint32 ComponentIndex) const
-{
-	static const FText RGBLabelText[] = {
-		LOCTEXT("ColorWheel_RedComponentLabel", "R"),
-		LOCTEXT("ColorWheel_GreenComponentLabel", "G"),
-		LOCTEXT("ColorWheel_BlueComponentLabel", "B")
-	};
-
-	static const FText HSVLabelText[] = {
-		LOCTEXT("ColorWheel_HueComponentLabel", "H"),
-		LOCTEXT("ColorWheel_SatComponentLabel", "S"),
-		LOCTEXT("ColorWheel_ValComponentLabel", "V")
-	};
-
-	if (ComponentIndex >= 0 && ComponentIndex < 3)
-	{ 
-		if (ColorDisplayMode.Get(EColorGradingColorDisplayMode::RGB) == EColorGradingColorDisplayMode::RGB)
-		{
-			return RGBLabelText[ComponentIndex];
-		}
-		else
-		{
-			return HSVLabelText[ComponentIndex];
-		}
-	}
-	else if (ComponentIndex == 3)
+#if WITH_EDITOR
+	if (GEditor)
 	{
-		return LOCTEXT("ColorWheel_LuminanceComponentLabel", "Y");
+		GEditor->EndTransaction();
 	}
+#endif
+}
 
-	return FText::GetEmpty();
+UE::ColorGrading::EColorGradingComponent SColorGradingColorWheel::GetComponent(uint32 ComponentIndex) const
+{
+	return UE::ColorGrading::GetColorGradingComponent(ColorDisplayMode.Get(), ComponentIndex);
 }
 
+
 TOptional<float> SColorGradingColorWheel::GetComponentValue(uint32 ComponentIndex) const
 {
 	if (ColorPropertyHandle.IsValid())
@@ -468,14 +459,13 @@ TOptional<float> SColorGradingColorWheel::GetComponentValue(uint32 ComponentInde
 		{
 			float Value = 0.0f;
 
-			if (ColorDisplayMode.Get(EColorGradingColorDisplayMode::RGB) == EColorGradingColorDisplayMode::RGB)
+			if (ColorDisplayMode.Get(UE::ColorGrading::EColorGradingColorDisplayMode::RGB) == UE::ColorGrading::EColorGradingColorDisplayMode::RGB)
 			{
 				Value = ColorValue[ComponentIndex];
 			}
 			else
 			{
-				FLinearColor HSVColor = FLinearColor(ColorValue.X, ColorValue.Y, ColorValue.Z, ColorValue.W).LinearRGBToHSV();
-				Value = HSVColor.Component(ComponentIndex);
+				Value = CurrentHSVColor.Component(ComponentIndex);
 			}
 
 			return Value;
@@ -494,15 +484,19 @@ void SColorGradingColorWheel::SetComponentValue(float NewValue, uint32 Component
 		{
 			FVector4 NewColorValue = CurrentColorValue;
 
-			if (ColorDisplayMode.Get(EColorGradingColorDisplayMode::RGB) == EColorGradingColorDisplayMode::RGB)
+			if (ColorDisplayMode.Get(UE::ColorGrading::EColorGradingColorDisplayMode::RGB) == UE::ColorGrading::EColorGradingColorDisplayMode::RGB)
 			{
 				NewColorValue[ComponentIndex] = NewValue;
+
+				if (ComponentIndex < 3)
+				{
+					CurrentHSVColor = FLinearColor(NewColorValue.X, NewColorValue.Y, NewColorValue.Z).LinearRGBToHSV();
+				}
 			}
 			else
 			{
-				FLinearColor HSVColor = FLinearColor(CurrentColorValue.X, CurrentColorValue.Y, CurrentColorValue.Z, CurrentColorValue.W).LinearRGBToHSV();
-				HSVColor.Component(ComponentIndex) = NewValue;
-				NewColorValue = (FVector4)HSVColor.HSVToLinearRGB();
+				CurrentHSVColor.Component(ComponentIndex) = NewValue;
+				NewColorValue = (FVector4)CurrentHSVColor.HSVToLinearRGB();
 			}
 
 			ColorPropertyHandle.Pin()->SetValue(NewColorValue, bIsUsingComponentSlider ? EPropertyValueSetFlags::InteractiveChange : EPropertyValueSetFlags::DefaultFlags);
@@ -515,7 +509,7 @@ bool SColorGradingColorWheel::ComponentSupportsDynamicSliderValue(bool bDefaultV
 {
 	if (bDefaultValue)
 	{
-		if (ColorDisplayMode.Get(EColorGradingColorDisplayMode::RGB) != EColorGradingColorDisplayMode::RGB)
+		if (ColorDisplayMode.Get(UE::ColorGrading::EColorGradingColorDisplayMode::RGB) != UE::ColorGrading::EColorGradingColorDisplayMode::RGB)
 		{
 			return ComponentIndex >= 2;
 		}
@@ -542,7 +536,7 @@ void SColorGradingColorWheel::UpdateComponentDynamicSliderMaxValue(float NewValu
 
 TOptional<float> SColorGradingColorWheel::GetComponentMaxValue(TOptional<float> DefaultValue, uint32 ComponentIndex) const
 {
-	if (ColorDisplayMode.Get(EColorGradingColorDisplayMode::RGB) != EColorGradingColorDisplayMode::RGB)
+	if (ColorDisplayMode.Get(UE::ColorGrading::EColorGradingColorDisplayMode::RGB) != UE::ColorGrading::EColorGradingColorDisplayMode::RGB)
 	{
 		if (ComponentIndex == 0)
 		{
@@ -559,7 +553,7 @@ TOptional<float> SColorGradingColorWheel::GetComponentMaxValue(TOptional<float>
 
 TOptional<float> SColorGradingColorWheel::GetComponentMinSliderValue(TOptional<float> DefaultValue, uint32 ComponentIndex) const
 {
-	if (ColorDisplayMode.Get(EColorGradingColorDisplayMode::RGB) != EColorGradingColorDisplayMode::RGB)
+	if (ColorDisplayMode.Get(UE::ColorGrading::EColorGradingColorDisplayMode::RGB) != UE::ColorGrading::EColorGradingColorDisplayMode::RGB)
 	{
 		return 0.0f;
 	}
@@ -569,7 +563,7 @@ TOptional<float> SColorGradingColorWheel::GetComponentMinSliderValue(TOptional<f
 
 TOptional<float> SColorGradingColorWheel::GetComponentMaxSliderValue(TOptional<float> DefaultValue, uint32 ComponentIndex) const
 {
-	if (ColorDisplayMode.Get(EColorGradingColorDisplayMode::RGB) != EColorGradingColorDisplayMode::RGB)
+	if (ColorDisplayMode.Get(UE::ColorGrading::EColorGradingColorDisplayMode::RGB) != UE::ColorGrading::EColorGradingColorDisplayMode::RGB)
 	{
 		if (ComponentIndex == 0)
 		{
@@ -586,7 +580,7 @@ TOptional<float> SColorGradingColorWheel::GetComponentMaxSliderValue(TOptional<f
 
 float SColorGradingColorWheel::GetComponentSliderDeltaValue(float DefaultValue, uint32 ComponentIndex) const
 {
-	if (ComponentIndex == 0 && ColorDisplayMode.Get(EColorGradingColorDisplayMode::RGB) != EColorGradingColorDisplayMode::RGB)
+	if (ComponentIndex == 0 && ColorDisplayMode.Get(UE::ColorGrading::EColorGradingColorDisplayMode::RGB) != UE::ColorGrading::EColorGradingColorDisplayMode::RGB)
 	{
 		return 1.0f;
 	}
@@ -594,164 +588,6 @@ float SColorGradingColorWheel::GetComponentSliderDeltaValue(float DefaultValue,
 	return DefaultValue;
 }
 
-FText SColorGradingColorWheel::GetComponentToolTipText(uint32 ComponentIndex) const
-{
-	static const FText RGBToolTipText[] =
-	{
-		LOCTEXT("ColorWheel_RedComponentToolTip", "Red"),
-		LOCTEXT("ColorWheel_GreenComponentToolTip", "Green"),
-		LOCTEXT("ColorWheel_BlueComponentToolTip", "Blue")
-	};
-
-	static const FText HSVToolTipText[] =
-	{
-		LOCTEXT("ColorWheel_HueComponentToolTip", "Hue"),
-		LOCTEXT("ColorWheel_SaturationComponentToolTip", "Saturation"),
-		LOCTEXT("ColorWheel_ValueComponentToolTip", "Value")
-	};
-
-	if (ComponentIndex >= 0 && ComponentIndex < 3)
-	{
-		if (ColorDisplayMode.Get(EColorGradingColorDisplayMode::RGB) != EColorGradingColorDisplayMode::RGB)
-		{
-			return HSVToolTipText[ComponentIndex];
-		}
-		else
-		{
-			return RGBToolTipText[ComponentIndex];
-		}
-	}
-	else if (ComponentIndex == 3)
-	{
-		return LOCTEXT("ColorWheel_LuminanceComponentToolTip", "Luminance");
-	}
-
-	return FText::GetEmpty();
-}
-
-TArray<FLinearColor> SColorGradingColorWheel::GetGradientColor(uint32 ComponentIndex)
-{
-	TArray<FLinearColor> GradientColors;
-
-	if (ColorPropertyHandle.IsValid())
-	{
-		const bool bIsRGB = ColorDisplayMode.Get(EColorGradingColorDisplayMode::RGB) == EColorGradingColorDisplayMode::RGB;
-		FVector4 ColorValue;
-			
-		if (ColorPropertyHandle.Pin()->GetValue(ColorValue) != FPropertyAccess::Success)
-		{
-			ColorValue = FVector4(0.0, 0.0, 0.0, 0.0);
-		}
-
-		if (bIsRGB || ComponentIndex > 0)
-		{
-			GradientColors.Add(GetGradientStartColor(ColorValue, bIsRGB, ComponentIndex));
-			GradientColors.Add(GetGradientEndColor(ColorValue, bIsRGB, ComponentIndex));
-			GradientColors.Add(GetGradientFillerColor(ColorValue, bIsRGB, ComponentIndex));
-		}
-		else // HSV Hue handling
-		{
-			for (int32 i = 0; i < 7; ++i)
-			{
-				GradientColors.Add(FLinearColor((i % 6) * 60.f, 1.f, 1.f).HSVToLinearRGB());
-			}
-		}
-	}
-
-	return GradientColors;
-}
-
-FLinearColor SColorGradingColorWheel::GetGradientStartColor(const FVector4& ColorValue, bool bIsRGB, uint32 ComponentIndex) const
-{
-	if (bIsRGB)
-	{
-		switch (ComponentIndex)
-		{
-		case 0:		return FLinearColor(0.0f, ColorValue.Y, ColorValue.Z, 1.0f);
-		case 1:		return FLinearColor(ColorValue.X, 0.0f, ColorValue.Z, 1.0f);
-		case 2:		return FLinearColor(ColorValue.X, ColorValue.Y, 0.0f, 1.0f);
-		case 3:		return FLinearColor(0.0f, 0.0f, 0.0f, 1.0f);
-		default:	return FLinearColor(ForceInit);
-		}
-	}
-	else
-	{
-		FLinearColor HSVColor = FLinearColor(ColorValue.X, ColorValue.Y, ColorValue.Z, ColorValue.W).LinearRGBToHSV();
-		switch (ComponentIndex)
-		{
-		case 0:		return FLinearColor(HSVColor.R, HSVColor.G, HSVColor.B, 1.0f);
-		case 1:		return FLinearColor(HSVColor.R, 0.0f, HSVColor.B, 1.0f).HSVToLinearRGB();
-		case 2:		return FLinearColor(HSVColor.R, HSVColor.G, 0.0f, 1.0f).HSVToLinearRGB();
-		case 3:		return FLinearColor(0.0f, 0.0f, 0.0f, 1.0f);
-		default:	return FLinearColor(ForceInit);
-		}
-	}
-}
-
-FLinearColor SColorGradingColorWheel::GetGradientEndColor(const FVector4& ColorValue, bool bIsRGB, uint32 ComponentIndex) const
-{
-	if (bIsRGB)
-	{
-		switch (ComponentIndex)
-		{
-		case 0:		return FLinearColor(1.0f, ColorValue.Y, ColorValue.Z, 1.0f);
-		case 1:		return FLinearColor(ColorValue.X, 1.0f, ColorValue.Z, 1.0f);
-		case 2:		return FLinearColor(ColorValue.X, ColorValue.Y, 1.0f, 1.0f);
-		case 3:		return FLinearColor(ColorValue.X, ColorValue.Y, ColorValue.Z, 1.0f);
-		default:	return FLinearColor(ForceInit);
-		}
-	}
-	else
-	{
-		FLinearColor HSVColor = FLinearColor(ColorValue.X, ColorValue.Y, ColorValue.Z, ColorValue.W).LinearRGBToHSV();
-		switch (ComponentIndex)
-		{
-		case 0:		return FLinearColor(HSVColor.R, HSVColor.G, HSVColor.B, 1.0f);
-		case 1:		return FLinearColor(HSVColor.R, 1.0f, HSVColor.B, 1.0f).HSVToLinearRGB();
-		case 2:		return FLinearColor(HSVColor.R, HSVColor.G, 1.0f, 1.0f).HSVToLinearRGB();
-		case 3:		return FLinearColor(HSVColor.R, HSVColor.G, HSVColor.B, 1.0f).HSVToLinearRGB();
-		default:	return FLinearColor(ForceInit);
-		}
-	}
-}
-
-FLinearColor SColorGradingColorWheel::GetGradientFillerColor(const FVector4& ColorValue, bool bIsRGB, uint32 ComponentIndex) const
-{
-	const float MaxValue = 1.0f/*SpinBoxMinMaxSliderValues.CurrentMaxSliderValue.GetValue()*/;
-	if (bIsRGB)
-	{
-		switch (ComponentIndex)
-		{
-		case 0:		return FLinearColor(MaxValue, ColorValue.Y, ColorValue.Z, 1.0f);
-		case 1:		return FLinearColor(ColorValue.X, MaxValue, ColorValue.Z, 1.0f);
-		case 2:		return FLinearColor(ColorValue.X, ColorValue.Y, MaxValue, 1.0f);
-		case 3:		return FLinearColor(ColorValue.X, ColorValue.Y, ColorValue.Z, 1.0f);
-		default:	return FLinearColor(ForceInit);
-		}
-	}
-	else
-	{
-		FLinearColor HSVColor = FLinearColor(ColorValue.X, ColorValue.Y, ColorValue.Z, ColorValue.W).LinearRGBToHSV();
-		switch (ComponentIndex)
-		{
-		case 0:		return FLinearColor(HSVColor.R, HSVColor.G, HSVColor.B, 1.0f);
-		case 1:		return FLinearColor(HSVColor.R, 1.0f, HSVColor.B, 1.0f).HSVToLinearRGB();
-		case 2:		return FLinearColor(HSVColor.R, HSVColor.G, MaxValue, 1.0f).HSVToLinearRGB();
-		case 3:		return FLinearColor(HSVColor.R, HSVColor.G, HSVColor.B, 1.0f).HSVToLinearRGB();
-		default:	return FLinearColor(ForceInit);
-		}
-	}
-}
-
-EVisibility SColorGradingColorWheel::GetGradientVisibility() const
-{
-	bool bShowGradient =
-		ColorPropertyMetadata->ColorGradingMode != EColorGradingModes::Offset &&
-		ColorPropertyMetadata->ColorGradingMode != EColorGradingModes::Invalid;
-
-	return bShowGradient ? EVisibility::Visible : EVisibility::Hidden;
-}
-
 TOptional<float> SColorGradingColorWheel::GetMetadataMinValue() const
 {
 	if (ColorPropertyMetadata.IsSet())
diff --git a/Engine/Plugins/Editor/ColorGrading/Source/ColorGradingEditor/Private/SColorGradingColorWheel.h b/Engine/Plugins/Editor/ColorGrading/Source/ColorGradingEditor/Private/SColorGradingColorWheel.h
--- a/Engine/Plugins/Editor/ColorGrading/Source/ColorGradingEditor/Private/SColorGradingColorWheel.h
+++ b/Engine/Plugins/Editor/ColorGrading/Source/ColorGradingEditor/Private/SColorGradingColorWheel.h
@@ -6,20 +6,27 @@
 #include "Customizations/MathStructCustomizations.h"
 #include "Widgets/DeclarativeSyntaxSupport.h"
 #include "Widgets/SCompoundWidget.h"
-#include "Widgets/Colors/SColorGradingPicker.h"
+#include "Widgets/ColorGrading/SColorGradingPicker.h"
 
 #include "ColorGradingPanelState.h"
 
+#if WITH_EDITOR
+#include "EditorUndoClient.h"
+#endif
+
 class IPropertyHandle;
 class SBox;
 
 /** A widget which encapsulates a color picker and numeric sliders for each color component, hooked up to a color property handle */
 class SColorGradingColorWheel : public SCompoundWidget
+#if WITH_EDITOR
+	, public FEditorUndoClient
+#endif
 {
 public:
 	struct FColorPropertyMetadata
 	{
-		EColorGradingModes ColorGradingMode = EColorGradingModes::Invalid;
+		UE::ColorGrading::EColorGradingModes ColorGradingMode = UE::ColorGrading::EColorGradingModes::Invalid;
 		TOptional<float> MinValue = TOptional<float>();
 		TOptional<float> MaxValue = TOptional<float>();
 		TOptional<float> SliderMinValue = TOptional<float>();
@@ -35,13 +42,14 @@ public:
 
 public:
 	SLATE_BEGIN_ARGS(SColorGradingColorWheel)
-		: _Orientation(Orient_Vertical)
 	{}
-		SLATE_ARGUMENT(EOrientation, Orientation)
-		SLATE_ATTRIBUTE(EColorGradingColorDisplayMode, ColorDisplayMode)
+		SLATE_ATTRIBUTE(UE::ColorGrading::EColorGradingColorDisplayMode, ColorDisplayMode)
 		SLATE_ARGUMENT(TSharedPtr<SWidget>, HeaderContent)
 	SLATE_END_ARGS()
 
+	SColorGradingColorWheel();
+	virtual ~SColorGradingColorWheel();
+
 	void Construct(const FArguments& InArgs);
 
 	/** Sets the property handle for the color property to edit with this color wheel */
@@ -50,8 +58,12 @@ public:
 	/** Sets the widget to display as the header of the color wheel */
 	void SetHeaderContent(const TSharedRef<SWidget>& HeaderContent);
 
-	/** Sets the orientation of the color wheel, which determines if the numeric sliders are below or to the right of the color picker wheel */
-	void SetOrientation(EOrientation NewOrientation);
+#if WITH_EDITOR
+	//~ Begin FEditorUndoClient Interface
+	virtual void PostUndo(bool bSuccess) override;
+	virtual void PostRedo(bool bSuccess) override;
+	// End of FEditorUndoClient
+#endif
 
 private:
 	TSharedRef<SWidget> CreateColorGradingPicker();
@@ -64,13 +76,14 @@ private:
 	bool GetColor(FVector4& OutCurrentColor);
 	void CommitColor(FVector4& NewValue, bool bShouldCommitValueChanges);
 	void TransactColorValue();
+	void RecalculateHSVColor();
 
 	void BeginUsingColorPickerSlider();
 	void EndUsingColorPickerSlider();
 	void BeginUsingComponentSlider(uint32 ComponentIndex);
 	void EndUsingComponentSlider(float NewValue, uint32 ComponentIndex);
 
-	FText GetComponentLabelText(uint32 ComponentIndex) const;
+	UE::ColorGrading::EColorGradingComponent GetComponent(uint32 ComponentIndex) const;
 	TOptional<float> GetComponentValue(uint32 ComponentIndex) const;
 	void SetComponentValue(float NewValue, uint32 ComponentIndex);
 
@@ -82,14 +95,6 @@ private:
 	TOptional<float> GetComponentMinSliderValue(TOptional<float> DefaultValue, uint32 ComponentIndex) const;
 	TOptional<float> GetComponentMaxSliderValue(TOptional<float> DefaultValue, uint32 ComponentIndex) const;
 	float GetComponentSliderDeltaValue(float DefaultValue, uint32 ComponentIndex) const;
-	FText GetComponentToolTipText(uint32 ComponentIndex) const;
-
-	TArray<FLinearColor> GetGradientColor(uint32 ComponentIndex);
-
-	FLinearColor GetGradientStartColor(const FVector4& ColorValue, bool bIsRGB, uint32 ComponentIndex) const;
-	FLinearColor GetGradientEndColor(const FVector4& ColorValue, bool bIsRGB, uint32 ComponentIndex) const;
-	FLinearColor GetGradientFillerColor(const FVector4& ColorValue, bool bIsRGB, uint32 ComponentIndex) const;
-	EVisibility GetGradientVisibility() const;
 
 	TOptional<float> GetMetadataMinValue() const;
 	TOptional<float> GetMetadataMaxValue() const;
@@ -102,9 +107,8 @@ private:
 	bool GetMetadataSupportDynamicSliderMaxValue() const;
 
 private:
-	TSharedPtr<SColorGradingPicker> ColorGradingPicker;
+	TSharedPtr<UE::ColorGrading::SColorGradingPicker> ColorGradingPicker;
 	TSharedPtr<SBox> HeaderBox;
-	TSharedPtr<SBox> ColorWheelBox;
 	TSharedPtr<SBox> ColorPickerBox;
 	TSharedPtr<SBox> ColorSlidersBox;
 
@@ -115,10 +119,7 @@ private:
 	TOptional<FColorPropertyMetadata> ColorPropertyMetadata;
 
 	/** Attribute for the color mode type the color wheel is presenting the color components in */
-	TAttribute<EColorGradingColorDisplayMode> ColorDisplayMode;
-
-	/** The orientation of the color wheel, which determines whether the color component sliders are next to or above the color picker */
-	EOrientation Orientation = Orient_Vertical;
+	TAttribute<UE::ColorGrading::EColorGradingColorDisplayMode> ColorDisplayMode;
 
 	/** Stored current min value of the color component numeric sliders */
 	TOptional<float> ComponentSliderDynamicMinValue;
@@ -131,4 +132,10 @@ private:
 
 	/** Indicates that a component's numeric slider is currently being used to change the color */
 	bool bIsUsingComponentSlider = false;
+
+	/**
+	 * The current color in HSV space.
+	 * Stored separately so that hue/saturation adjustments aren't lost when the color is 0.
+	 */
+	FLinearColor CurrentHSVColor;
 };
\ No newline at end of file
diff --git a/Engine/Plugins/Editor/ColorGrading/Source/ColorGradingEditor/Private/SColorGradingColorWheelPanel.cpp b/Engine/Plugins/Editor/ColorGrading/Source/ColorGradingEditor/Private/SColorGradingColorWheelPanel.cpp
--- a/Engine/Plugins/Editor/ColorGrading/Source/ColorGradingEditor/Private/SColorGradingColorWheelPanel.cpp
+++ b/Engine/Plugins/Editor/ColorGrading/Source/ColorGradingEditor/Private/SColorGradingColorWheelPanel.cpp
@@ -13,8 +13,9 @@
 #include "Modules/ModuleManager.h"
 #include "PropertyHandle.h"
 #include "Widgets/SBoxPanel.h"
-#include "Widgets/Input/SComboButton.h"
+#include "Widgets/Input/SSegmentedControl.h"
 #include "Widgets/Layout/SBox.h"
+#include "Widgets/Layout/SSeparator.h"
 #include "Widgets/Layout/SSpacer.h"
 #include "Widgets/Layout/SSplitter.h"
  
@@ -39,13 +40,7 @@ void SColorGradingColorWheelPanel::Construct(const FArguments& InArgs)
 		ColorGradingDataModel->OnColorGradingElementSelectionChanged().AddSP(this, &SColorGradingColorWheelPanel::OnColorGradingElementSelectionChanged);
 	}
 
-	CommandList = MakeShared<FUICommandList>();
-	BindCommands();
-
-	ColorWheelOrientation = EOrientation::Orient_Vertical;
-
 	ColorWheels.AddDefaulted(NumColorWheels);
-	HiddenColorWheels.AddZeroed(NumColorWheels);
 
 	for (int32 Index = 0; Index < NumColorWheels; ++Index)
     ChildSlot
@@ -70,7 +65,7 @@ void SColorGradingColorWheelPanel::Construct(const FArguments& InArgs)
 		[
 			SNew(SSplitter)
 			.Orientation(Orient_Horizontal)
-			.PhysicalSplitterHandleSize(1.0f)
+			.PhysicalSplitterHandleSize(2.0f)
 			.HitDetectionSplitterHandleSize(5.0f)
 			.Style(FAppStyle::Get(), "DetailsView.Splitter")
 			.Visibility(this, &SColorGradingColorWheelPanel::GetColorWheelPanelVisibility)
@@ -116,86 +111,95 @@ void SColorGradingColorWheelPanel::Construct(const FArguments& InArgs)
 					[
 						MakeColorDisplayModeCheckbox()
 					]
+				]	
 
-					+ SHorizontalBox::Slot()
-					.AutoWidth()
-					.Padding(5.0f, 0.0f, 0.0f, 0.0f)
-					.HAlign(HAlign_Right)
-					.VAlign(VAlign_Center)
-					[
-						SNew(SComboButton)
-						.ComboButtonStyle(&FAppStyle::Get().GetWidgetStyle<FComboButtonStyle>("SimpleComboButton"))
-						.OnGetMenuContent(this, &SColorGradingColorWheelPanel::MakeSettingsMenu)
-						.HasDownArrow(false)
-						.ContentPadding(FMargin(1.0f, 0.0f))
-						.ButtonContent()
-						[
-							SNew(SImage)
-							.ColorAndOpacity(FSlateColor::UseForeground())
-							.Image(FAppStyle::Get().GetBrush("Icons.Settings"))
-						]
-					]
+				+ SVerticalBox::Slot()
+				.AutoHeight()
+				[
+					SNew(SSeparator)
+					.Thickness(2.0f)
 				]
 
+
 				+ SVerticalBox::Slot()
 				.FillHeight(1.0f)
-				.Padding(6, 4)
 				[
 					SNew(SHorizontalBox)
 
 					+ SHorizontalBox::Slot()
 					.FillWidth(1.0f)
 					.HAlign(HAlign_Fill)
-					.Padding(2, 0)
 					[
 						SAssignNew(ColorWheels[0], SColorGradingColorWheel)
-						.ColorDisplayMode(this, &SColorGradingColorWheelPanel::GetColorDisplayMode)
-						.Orientation(ColorWheelOrientation)
-						.Visibility(this, &SColorGradingColorWheelPanel::GetColorWheelVisibility, 0)
+						.ColorDisplayMode(UE::ColorGrading::EColorGradingColorDisplayMode::RGB) // Offset wheel is locked to RGB mode
+
+					]
+
+					+ SHorizontalBox::Slot()
+					.AutoWidth()
+					[
+						SNew(SSeparator)
+						.Orientation(EOrientation::Orient_Vertical)
+						.Thickness(2.0f)
 					]
 
 					+ SHorizontalBox::Slot()
 					.FillWidth(1.0f)
 					.HAlign(HAlign_Fill)
-					.Padding(2, 0)
 					[
 						SAssignNew(ColorWheels[1], SColorGradingColorWheel)
 						.ColorDisplayMode(this, &SColorGradingColorWheelPanel::GetColorDisplayMode)
-						.Orientation(ColorWheelOrientation)
-						.Visibility(this, &SColorGradingColorWheelPanel::GetColorWheelVisibility, 1)
+
+					]
+
+					+ SHorizontalBox::Slot()
+					.AutoWidth()
+					[
+						SNew(SSeparator)
+						.Orientation(EOrientation::Orient_Vertical)
+						.Thickness(2.0f)
 					]
 
 					+ SHorizontalBox::Slot()
 					.FillWidth(1.0f)
 					.HAlign(HAlign_Fill)
-					.Padding(2, 0)
 					[
 						SAssignNew(ColorWheels[2], SColorGradingColorWheel)
 						.ColorDisplayMode(this, &SColorGradingColorWheelPanel::GetColorDisplayMode)
-						.Orientation(ColorWheelOrientation)
-						.Visibility(this, &SColorGradingColorWheelPanel::GetColorWheelVisibility, 2)
+
+					]
+
+					+ SHorizontalBox::Slot()
+					.AutoWidth()
+					[
+						SNew(SSeparator)
+						.Orientation(EOrientation::Orient_Vertical)
+						.Thickness(2.0f)
 					]
 
 					+ SHorizontalBox::Slot()
 					.FillWidth(1.0f)
 					.HAlign(HAlign_Fill)
-					.Padding(2, 0)
 					[
 						SAssignNew(ColorWheels[3], SColorGradingColorWheel)
 						.ColorDisplayMode(this, &SColorGradingColorWheelPanel::GetColorDisplayMode)
-						.Orientation(ColorWheelOrientation)
-						.Visibility(this, &SColorGradingColorWheelPanel::GetColorWheelVisibility, 3)
+
+					]
+
+					+ SHorizontalBox::Slot()
+					.AutoWidth()
+					[
+						SNew(SSeparator)
+						.Orientation(EOrientation::Orient_Vertical)
+						.Thickness(2.0f)
 					]
 
 					+ SHorizontalBox::Slot()
 					.FillWidth(1.0f)
 					.HAlign(HAlign_Fill)
-					.Padding(2, 0)
 					[
 						SAssignNew(ColorWheels[4], SColorGradingColorWheel)
 						.ColorDisplayMode(this, &SColorGradingColorWheelPanel::GetColorDisplayMode)
-						.Orientation(ColorWheelOrientation)
-						.Visibility(this, &SColorGradingColorWheelPanel::GetColorWheelVisibility, 4)
 					]
 				]
 			]
@@ -242,120 +246,30 @@ void SColorGradingColorWheelPanel::Refresh()
 
 void SColorGradingColorWheelPanel::GetPanelState(FColorGradingPanelState& OutPanelState)
 {
-	OutPanelState.HiddenColorWheels = HiddenColorWheels;
 	OutPanelState.ColorDisplayMode = ColorDisplayMode;
-	OutPanelState.ColorWheelOrientation = ColorWheelOrientation;
 }
 
 void SColorGradingColorWheelPanel::SetPanelState(const FColorGradingPanelState& InPanelState)
 {
 	// TODO: These could also be output to a config file to be stored between runs
-	HiddenColorWheels = InPanelState.HiddenColorWheels;
 	ColorDisplayMode = InPanelState.ColorDisplayMode;
-	SetColorWheelOrientation(InPanelState.ColorWheelOrientation);
-}
-
-void SColorGradingColorWheelPanel::BindCommands()
-{
-	const FColorGradingCommands& Commands = FColorGradingCommands::Get();
-
-	CommandList->MapAction(
-		Commands.SaturationColorWheelVisibility,
-		FExecuteAction::CreateSP(this, &SColorGradingColorWheelPanel::ToggleColorWheelVisible, 3),
-		FCanExecuteAction(),
-		FIsActionChecked::CreateSP(this, &SColorGradingColorWheelPanel::IsColorWheelVisible, 3)
-	);
-
-	CommandList->MapAction(
-		Commands.ContrastColorWheelVisibility,
-		FExecuteAction::CreateSP(this, &SColorGradingColorWheelPanel::ToggleColorWheelVisible, 4),
-		FCanExecuteAction(),
-		FIsActionChecked::CreateSP(this, &SColorGradingColorWheelPanel::IsColorWheelVisible, 4)
-	);
-
-	CommandList->MapAction(
-		Commands.ColorWheelSliderOrientationHorizontal,
-		FExecuteAction::CreateSP(this, &SColorGradingColorWheelPanel::SetColorWheelOrientation, EOrientation::Orient_Horizontal),
-		FCanExecuteAction(),
-		FIsActionChecked::CreateSP(this, &SColorGradingColorWheelPanel::IsColorWheelOrientationSelected, EOrientation::Orient_Horizontal)
-	);
-
-	CommandList->MapAction(
-		Commands.ColorWheelSliderOrientationVertical,
-		FExecuteAction::CreateSP(this, &SColorGradingColorWheelPanel::SetColorWheelOrientation, EOrientation::Orient_Vertical),
-		FCanExecuteAction(),
-		FIsActionChecked::CreateSP(this, &SColorGradingColorWheelPanel::IsColorWheelOrientationSelected, EOrientation::Orient_Vertical)
-	);
 }
 
 TSharedRef<SWidget> SColorGradingColorWheelPanel::MakeColorDisplayModeCheckbox()
 {
-	return SNew(SHorizontalBox)
-
-		+ SHorizontalBox::Slot()
-		.AutoWidth()
-		.HAlign(HAlign_Left)
-		.VAlign(VAlign_Top)
-		.Padding(FMargin(0.0f, 0.0f, 3.0f, 0.0f))
-		[
-			SNew(SCheckBox)
-			.Style(FAppStyle::Get(), "ToggleButtonCheckbox")
-			.Type(ESlateCheckBoxType::ToggleButton)
-			.IsChecked(this, &SColorGradingColorWheelPanel::IsColorDisplayModeChecked, EColorGradingColorDisplayMode::RGB)
-			.OnCheckStateChanged(this, &SColorGradingColorWheelPanel::OnColorDisplayModeCheckedChanged, EColorGradingColorDisplayMode::RGB)
-			.ToolTipText(this, &SColorGradingColorWheelPanel::GetColorDisplayModeToolTip, EColorGradingColorDisplayMode::RGB)
-			.Padding(4)
-			.Content()
-			[
-				SNew(STextBlock)
-				.Text(this, &SColorGradingColorWheelPanel::GetColorDisplayModeLabel, EColorGradingColorDisplayMode::RGB)
-				.Font(FAppStyle::GetFontStyle("StandardDialog.SmallFont"))
-			]
-		]
-		+ SHorizontalBox::Slot()
-		.AutoWidth()
-		.HAlign(HAlign_Left)			
-		.VAlign(VAlign_Top)
-		.Padding(FMargin(0.0f, 0.0f, 3.0f, 0.0f))
-		[
-			SNew(SCheckBox)
-			.Style(FAppStyle::Get(), "ToggleButtonCheckbox")
-			.Type(ESlateCheckBoxType::ToggleButton)
-			.IsChecked(this, &SColorGradingColorWheelPanel::IsColorDisplayModeChecked, EColorGradingColorDisplayMode::HSV)
-			.OnCheckStateChanged(this, &SColorGradingColorWheelPanel::OnColorDisplayModeCheckedChanged, EColorGradingColorDisplayMode::HSV)
-			.ToolTipText(this, &SColorGradingColorWheelPanel::GetColorDisplayModeToolTip, EColorGradingColorDisplayMode::HSV)
-			.Padding(4)
-			.Content()
-			[
-				SNew(STextBlock)
-				.Text(this, &SColorGradingColorWheelPanel::GetColorDisplayModeLabel, EColorGradingColorDisplayMode::HSV)
-				.Font(FAppStyle::GetFontStyle("StandardDialog.SmallFont"))
-			]
-		];
-}
-
-TSharedRef<SWidget> SColorGradingColorWheelPanel::MakeSettingsMenu()
-{
-	const bool bShouldCloseWindowAfterMenuSelection = true;
-	FMenuBuilder MenuBuilder(bShouldCloseWindowAfterMenuSelection, CommandList);
-
-	const FColorGradingCommands& Commands = FColorGradingCommands::Get();
-
-	MenuBuilder.BeginSection(TEXT("ColorWheelVisibility"), LOCTEXT("ColorWheelPanel_ShowLabel", "Show"));
-	{
-		MenuBuilder.AddMenuEntry(Commands.SaturationColorWheelVisibility);
-		MenuBuilder.AddMenuEntry(Commands.ContrastColorWheelVisibility);
-	}
-	MenuBuilder.EndSection();
-
-	MenuBuilder.BeginSection(TEXT("ColorWheelSliders"), LOCTEXT("ColorWheelPanel_SlidersLabel", "Sliders"));
-	{
-		MenuBuilder.AddMenuEntry(Commands.ColorWheelSliderOrientationVertical);
-		MenuBuilder.AddMenuEntry(Commands.ColorWheelSliderOrientationHorizontal);
-	}
-	MenuBuilder.EndSection();
-
-	return MenuBuilder.MakeWidget();
+	using SDisplayModeControl = SSegmentedControl<UE::ColorGrading::EColorGradingColorDisplayMode>;
+	return SNew(SDisplayModeControl)
+		.OnValueChanged(this, &SColorGradingColorWheelPanel::OnColorDisplayModeChanged)
+		.Value(this, &SColorGradingColorWheelPanel::GetColorDisplayMode)
+		.UniformPadding(FMargin(16.0f, 2.0f))
+		
+		+ SDisplayModeControl::Slot(UE::ColorGrading::EColorGradingColorDisplayMode::RGB)
+		.Text(this, &SColorGradingColorWheelPanel::GetColorDisplayModeLabel, UE::ColorGrading::EColorGradingColorDisplayMode::RGB)
+		.ToolTip(this, &SColorGradingColorWheelPanel::GetColorDisplayModeToolTip, UE::ColorGrading::EColorGradingColorDisplayMode::RGB)
+		
+		+ SDisplayModeControl::Slot(UE::ColorGrading::EColorGradingColorDisplayMode::HSV)
+		.Text(this, &SColorGradingColorWheelPanel::GetColorDisplayModeLabel, UE::ColorGrading::EColorGradingColorDisplayMode::HSV)
+		.ToolTip(this, &SColorGradingColorWheelPanel::GetColorDisplayModeToolTip, UE::ColorGrading::EColorGradingColorDisplayMode::HSV);
 }
 
 void SColorGradingColorWheelPanel::FillColorGradingGroupProperty(const FColorGradingEditorDataModel::FColorGradingGroup& ColorGradingGroup)
@@ -471,12 +385,6 @@ TSharedRef<SWidget> SColorGradingColorWheelPanel::CreateColorWheelHeaderWidget(c
 
 		TSharedRef<SHorizontalBox> PropertyNameBox = SNew(SHorizontalBox);
 
-		PropertyNameBox->AddSlot()
-			.FillWidth(1.0f)
-			[
-				SNew(SSpacer)
-			];
-
 		if (NodeWidgets.EditConditionWidget.IsValid())
 		{
 			PropertyNameBox->AddSlot()
@@ -603,45 +511,6 @@ bool SColorGradingColorWheelPanel::FilterDetailTreeNode(const TSharedRef<IDetail
 	return false;
 }
 
-void SColorGradingColorWheelPanel::SetColorWheelOrientation(EOrientation NewOrientation)
-{
-	if (ColorWheelOrientation != NewOrientation)
-	{
-		ColorWheelOrientation = NewOrientation;
-
-		for (const TSharedPtr<SColorGradingColorWheel>& ColorWheel : ColorWheels)
-		{
-			if (ColorWheel.IsValid())
-			{
-				ColorWheel->SetOrientation(ColorWheelOrientation);
-			}
-		};
-	}
-}
-
-bool SColorGradingColorWheelPanel::IsColorWheelOrientationSelected(EOrientation Orientation) const
-{
-	return ColorWheelOrientation == Orientation;
-}
-
-void SColorGradingColorWheelPanel::ToggleColorWheelVisible(int32 ColorWheelIndex)
-{
-	if (ColorWheelIndex >= 0 && ColorWheelIndex < HiddenColorWheels.Num())
-	{
-		HiddenColorWheels[ColorWheelIndex] = !HiddenColorWheels[ColorWheelIndex];
-	}
-}
-
-bool SColorGradingColorWheelPanel::IsColorWheelVisible(int32 ColorWheelIndex)
-{
-	if (ColorWheelIndex >= 0 && ColorWheelIndex < HiddenColorWheels.Num())
-	{
-		return !HiddenColorWheels[ColorWheelIndex];
-	}
-
-	return false;
-}
-
 void SColorGradingColorWheelPanel::OnColorGradingGroupSelectionChanged()
 {
 	Refresh();
@@ -703,47 +572,32 @@ EVisibility SColorGradingColorWheelPanel::GetMultiSelectWarningVisibility() cons
 	return bHasMultipleObjects ? EVisibility::Visible : EVisibility::Collapsed;
 }
 
-EVisibility SColorGradingColorWheelPanel::GetColorWheelVisibility(int32 ColorWheelIndex) const
+void SColorGradingColorWheelPanel::OnColorDisplayModeChanged(UE::ColorGrading::EColorGradingColorDisplayMode InColorDisplayMode)
 {
-	bool bIsHidden = HiddenColorWheels[ColorWheelIndex];
-	return bIsHidden ? EVisibility::Collapsed : EVisibility::Visible;
-}
-
-ECheckBoxState SColorGradingColorWheelPanel::IsColorDisplayModeChecked(EColorGradingColorDisplayMode InColorDisplayMode) const
-{
-	bool bIsModeSelected = InColorDisplayMode == ColorDisplayMode;
-	return bIsModeSelected ? ECheckBoxState::Checked : ECheckBoxState::Unchecked;
-}
-
-void SColorGradingColorWheelPanel::OnColorDisplayModeCheckedChanged(ECheckBoxState State, EColorGradingColorDisplayMode InColorDisplayMode)
-{
-	if (State == ECheckBoxState::Checked)
-	{
-		ColorDisplayMode = InColorDisplayMode;
-	}
+	ColorDisplayMode = InColorDisplayMode;
 }
 
-FText SColorGradingColorWheelPanel::GetColorDisplayModeLabel(EColorGradingColorDisplayMode InColorDisplayMode) const
+FText SColorGradingColorWheelPanel::GetColorDisplayModeLabel(UE::ColorGrading::EColorGradingColorDisplayMode InColorDisplayMode) const
 {
 	FText Text;
 
 	switch (InColorDisplayMode)
 	{
-	case EColorGradingColorDisplayMode::RGB: Text = LOCTEXT("ColorWheel_RGBColorDisplayModeLabel", "RGB"); break;
-	case EColorGradingColorDisplayMode::HSV: Text = LOCTEXT("ColorWheel_HSVColorDisplayModeLabel", "HSV"); break;
+	case UE::ColorGrading::EColorGradingColorDisplayMode::RGB: Text = LOCTEXT("ColorWheel_RGBColorDisplayModeLabel", "RGB"); break;
+	case UE::ColorGrading::EColorGradingColorDisplayMode::HSV: Text = LOCTEXT("ColorWheel_HSVColorDisplayModeLabel", "HSV"); break;
 	}
 
 	return Text;
 }
 
-FText SColorGradingColorWheelPanel::GetColorDisplayModeToolTip(EColorGradingColorDisplayMode InColorDisplayMode) const
+FText SColorGradingColorWheelPanel::GetColorDisplayModeToolTip(UE::ColorGrading::EColorGradingColorDisplayMode InColorDisplayMode) const
 {
 	FText Text;
 
 	switch (InColorDisplayMode)
 	{
-	case EColorGradingColorDisplayMode::RGB: Text = LOCTEXT("ColorWheel_RGBColorDisplayModeToolTip", "Change to RGB color mode"); break;
-	case EColorGradingColorDisplayMode::HSV: Text = LOCTEXT("ColorWheel_HSVColorDisplayModeToolTip", "Change to HSV color mode"); break;
+	case UE::ColorGrading::EColorGradingColorDisplayMode::RGB: Text = LOCTEXT("ColorWheel_RGBColorDisplayModeToolTip", "Change to RGB color mode"); break;
+	case UE::ColorGrading::EColorGradingColorDisplayMode::HSV: Text = LOCTEXT("ColorWheel_HSVColorDisplayModeToolTip", "Change to HSV color mode"); break;
 	}
 
 	return Text;
diff --git a/Engine/Plugins/Editor/ColorGrading/Source/ColorGradingEditor/Private/SColorGradingColorWheelPanel.h b/Engine/Plugins/Editor/ColorGrading/Source/ColorGradingEditor/Private/SColorGradingColorWheelPanel.h
--- a/Engine/Plugins/Editor/ColorGrading/Source/ColorGradingEditor/Private/SColorGradingColorWheelPanel.h
+++ b/Engine/Plugins/Editor/ColorGrading/Source/ColorGradingEditor/Private/SColorGradingColorWheelPanel.h
@@ -43,10 +43,7 @@ public:
 	void SetPanelState(const FColorGradingPanelState& InPanelState);
 
 private:
-	void BindCommands();
-
 	TSharedRef<SWidget> MakeColorDisplayModeCheckbox();
-	TSharedRef<SWidget> MakeSettingsMenu();
 
 	void FillColorGradingGroupProperty(const FColorGradingEditorDataModel::FColorGradingGroup& ColorGradingGroup);
 	void ClearColorGradingGroupProperty();
@@ -62,12 +59,6 @@ private:
 
 	bool FilterDetailTreeNode(const TSharedRef<IDetailTreeNode>& InDetailTreeNode);
 
-	void SetColorWheelOrientation(EOrientation NewOrientation);
-	bool IsColorWheelOrientationSelected(EOrientation Orientation) const;
-
-	void ToggleColorWheelVisible(int32 ColorWheelIndex);
-	bool IsColorWheelVisible(int32 ColorWheelIndex);
-
 	void OnColorGradingGroupSelectionChanged();
 	void OnColorGradingElementSelectionChanged();
 
@@ -76,13 +67,11 @@ private:
 	ECheckBoxState IsColorGradingElementSelected(FText ElementName) const;
 	EVisibility GetColorWheelPanelVisibility() const;
 	EVisibility GetMultiSelectWarningVisibility() const;
-	EVisibility GetColorWheelVisibility(int32 ColorWheelIndex) const;
 
-	EColorGradingColorDisplayMode GetColorDisplayMode() const { return ColorDisplayMode; }
-	ECheckBoxState IsColorDisplayModeChecked(EColorGradingColorDisplayMode InColorDisplayMode) const;
-	void OnColorDisplayModeCheckedChanged(ECheckBoxState State, EColorGradingColorDisplayMode InColorDisplayMode);
-	FText GetColorDisplayModeLabel(EColorGradingColorDisplayMode InColorDisplayMode) const;
-	FText GetColorDisplayModeToolTip(EColorGradingColorDisplayMode InColorDisplayMode) const;
+	UE::ColorGrading::EColorGradingColorDisplayMode GetColorDisplayMode() const { return ColorDisplayMode; }
+	void OnColorDisplayModeChanged(UE::ColorGrading::EColorGradingColorDisplayMode InColorDisplayMode);
+	FText GetColorDisplayModeLabel(UE::ColorGrading::EColorGradingColorDisplayMode InColorDisplayMode) const;
+	FText GetColorDisplayModeToolTip(UE::ColorGrading::EColorGradingColorDisplayMode InColorDisplayMode) const;
 
 private:
 	/** The color grading data model that the panel is displaying */
@@ -92,19 +81,12 @@ private:
 	TSharedPtr<SHorizontalBox> ColorGradingElementsToolBarBox;
 
 	TArray<TSharedPtr<SColorGradingColorWheel>> ColorWheels;
-	TArray<bool> HiddenColorWheels;
 
 	TSharedPtr<SColorGradingDetailView> DetailView;
 
-	/** Commands used by the color wheel panel */
-	TSharedPtr<FUICommandList> CommandList;
-
 	/** The currently selected color grading group */
 	int32 SelectedColorGradingGroup = INDEX_NONE;
 
 	/** The current color display mode for the color wheels */
-	EColorGradingColorDisplayMode ColorDisplayMode = EColorGradingColorDisplayMode::RGB;
-
-	/** The current orientation for the color wheels */
-	EOrientation ColorWheelOrientation = EOrientation::Orient_Vertical;
+	UE::ColorGrading::EColorGradingColorDisplayMode ColorDisplayMode = UE::ColorGrading::EColorGradingColorDisplayMode::RGB;
 };
\ No newline at end of file
diff --git a/Engine/Plugins/Editor/ColorGrading/Source/ColorGradingEditor/Private/SColorGradingPanel.cpp b/Engine/Plugins/Editor/ColorGrading/Source/ColorGradingEditor/Private/SColorGradingPanel.cpp
--- a/Engine/Plugins/Editor/ColorGrading/Source/ColorGradingEditor/Private/SColorGradingPanel.cpp
+++ b/Engine/Plugins/Editor/ColorGrading/Source/ColorGradingEditor/Private/SColorGradingPanel.cpp
@@ -173,8 +173,8 @@ void SColorGradingPanel::Construct(const FArguments& InArgs)
 				+SVerticalBox::Slot()
 				[
 					SNew(SBorder)
+					.Padding(0.f)
 					.BorderImage(FAppStyle::Get().GetBrush("Brushes.Panel"))
-					.Padding(FMargin(2.0f, 2.0f, 2.0f, 0.0f))
 					[
 						SAssignNew(ColorWheelPanel, SColorGradingColorWheelPanel)
 						.ColorGradingDataModelSource(ColorGradingDataModel)
diff --git a/Engine/Plugins/Editor/ColorGrading/Source/ColorGradingEditor/Public/ColorGradingPanelState.h b/Engine/Plugins/Editor/ColorGrading/Source/ColorGradingEditor/Public/ColorGradingPanelState.h
--- a/Engine/Plugins/Editor/ColorGrading/Source/ColorGradingEditor/Public/ColorGradingPanelState.h
+++ b/Engine/Plugins/Editor/ColorGrading/Source/ColorGradingEditor/Public/ColorGradingPanelState.h
@@ -4,13 +4,7 @@
 
 #include "CoreMinimal.h"
 #include "Styling/SlateTypes.h"
-
-/** Modes that colors can be displayed in for the color grading panel */
-enum class COLORGRADINGEDITOR_API EColorGradingColorDisplayMode
-{
-	RGB,
-	HSV
-};
+#include "Framework/ColorGrading/ColorGradingCommon.h"
 
 /**
  * Stores the state of the color grading panel UI that can be reloaded in cases where the panel or any of its elements
@@ -27,14 +21,8 @@ struct COLORGRADINGEDITOR_API FColorGradingPanelState
 	/** The color grading element that is selected */
 	int32 SelectedColorGradingElement = INDEX_NONE;
 
-	/** Indicates which color wheels are hidden */
-	TArray<bool> HiddenColorWheels;
-
-	/** The selected orientation of the color wheels */
-	EOrientation ColorWheelOrientation = EOrientation::Orient_Vertical;
-
 	/** The color display mode of the color wheels */
-	EColorGradingColorDisplayMode ColorDisplayMode;
+	UE::ColorGrading::EColorGradingColorDisplayMode ColorDisplayMode;
 
 	/** Indicates which subsections were selected for each section in the details panel */
 	TArray<int32> SelectedDetailsSubsections;
diff --git a/Engine/Source/Editor/DetailCustomizations/DetailCustomizations.Build.cs b/Engine/Source/Editor/DetailCustomizations/DetailCustomizations.Build.cs
--- a/Engine/Source/Editor/DetailCustomizations/DetailCustomizations.Build.cs
+++ b/Engine/Source/Editor/DetailCustomizations/DetailCustomizations.Build.cs
@@ -8,6 +8,7 @@ public class DetailCustomizations : ModuleRules
 	{
 		PrivateDependencyModuleNames.AddRange(
 			new string[] {
+				"AdvancedWidgets",
 				"AppFramework",
 				"Core",
 // 				"AudioEditor",
diff --git a/Engine/Source/Editor/DetailCustomizations/Private/ColorGradingVectorCustomization.cpp b/Engine/Source/Editor/DetailCustomizations/Private/ColorGradingVectorCustomization.cpp
--- a/Engine/Source/Editor/DetailCustomizations/Private/ColorGradingVectorCustomization.cpp
+++ b/Engine/Source/Editor/DetailCustomizations/Private/ColorGradingVectorCustomization.cpp
@@ -33,17 +33,17 @@
 #include "Types/SlateEnums.h"
 #include "Types/SlateStructs.h"
 #include "UObject/UnrealType.h"
+#include "Util/ColorGradingUtil.h"
 #include "Vector4StructCustomization.h"
+#include "Widgets/ColorGrading/SColorGradingPicker.h"
+#include "Widgets/ColorGrading/SColorGradingComponentViewer.h"
 #include "Widgets/Colors/SColorBlock.h"
-#include "Widgets/Colors/SColorGradingPicker.h"
 #include "Widgets/Colors/SComplexGradient.h"
 #include "Widgets/DeclarativeSyntaxSupport.h"
-#include "Widgets/Input/SCheckBox.h"
-#include "Widgets/Input/SNumericEntryBox.h"
-#include "Widgets/Input/SSpinBox.h"
 #include "Widgets/Layout/SBox.h"
 #include "Widgets/SBoxPanel.h"
 #include "Widgets/SOverlay.h"
+#include "Widgets/Input/SSegmentedControl.h"
 #include "Widgets/SWidget.h"
 #include "Widgets/Text/STextBlock.h"
 
@@ -104,9 +104,9 @@ FColorGradingVectorCustomizationBase::FColorGradingVectorCustomizationBase(TWeak
 	}
 }
 
-EColorGradingModes FColorGradingVectorCustomizationBase::GetColorGradingMode() const
+UE::ColorGrading::EColorGradingModes FColorGradingVectorCustomizationBase::GetColorGradingMode() const
 {
-	EColorGradingModes ColorGradingMode = EColorGradingModes::Invalid;
+	UE::ColorGrading::EColorGradingModes ColorGradingMode = UE::ColorGrading::EColorGradingModes::Invalid;
 
 	if (ColorGradingPropertyHandle.IsValid())
 	{
@@ -118,23 +118,23 @@ EColorGradingModes FColorGradingVectorCustomizationBase::GetColorGradingMode() c
 		{
 			if (ColorGradingModeString.Compare(TEXT("saturation")) == 0)
 			{
-				ColorGradingMode = EColorGradingModes::Saturation;
+				ColorGradingMode = UE::ColorGrading::EColorGradingModes::Saturation;
 			}
 			else if (ColorGradingModeString.Compare(TEXT("contrast")) == 0)
 			{
-				ColorGradingMode = EColorGradingModes::Contrast;
+				ColorGradingMode = UE::ColorGrading::EColorGradingModes::Contrast;
 			}
 			else if (ColorGradingModeString.Compare(TEXT("gamma")) == 0)
 			{
-				ColorGradingMode = EColorGradingModes::Gamma;
+				ColorGradingMode = UE::ColorGrading::EColorGradingModes::Gamma;
 			}
 			else if (ColorGradingModeString.Compare(TEXT("gain")) == 0)
 			{
-				ColorGradingMode = EColorGradingModes::Gain;
+				ColorGradingMode = UE::ColorGrading::EColorGradingModes::Gain;
 			}
 			else if (ColorGradingModeString.Compare(TEXT("offset")) == 0)
 			{
-				ColorGradingMode = EColorGradingModes::Offset;
+				ColorGradingMode = UE::ColorGrading::EColorGradingModes::Offset;
 			}
 		}
 	}
@@ -210,66 +210,17 @@ void FColorGradingVectorCustomizationBase::OnEndSliderMovement(float NewValue, i
 	GEditor->EndTransaction();
 }
 
-FText FColorGradingVectorCustomizationBase::OnGetColorLabelText(FText DefaultText, int32 ColorIndex) const
+UE::ColorGrading::EColorGradingComponent FColorGradingVectorCustomizationBase::OnGetColorComponent(int32 ColorIndex) const
 {
-	if (ColorIndex >= 0 && ColorIndex < 3)
-	{
-		if (IsRGBMode)
-		{
-			FText LabelRGBText[3];
-			LabelRGBText[0] = NSLOCTEXT("ColorGradingVectorCustomizationRGBNS", "RedChannelSmallName", "R");
-			LabelRGBText[1] = NSLOCTEXT("ColorGradingVectorCustomizationRGBNS", "GreenChannelSmallName", "G");
-			LabelRGBText[2] = NSLOCTEXT("ColorGradingVectorCustomizationRGBNS", "BlueChannelSmallName", "B");
-
-			return LabelRGBText[ColorIndex];
-		}
-		else
-		{
-			FText LabelHSVText[3];
-			LabelHSVText[0] = NSLOCTEXT("ColorGradingVectorCustomizationHSVNS", "HueChannelSmallName", "H");
-			LabelHSVText[1] = NSLOCTEXT("ColorGradingVectorCustomizationHSVNS", "SaturationChannelSmallName", "S");
-			LabelHSVText[2] = NSLOCTEXT("ColorGradingVectorCustomizationHSVNS", "ValueChannelSmallName", "V");
-
-			return LabelHSVText[ColorIndex];
-		}
-	}
-	else if (ColorIndex == 3)
-	{
-		return NSLOCTEXT("ColorGradingVectorCustomizationNS", "LuminanceChannelSmallName", "Y");
-	}
-
-	return DefaultText;
+	return UE::ColorGrading::GetColorGradingComponent(
+		IsRGBMode ? UE::ColorGrading::EColorGradingColorDisplayMode::RGB : UE::ColorGrading::EColorGradingColorDisplayMode::HSV,
+		ColorIndex
+	);
 }
 
-FText FColorGradingVectorCustomizationBase::OnGetColorLabelToolTipsText(FText DefaultText, int32 ColorIndex) const
+bool FColorGradingVectorCustomizationBase::GetCurrentColorGradingValue(FVector4& OutCurrentValue)
 {
-	if (ColorIndex >= 0 && ColorIndex < 3)
-	{
-		if (IsRGBMode)
-		{
-			FText LabelRGBText[3];
-			LabelRGBText[0] = NSLOCTEXT("ColorGradingVectorCustomizationRGBNSToolTips", "RedChannelSmallNameToolTips", "Red");
-			LabelRGBText[1] = NSLOCTEXT("ColorGradingVectorCustomizationRGBNSToolTips", "GreenChannelSmallNameToolTips", "Green");
-			LabelRGBText[2] = NSLOCTEXT("ColorGradingVectorCustomizationRGBNSToolTips", "BlueChannelSmallNameToolTips", "Blue");
-
-			return LabelRGBText[ColorIndex];
-		}
-		else
-		{
-			FText LabelHSVText[3];
-			LabelHSVText[0] = NSLOCTEXT("ColorGradingVectorCustomizationHSVNSToolTips", "HueChannelSmallNameToolTips", "Hue");
-			LabelHSVText[1] = NSLOCTEXT("ColorGradingVectorCustomizationHSVNSToolTips", "SaturationChannelSmallNameToolTips", "Saturation");
-			LabelHSVText[2] = NSLOCTEXT("ColorGradingVectorCustomizationHSVNSToolTips", "ValueChannelSmallNameToolTips", "Value");
-
-			return LabelHSVText[ColorIndex];
-		}
-	}
-	else if (ColorIndex == 3)
-	{
-		return NSLOCTEXT("ColorGradingVectorCustomizationNSToolTips", "LuminanceChannelSmallNameToolTips", "Luminance");
-	}
-
-	return DefaultText;
+	return ColorGradingPropertyHandle.Pin()->GetValue(OutCurrentValue) == FPropertyAccess::Success;
 }
 
 void FColorGradingVectorCustomizationBase::OnValueChanged(float NewValue, int32 ColorIndex)
@@ -343,119 +294,6 @@ void FColorGradingVectorCustomizationBase::OnCurrentHSVColorChangedDelegate(FLin
 	}
 }
 
-FLinearColor FColorGradingVectorCustomizationBase::GetGradientFillerColor(int32 ColorIndex) const
-{
-	FVector4 ValueVector;
-
-	if (ColorGradingPropertyHandle.Pin()->GetValue(ValueVector) == FPropertyAccess::Success)
-	{
-		if (IsRGBMode)
-		{
-			switch (ColorIndex)
-			{
-				case 0:	return FLinearColor(SpinBoxMinMaxSliderValues.CurrentMaxSliderValue.GetValue(), ValueVector.Y, ValueVector.Z, 1.0f);
-				case 1:	return FLinearColor(ValueVector.X, SpinBoxMinMaxSliderValues.CurrentMaxSliderValue.GetValue(), ValueVector.Z, 1.0f);
-				case 2:	return FLinearColor(ValueVector.X, ValueVector.Y, SpinBoxMinMaxSliderValues.CurrentMaxSliderValue.GetValue(), 1.0f);
-				case 3: return FLinearColor(ValueVector.X, ValueVector.Y, ValueVector.Z, 1.0f);
-				default:	return FLinearColor(ForceInit);
-			}
-		}
-
-		switch (ColorIndex)
-		{
-			case 0:		return FLinearColor(CurrentHSVColor.R, CurrentHSVColor.G, CurrentHSVColor.B, 1.0f);
-			case 1:		return FLinearColor(CurrentHSVColor.R, 1.0f, CurrentHSVColor.B, 1.0f).HSVToLinearRGB();
-			case 2:		return FLinearColor(CurrentHSVColor.R, CurrentHSVColor.G, SpinBoxMinMaxSliderValues.CurrentMaxSliderValue.GetValue(), 1.0f).HSVToLinearRGB();
-			case 3:		return FLinearColor(CurrentHSVColor.R, CurrentHSVColor.G, CurrentHSVColor.B, 1.0f).HSVToLinearRGB();
-			default:	return FLinearColor(ForceInit);
-		}
-	}
-
-	return FLinearColor(ForceInit);
-}
-
-FLinearColor FColorGradingVectorCustomizationBase::GetGradientEndColor(int32 ColorIndex) const
-{
-	FVector4 ValueVector;
-	if (ColorGradingPropertyHandle.Pin()->GetValue(ValueVector) == FPropertyAccess::Success)
-	{
-		if (IsRGBMode)
-		{
-			switch (ColorIndex)
-			{
-				case 0:	return FLinearColor(1.0f, ValueVector.Y, ValueVector.Z, 1.0f);
-				case 1:	return FLinearColor(ValueVector.X, 1.0f, ValueVector.Z, 1.0f);
-				case 2:	return FLinearColor(ValueVector.X, ValueVector.Y, 1.0f, 1.0f);
-				case 3: return FLinearColor(ValueVector.X, ValueVector.Y, ValueVector.Z, 1.0f);
-				default:	return FLinearColor(ForceInit);
-			}
-		}
-
-		switch (ColorIndex)
-		{
-			case 0:		return FLinearColor(CurrentHSVColor.R, CurrentHSVColor.G, CurrentHSVColor.B, 1.0f);
-			case 1:		return FLinearColor(CurrentHSVColor.R, 1.0f, CurrentHSVColor.B, 1.0f).HSVToLinearRGB();
-			case 2:		return FLinearColor(CurrentHSVColor.R, CurrentHSVColor.G, 1.0f, 1.0f).HSVToLinearRGB();
-			case 3:		return FLinearColor(CurrentHSVColor.R, CurrentHSVColor.G, CurrentHSVColor.B, 1.0f).HSVToLinearRGB();
-			default:	return FLinearColor(ForceInit);
-		}
-	}
-
-	return FLinearColor(ForceInit);
-}
-
-FLinearColor FColorGradingVectorCustomizationBase::GetGradientStartColor(int32 ColorIndex) const
-{
-	FVector4 ValueVector;
-
-	if (ColorGradingPropertyHandle.Pin()->GetValue(ValueVector) == FPropertyAccess::Success)
-	{
-		if (IsRGBMode)
-		{
-			switch (ColorIndex)
-			{
-			case 0:		return FLinearColor(0.0f, ValueVector.Y, ValueVector.Z, 1.0f);
-			case 1:		return FLinearColor(ValueVector.X, 0.0f, ValueVector.Z, 1.0f);
-			case 2:		return FLinearColor(ValueVector.X, ValueVector.Y, 0.0f, 1.0f);
-			case 3:		return FLinearColor(0.0f, 0.0f, 0.0f, 1.0f);
-			default:	return FLinearColor(ForceInit);
-			}
-		}
-
-		switch (ColorIndex)
-		{
-		case 0:		return FLinearColor(CurrentHSVColor.R, CurrentHSVColor.G, CurrentHSVColor.B, 1.0f);
-		case 1:		return FLinearColor(CurrentHSVColor.R, 0.0f, CurrentHSVColor.B, 1.0f).HSVToLinearRGB();
-		case 2:		return FLinearColor(CurrentHSVColor.R, CurrentHSVColor.G, 0.0f, 1.0f).HSVToLinearRGB();
-		case 3:		return FLinearColor(0.0f, 0.0f, 0.0f, 1.0f);
-		default:	return FLinearColor(ForceInit);
-		}
-	}
-
-	return FLinearColor(ForceInit);
-}
-
-TArray<FLinearColor> FColorGradingVectorCustomizationBase::GetGradientColor(int32 ColorIndex) const
-{
-	TArray<FLinearColor> GradientColors;
-
-	if (IsRGBMode || ColorIndex > 0)
-	{
-		GradientColors.Add(GetGradientStartColor(ColorIndex));
-		GradientColors.Add(GetGradientEndColor(ColorIndex));
-		GradientColors.Add(GetGradientFillerColor(ColorIndex));
-	}
-	else // HSV Hue handling
-	{
-		for (int32 i = 0; i < 7; ++i)
-		{
-			GradientColors.Add(FLinearColor((i % 6) * 60.f, 1.f, 1.f).HSVToLinearRGB());
-		}
-	}
-
-	return GradientColors;
-}
-
 void FColorGradingVectorCustomizationBase::PostUndo(bool bSuccess)
 {
 	if (ColorGradingPropertyHandle.IsValid())
@@ -476,7 +314,7 @@ void FColorGradingVectorCustomizationBase::PostRedo(bool bSuccess)
 
 void FColorGradingVectorCustomizationBase::OnDynamicSliderMaxValueChanged(float NewMaxSliderValue, TWeakPtr<SWidget> InValueChangedSourceWidget, bool IsOriginator, bool UpdateOnlyIfHigher)
 {
-	if (NumericEntryBoxWidgetList.Num() > 0)
+	if (ComponentViewers.Num() > 0)
 	{
 		if (!SpinBoxMinMaxSliderValues.CurrentMaxSliderValue.IsSet() || (NewMaxSliderValue > SpinBoxMinMaxSliderValues.CurrentMaxSliderValue.GetValue() && UpdateOnlyIfHigher) || !UpdateOnlyIfHigher)
 		{
@@ -492,7 +330,7 @@ void FColorGradingVectorCustomizationBase::OnDynamicSliderMaxValueChanged(float
 
 void FColorGradingVectorCustomizationBase::OnDynamicSliderMinValueChanged(float NewMinSliderValue, TWeakPtr<SWidget> InValueChangedSourceWidget, bool IsOriginator, bool UpdateOnlyIfLower)
 {
-	if (NumericEntryBoxWidgetList.Num() > 0)
+	if (ComponentViewers.Num() > 0)
 	{
 		if (!SpinBoxMinMaxSliderValues.CurrentMinSliderValue.IsSet() || (NewMinSliderValue < SpinBoxMinMaxSliderValues.CurrentMinSliderValue.GetValue() && UpdateOnlyIfLower) || !UpdateOnlyIfLower)
 		{
@@ -537,20 +375,21 @@ bool FColorGradingVectorCustomizationBase::IsEntryBoxEnabled(int32 ColorIndex) c
 	return OnSliderGetValue(ColorIndex) != TOptional<float>();
 }
 
-TSharedRef<SNumericEntryBox<float>> FColorGradingVectorCustomizationBase::MakeNumericEntryBox(int32 ColorIndex, TOptional<float>& MinValue, TOptional<float>& MaxValue, TOptional<float>& SliderMinValue, TOptional<float>& SliderMaxValue, float& SliderExponent, float& Delta, float& ShiftMultiplier, float& CtrlMultiplier, bool& SupportDynamicSliderMaxValue, bool& SupportDynamicSliderMinValue)
+TSharedRef<UE::ColorGrading::SColorGradingComponentViewer> FColorGradingVectorCustomizationBase::MakeComponentViewer(int32 ColorIndex, TOptional<float>& MinValue, TOptional<float>& MaxValue, TOptional<float>& SliderMinValue, TOptional<float>& SliderMaxValue, float& SliderExponent, float& Delta, float& ShiftMultiplier, float& CtrlMultiplier, bool& SupportDynamicSliderMaxValue, bool& SupportDynamicSliderMinValue, bool UseCompactDisplay)
 {
-	TAttribute<FText> TextGetter = TAttribute<FText>::Create(TAttribute<FText>::FGetter::CreateSP(this, &FColorGradingVectorCustomizationBase::OnGetColorLabelText, ColorGradingPropertyHandle.Pin()->GetPropertyDisplayName(), ColorIndex));
-	TSharedRef<SWidget> LabelWidget = SNumericEntryBox<float>::BuildLabel(TextGetter, FLinearColor::White, FLinearColor(0.2f, 0.2f, 0.2f));
-
-	return SNew(SNumericEntryBox<float>)
-		.SpinBoxStyle(&FCoreStyle::Get().GetWidgetStyle<FSpinBoxStyle>("NumericEntrySpinBox_Dark"))
-		.EditableTextBoxStyle(&FCoreStyle::Get().GetWidgetStyle<FEditableTextBoxStyle>("DarkEditableTextBox"))
-		.Font(IDetailLayoutBuilder::GetDetailFont())
-		.UndeterminedString(NSLOCTEXT("PropertyEditor", "MultipleValues", "Multiple Values"))
+	TAttribute<UE::ColorGrading::EColorGradingComponent> ComponentGetter = TAttribute<UE::ColorGrading::EColorGradingComponent>::Create(
+		TAttribute<UE::ColorGrading::EColorGradingComponent>::FGetter::CreateSP(this, &FColorGradingVectorCustomizationBase::OnGetColorComponent, ColorIndex)
+	);
+
+	return SNew(UE::ColorGrading::SColorGradingComponentViewer)
+		.Component(ComponentGetter)
+		.ColorGradingMode(GetColorGradingMode())
+		.UseCompactDisplay(UseCompactDisplay)
 		.Value(this, &FColorGradingVectorCustomizationBase::OnSliderGetValue, ColorIndex)
 		.OnValueChanged(this, &FColorGradingVectorCustomizationBase::OnValueChanged, ColorIndex)
 		.OnBeginSliderMovement(this, &FColorGradingVectorCustomizationBase::OnBeginSliderMovement)
 		.OnEndSliderMovement(this, &FColorGradingVectorCustomizationBase::OnEndSliderMovement, ColorIndex)
+		.OnQueryCurrentColor(this, &FColorGradingVectorCustomizationBase::GetCurrentColorGradingValue)
 		// Only allow spin on handles with one object.  Otherwise it is not clear what value to spin
 		.AllowSpin(ColorGradingPropertyHandle.Pin()->GetNumOuterObjects() == 1)
 		.ShiftMultiplier(ShiftMultiplier)
@@ -566,13 +405,7 @@ TSharedRef<SNumericEntryBox<float>> FColorGradingVectorCustomizationBase::MakeNu
 		.SliderExponent(SliderExponent)
 		.SliderExponentNeutralValue(SliderMinValue.GetValue() + (SliderMaxValue.GetValue() - SliderMinValue.GetValue()) / 2.0f)
 		.Delta(this, &FColorGradingVectorCustomizationBase::OnGetSliderDeltaValue, Delta, ColorIndex)
-		.ToolTipText(this, &FColorGradingVectorCustomizationBase::OnGetColorLabelToolTipsText, ColorGradingPropertyHandle.Pin()->GetPropertyDisplayName(), ColorIndex)
-		.LabelPadding(FMargin(0))
-		.IsEnabled(this, &FColorGradingVectorCustomizationBase::IsEntryBoxEnabled, ColorIndex)
-		.Label()
-		[
-			LabelWidget
-		];
+		.IsEnabled(this, &FColorGradingVectorCustomizationBase::IsEntryBoxEnabled, ColorIndex);
 }
 
 //////////////////////////////////////////////////////////////////////////
@@ -604,9 +437,9 @@ void FColorGradingVectorCustomization::MakeHeaderRow(FDetailWidgetRow& Row, TSha
 			]
 		];
 
-	EColorGradingModes ColorGradingMode = GetColorGradingMode();
+	UE::ColorGrading::EColorGradingModes ColorGradingMode = GetColorGradingMode();
 
-	if (ColorGradingMode == EColorGradingModes::Offset)
+	if (ColorGradingMode == UE::ColorGrading::EColorGradingModes::Offset)
 	{
 		Row.ValueContent()
 			// Make enough space for each child handle
@@ -625,29 +458,26 @@ void FColorGradingVectorCustomization::MakeHeaderRow(FDetailWidgetRow& Row, TSha
 		for (int32 ColorIndex = 0; ColorIndex < SortedChildArray.Num(); ++ColorIndex)
 		{
 			TWeakPtr<IPropertyHandle> WeakHandlePtr = SortedChildArray[ColorIndex];
-			TSharedRef<SNumericEntryBox<float>> NumericEntryBox = MakeNumericEntryBox(ColorIndex, Metadata.MinValue, Metadata.MaxValue, 
+			TSharedRef<UE::ColorGrading::SColorGradingComponentViewer> ComponentViewer = MakeComponentViewer(ColorIndex, Metadata.MinValue, Metadata.MaxValue, 
 				Metadata.SliderMinValue, Metadata.SliderMaxValue, Metadata.SliderExponent, Metadata.Delta, 
-				Metadata.ShiftMultiplier, Metadata.CtrlMultiplier, Metadata.bSupportDynamicSliderMaxValue, Metadata.bSupportDynamicSliderMinValue);
-			TSharedPtr<SSpinBox<float>> NumericEntrySpinBox = StaticCastSharedPtr<SSpinBox<float>>(NumericEntryBox->GetSpinBox());
+				Metadata.ShiftMultiplier, Metadata.CtrlMultiplier, Metadata.bSupportDynamicSliderMaxValue, Metadata.bSupportDynamicSliderMinValue,
+				/*UseCompactDisplay = */true);
 			 
-			NumericEntryBoxWidgetList.Add(NumericEntryBox);
+			ComponentViewers.Add(ComponentViewer);
 
-			if (NumericEntrySpinBox.IsValid())
-			{
-				float MinSliderValue = NumericEntrySpinBox->GetMinSliderValue();
-				float MaxSliderValue = NumericEntrySpinBox->GetMaxSliderValue();
+			float MinSliderValue = ComponentViewer->GetMinSliderValue();
+			float MaxSliderValue = ComponentViewer->GetMaxSliderValue();
 
-				SpinBoxMinMaxSliderValues.CurrentMinSliderValue = MinSliderValue == TNumericLimits<float>::Lowest() ? TOptional<float>() : MinSliderValue;
-				SpinBoxMinMaxSliderValues.CurrentMaxSliderValue = MaxSliderValue == TNumericLimits<float>::Max() ? TOptional<float>() : MaxSliderValue;
-				SpinBoxMinMaxSliderValues.DefaultMinSliderValue = SpinBoxMinMaxSliderValues.CurrentMinSliderValue;
-				SpinBoxMinMaxSliderValues.DefaultMaxSliderValue = SpinBoxMinMaxSliderValues.CurrentMaxSliderValue;
-			}
+			SpinBoxMinMaxSliderValues.CurrentMinSliderValue = MinSliderValue == TNumericLimits<float>::Lowest() ? TOptional<float>() : MinSliderValue;
+			SpinBoxMinMaxSliderValues.CurrentMaxSliderValue = MaxSliderValue == TNumericLimits<float>::Max() ? TOptional<float>() : MaxSliderValue;
+			SpinBoxMinMaxSliderValues.DefaultMinSliderValue = SpinBoxMinMaxSliderValues.CurrentMinSliderValue;
+			SpinBoxMinMaxSliderValues.DefaultMaxSliderValue = SpinBoxMinMaxSliderValues.CurrentMaxSliderValue;
 
 			ContentHorizontalBox->AddSlot()
-				.Padding(FMargin(0.0f, 2.0f, 3.0f, 0.0f))
+				.Padding(FMargin((ColorIndex == 0) ? 0.0f : 4.0f, 2.0f, 4.0f, 0.0f))
 				.VAlign(VAlign_Top)
 				[
-					NumericEntryBox
+					ComponentViewer
 				];
 		}
 	}
@@ -655,34 +485,46 @@ void FColorGradingVectorCustomization::MakeHeaderRow(FDetailWidgetRow& Row, TSha
 	{
 		Row.ValueContent()
 			.VAlign(VAlign_Center)
-			.MinDesiredWidth(250.0f)
+			.MinDesiredWidth(125)
 			[
 				ContentHorizontalBox.ToSharedRef()
 			];
 
 		ContentHorizontalBox->AddSlot()
 			.VAlign(VAlign_Center)
-			.Padding(FMargin(0.0f, 2.0f, 3.0f, 0.0f))
+			.Padding(FMargin(0.0f, 0.0f, 4.0f, 0.0f))
 			[
-				SNew(SOverlay)
-				+ SOverlay::Slot()
-				[
-					SNew(SColorBlock)
-					.Color(this, &FColorGradingVectorCustomization::OnGetHeaderColorBlock)
-					.ShowBackgroundForAlpha(false)
-					.AlphaDisplayMode(EColorBlockAlphaDisplayMode::Ignore)
-					.ColorIsHSV(false)
-					.Size(FVector2D(70.0f, 12.0f))
-				]
-				+ SOverlay::Slot()
-				.HAlign(HAlign_Center)
+				SNew(SBorder)
+				.Padding(1)
+				.BorderImage(FAppStyle::Get().GetBrush("ColorPicker.RoundedSolidBackground"))
+				.BorderBackgroundColor(FAppStyle::Get().GetSlateColor("Colors.InputOutline"))
 				.VAlign(VAlign_Center)
 				[
-					SNew(STextBlock)
-					.Text(NSLOCTEXT("PropertyEditor", "MultipleValues", "Multiple Values"))
-					.Font(IDetailLayoutBuilder::GetDetailFont())
-					.ColorAndOpacity(FSlateColor(FLinearColor::Black)) // we know the background is always white, so can safely set this to black
-					.Visibility(this, &FColorGradingVectorCustomization::GetMultipleValuesTextVisibility)
+					SNew(SOverlay)
+					+ SOverlay::Slot()
+					.VAlign(VAlign_Center)
+					[
+						SNew(SColorBlock)
+						.Color(this, &FColorGradingVectorCustomization::OnGetColorForHeaderColorBlock)
+						.ShowBackgroundForAlpha(false)
+						.Size(FVector2D(70.0f, 20.0f))
+						.CornerRadius(FVector4(4.0f,4.0f,4.0f,4.0f))
+					]
+					+ SOverlay::Slot()
+					.VAlign(VAlign_Center)
+					[
+						SNew(SBorder)
+						.Visibility(this, &FColorGradingVectorCustomization::GetMultipleValuesTextVisibility)
+						.BorderImage(FAppStyle::Get().GetBrush("ColorPicker.MultipleValuesBackground"))
+						.VAlign(VAlign_Center)
+						.ForegroundColor(FAppStyle::Get().GetWidgetStyle<FEditableTextBoxStyle>("NormalEditableTextBox").ForegroundColor)
+						.Padding(FMargin(12.0f, 2.0f))
+						[
+							SNew(STextBlock)
+							.Text(NSLOCTEXT("PropertyEditor", "MultipleValues", "Multiple Values"))
+							.Font(IDetailLayoutBuilder::GetDetailFont())
+						]
+					]
 				]
 			];
 	}
@@ -694,7 +536,7 @@ EVisibility FColorGradingVectorCustomization::GetMultipleValuesTextVisibility()
 	return (ColorGradingPropertyHandle.Pin()->GetValue(VectorValue) == FPropertyAccess::MultipleValues) ? EVisibility::Visible : EVisibility::Collapsed;
 }
 
-FLinearColor FColorGradingVectorCustomization::OnGetHeaderColorBlock() const
+FLinearColor FColorGradingVectorCustomization::OnGetColorForHeaderColorBlock() const
 {
 	FLinearColor ColorValue(0.0f, 0.0f, 0.0f);
 	FVector4 VectorValue;
@@ -715,21 +557,6 @@ FLinearColor FColorGradingVectorCustomization::OnGetHeaderColorBlock() const
 void FColorGradingVectorCustomization::OnColorModeChanged(bool InIsRGBMode)
 {
 	IsRGBMode = InIsRGBMode;
-
-	for (int32 ColorIndex = 0; ColorIndex < NumericEntryBoxWidgetList.Num(); ++ColorIndex)
-	{
-		TSharedPtr<SNumericEntryBox<float>> NumericBox = StaticCastSharedPtr<SNumericEntryBox<float>>(NumericEntryBoxWidgetList[ColorIndex].Pin());
-
-		if (NumericBox.IsValid())
-		{
-			TSharedPtr<SSpinBox<float>> SpinBox = StaticCastSharedPtr<SSpinBox<float>>(NumericBox->GetSpinBox());
-
-			if (SpinBox.IsValid())
-			{
-				SpinBox->SetValue(SpinBox->GetValueAttribute());
-			}
-		}
-	}
 }
 
 void FColorGradingVectorCustomization::CustomizeChildren(IDetailChildrenBuilder& StructBuilder, IPropertyTypeCustomizationUtils& StructCustomizationUtils)
@@ -863,21 +690,27 @@ void FColorGradingCustomBuilder::GenerateHeaderRowContent(FDetailWidgetRow& Node
 	FMathStructCustomization::FNumericMetadata<float> Metadata;
 	FMathStructCustomization::ExtractNumericMetadata(ColorGradingPropertyHandleRef, Metadata);
 	
-	EColorGradingModes ColorGradingMode = GetColorGradingMode();
+	UE::ColorGrading::EColorGradingModes ColorGradingMode = GetColorGradingMode();
+
+	// Add padding on right to compensate for the hidden arrow, which pushes the wheel off-center
+	float RightPadding = 0.f;
+	if (const FSlateBrush* ArrowBrush = FAppStyle::Get().GetBrush("TreeArrow_Expanded"))
+	{
+		RightPadding = ArrowBrush->GetImageSize().X;
+	}
 
 	NodeRow.NameContent()
-		.HAlign(HAlign_Fill)
+	.HAlign(HAlign_Fill)
+	[
+		SNew(SBox)
+		.Padding(FMargin(0.0f, 8.0f, RightPadding, 8.0f))
+		.HAlign(HAlign_Center)
 		[
 			SNew(SBox)
-			.HAlign(HAlign_Center)
-			.VAlign(VAlign_Center)
-			.WidthOverride(125.0f)
-			//.HeightOverride(150.0f)
-			.MinDesiredWidth(125.0f)
-			.MaxDesiredWidth(125.0f)
-			.Padding(FMargin(2.0f, 2.0f, 2.0f, 2.0f))
+			.WidthOverride(175.0f)
+			.HeightOverride(150.0f)
 			[
-				SAssignNew(ColorGradingPickerWidget, SColorGradingPicker)
+				SAssignNew(ColorGradingPickerWidget, UE::ColorGrading::SColorGradingPicker)
 				.ValueMin(Metadata.MinValue)
 				.ValueMax(Metadata.MaxValue)
 				.SliderValueMin(Metadata.SliderMinValue)
@@ -896,107 +729,59 @@ void FColorGradingCustomBuilder::GenerateHeaderRowContent(FDetailWidgetRow& Node
 				.OnBeginMouseCapture(this, &FColorGradingCustomBuilder::OnBeginMouseCapture)
 				.OnEndMouseCapture(this, &FColorGradingCustomBuilder::OnEndMouseCapture)
 			]
-		];
+		]
+	];
 
 	TSharedPtr<SVerticalBox> VerticalBox = SNew(SVerticalBox);
 
+	using SDisplayModeControl = SSegmentedControl<UE::ColorGrading::EColorGradingColorDisplayMode>;
+
 	VerticalBox->AddSlot()
 	.Padding(FMargin(0.0f, 2.0f, 0.0f, 2.0f))
 	.VAlign(VAlign_Top)
+	.HAlign(HAlign_Left)
 	.AutoHeight()
 	[
-		SNew(SHorizontalBox)
-		+ SHorizontalBox::Slot()
-		.AutoWidth()
-		.HAlign(HAlign_Left)
-		.VAlign(VAlign_Top)
-		.Padding(FMargin(0.0f, 0.0f, 3.0f, 0.0f))
-		[
-			SNew(SCheckBox)
-			.Style(FAppStyle::Get(), "ToggleButtonCheckbox")
-			.Type(ESlateCheckBoxType::ToggleButton)
-			.IsChecked(this, &FColorGradingCustomBuilder::OnGetChangeColorMode, ColorModeType::RGB)
-			.OnCheckStateChanged(this, &FColorGradingCustomBuilder::OnChangeColorModeClicked, ColorModeType::RGB)
-			.ToolTipText(this, &FColorGradingCustomBuilder::OnChangeColorModeToolTipText, ColorModeType::RGB)
-			.Visibility(this, &FColorGradingCustomBuilder::OnGetRGBHSVButtonVisibility, ColorModeType::RGB)
-			.Padding(4)
-			.Content()
-			[
-				SNew(STextBlock)
-				.Text(this, &FColorGradingCustomBuilder::OnChangeColorModeText, ColorModeType::RGB)
-				.Font(FAppStyle::GetFontStyle("StandardDialog.SmallFont"))
-			]
-		]
-		+ SHorizontalBox::Slot()
-		.AutoWidth()
-		.HAlign(HAlign_Left)			
-		.VAlign(VAlign_Top)
-		.Padding(FMargin(0.0f, 0.0f, 3.0f, 0.0f))
-		[
-			SNew(SCheckBox)
-			.Style(FAppStyle::Get(), "ToggleButtonCheckbox")
-			.Type(ESlateCheckBoxType::ToggleButton)
-			.IsChecked(this, &FColorGradingCustomBuilder::OnGetChangeColorMode, ColorModeType::HSV)
-			.OnCheckStateChanged(this, &FColorGradingCustomBuilder::OnChangeColorModeClicked, ColorModeType::HSV)
-			.ToolTipText(this, &FColorGradingCustomBuilder::OnChangeColorModeToolTipText, ColorModeType::HSV)
-			.Visibility(this, &FColorGradingCustomBuilder::OnGetRGBHSVButtonVisibility, ColorModeType::RGB)
-			.Padding(4)
-			.Content()
-			[
-				SNew(STextBlock)
-				.Text(this, &FColorGradingCustomBuilder::OnChangeColorModeText, ColorModeType::HSV)
-				.Font(FAppStyle::GetFontStyle("StandardDialog.SmallFont"))
-			]
-		]
+		SNew(SDisplayModeControl)
+		.OnValueChanged(this, &FColorGradingCustomBuilder::OnChangeColorModeClicked)
+		.Value(this, &FColorGradingCustomBuilder::OnGetChangeColorMode)
+		.UniformPadding(FMargin(16.0f, 2.0f))
+
+		+ SDisplayModeControl::Slot(UE::ColorGrading::EColorGradingColorDisplayMode::RGB)
+		.Text(this, &FColorGradingCustomBuilder::OnChangeColorModeText, UE::ColorGrading::EColorGradingColorDisplayMode::RGB)
+		.ToolTip(this, &FColorGradingCustomBuilder::OnChangeColorModeToolTipText, UE::ColorGrading::EColorGradingColorDisplayMode::RGB)
+
+		+ SDisplayModeControl::Slot(UE::ColorGrading::EColorGradingColorDisplayMode::HSV)
+		.Text(this, &FColorGradingCustomBuilder::OnChangeColorModeText, UE::ColorGrading::EColorGradingColorDisplayMode::HSV)
+		.ToolTip(this, &FColorGradingCustomBuilder::OnChangeColorModeToolTipText, UE::ColorGrading::EColorGradingColorDisplayMode::HSV)
 	];
 
 	for (int32 ColorIndex = 0; ColorIndex < SortedChildArray.Num(); ++ColorIndex)
 	{
 		TWeakPtr<IPropertyHandle> WeakHandlePtr = SortedChildArray[ColorIndex];
 
-		TSharedRef<SNumericEntryBox<float>> NumericEntryBox = MakeNumericEntryBox(ColorIndex, Metadata.MinValue, Metadata.MaxValue, 
-			Metadata.SliderMinValue, Metadata.SliderMaxValue, Metadata.SliderExponent, Metadata.Delta, 
-			Metadata.ShiftMultiplier, Metadata.CtrlMultiplier, Metadata.bSupportDynamicSliderMaxValue, Metadata.bSupportDynamicSliderMinValue);
-		TSharedPtr<SSpinBox<float>> NumericEntrySpinBox = StaticCastSharedPtr<SSpinBox<float>>(NumericEntryBox->GetSpinBox());
+		TSharedRef<UE::ColorGrading::SColorGradingComponentViewer> ComponentViewer = MakeComponentViewer(ColorIndex, Metadata.MinValue, Metadata.MaxValue,
+			Metadata.SliderMinValue, Metadata.SliderMaxValue, Metadata.SliderExponent, Metadata.Delta,
+			Metadata.ShiftMultiplier, Metadata.CtrlMultiplier, Metadata.bSupportDynamicSliderMaxValue, Metadata.bSupportDynamicSliderMinValue,
+			/*UseCompactDisplay = */false);
 
-		NumericEntryBoxWidgetList.Add(NumericEntryBox);
+		ComponentViewers.Add(ComponentViewer);
 
-		if (NumericEntrySpinBox.IsValid())
-		{
-			float MinSliderValue = NumericEntrySpinBox->GetMinSliderValue();
-			float MaxSliderValue = NumericEntrySpinBox->GetMaxSliderValue();
+		float MinSliderValue = ComponentViewer->GetMinSliderValue();
+		float MaxSliderValue = ComponentViewer->GetMaxSliderValue();
 
-			SpinBoxMinMaxSliderValues.CurrentMinSliderValue = MinSliderValue == TNumericLimits<float>::Lowest() ? TOptional<float>() : MinSliderValue;
-			SpinBoxMinMaxSliderValues.CurrentMaxSliderValue = MaxSliderValue == TNumericLimits<float>::Max() ? TOptional<float>() : MaxSliderValue;
-			SpinBoxMinMaxSliderValues.DefaultMinSliderValue = SpinBoxMinMaxSliderValues.CurrentMinSliderValue;
-			SpinBoxMinMaxSliderValues.DefaultMaxSliderValue = SpinBoxMinMaxSliderValues.CurrentMaxSliderValue;
-		}
+		SpinBoxMinMaxSliderValues.CurrentMinSliderValue = MinSliderValue == TNumericLimits<float>::Lowest() ? TOptional<float>() : MinSliderValue;
+		SpinBoxMinMaxSliderValues.CurrentMaxSliderValue = MaxSliderValue == TNumericLimits<float>::Max() ? TOptional<float>() : MaxSliderValue;
+		SpinBoxMinMaxSliderValues.DefaultMinSliderValue = SpinBoxMinMaxSliderValues.CurrentMinSliderValue;
+		SpinBoxMinMaxSliderValues.DefaultMaxSliderValue = SpinBoxMinMaxSliderValues.CurrentMaxSliderValue;
 
 		VerticalBox->AddSlot()
-			.Padding(FMargin(0.0f, 2.0f, 3.0f, 0.0f))
+			.Padding(FMargin(0.0f, 8.0f, 4.0f, 0.0f))
 			.AutoHeight()
 			.HAlign(HAlign_Fill)
 			.VAlign(VAlign_Fill)
 			[
-				NumericEntryBox
-			];
-
-
-		// Color Box
-
-		VerticalBox->AddSlot()
-			.Padding(FMargin(15.0f, 0.0f, 3.0f, 2.0f))
-			.AutoHeight()
-			.HAlign(HAlign_Fill)
-			.VAlign(VAlign_Bottom)
-			[
-				SNew(SBox)
-				.HeightOverride(6.0f)
-				[
-					SNew(SComplexGradient)
-						.GradientColors(TAttribute<TArray<FLinearColor>>::Create(TAttribute<TArray<FLinearColor>>::FGetter::CreateSP(this, &FColorGradingVectorCustomizationBase::GetGradientColor, ColorIndex)))
-						.Visibility(this, &FColorGradingCustomBuilder::OnGetGradientVisibility)
-				]
+				ComponentViewer
 			];
 	}	
 	
@@ -1020,8 +805,8 @@ void FColorGradingCustomBuilder::GenerateHeaderRowContent(FDetailWidgetRow& Node
 
 		ColorGradingCustomization->GetOnNumericEntryBoxDynamicSliderMaxValueChangedDelegate().AddSP(this, &FColorGradingVectorCustomizationBase::OnDynamicSliderMaxValueChanged);
 		ColorGradingCustomization->GetOnNumericEntryBoxDynamicSliderMinValueChangedDelegate().AddSP(this, &FColorGradingVectorCustomizationBase::OnDynamicSliderMinValueChanged);
-		ColorGradingCustomization->GetOnNumericEntryBoxDynamicSliderMaxValueChangedDelegate().AddSP(ColorGradingPickerWidget.Pin().Get(), &SColorGradingPicker::OnDynamicSliderMaxValueChanged);
-		ColorGradingCustomization->GetOnNumericEntryBoxDynamicSliderMinValueChangedDelegate().AddSP(ColorGradingPickerWidget.Pin().Get(), &SColorGradingPicker::OnDynamicSliderMinValueChanged);
+		ColorGradingCustomization->GetOnNumericEntryBoxDynamicSliderMaxValueChangedDelegate().AddSP(ColorGradingPickerWidget.Pin().Get(), &UE::ColorGrading::SColorGradingPicker::OnDynamicSliderMaxValueChanged);
+		ColorGradingCustomization->GetOnNumericEntryBoxDynamicSliderMinValueChangedDelegate().AddSP(ColorGradingPickerWidget.Pin().Get(), &UE::ColorGrading::SColorGradingPicker::OnDynamicSliderMinValueChanged);
 
 		OnNumericEntryBoxDynamicSliderMaxValueChanged.AddSP(ColorGradingCustomization.Get(), &FColorGradingVectorCustomizationBase::OnDynamicSliderMaxValueChanged);
 		OnNumericEntryBoxDynamicSliderMinValueChanged.AddSP(ColorGradingCustomization.Get(), &FColorGradingVectorCustomizationBase::OnDynamicSliderMinValueChanged);
@@ -1038,8 +823,8 @@ void FColorGradingCustomBuilder::GenerateHeaderRowContent(FDetailWidgetRow& Node
 		ColorGradingPickerWidget.Pin()->GetOnNumericEntryBoxDynamicSliderMaxValueChangedDelegate().AddSP(this, &FColorGradingVectorCustomizationBase::OnDynamicSliderMaxValueChanged);
 		ColorGradingPickerWidget.Pin()->GetOnNumericEntryBoxDynamicSliderMinValueChangedDelegate().AddSP(this, &FColorGradingVectorCustomizationBase::OnDynamicSliderMinValueChanged);
 
-		OnNumericEntryBoxDynamicSliderMaxValueChanged.AddSP(ColorGradingPickerWidget.Pin().Get(), &SColorGradingPicker::OnDynamicSliderMaxValueChanged);
-		OnNumericEntryBoxDynamicSliderMinValueChanged.AddSP(ColorGradingPickerWidget.Pin().Get(), &SColorGradingPicker::OnDynamicSliderMinValueChanged);
+		OnNumericEntryBoxDynamicSliderMaxValueChanged.AddSP(ColorGradingPickerWidget.Pin().Get(), &UE::ColorGrading::SColorGradingPicker::OnDynamicSliderMaxValueChanged);
+		OnNumericEntryBoxDynamicSliderMinValueChanged.AddSP(ColorGradingPickerWidget.Pin().Get(), &UE::ColorGrading::SColorGradingPicker::OnDynamicSliderMinValueChanged);
 	}
 
 	OnCurrentHSVColorChanged.AddSP(this, &FColorGradingVectorCustomizationBase::OnCurrentHSVColorChangedDelegate);
@@ -1050,25 +835,18 @@ void FColorGradingCustomBuilder::GenerateHeaderRowContent(FDetailWidgetRow& Node
 	float BestMaxSliderValue = 0.0f;
 	float BestMinSliderValue = 0.0f;
 
-	for (TWeakPtr<SWidget>& Widget : NumericEntryBoxWidgetList)
+	for (TWeakPtr<UE::ColorGrading::SColorGradingComponentViewer>& ComponentViewer : ComponentViewers)
 	{
-		TSharedPtr<SNumericEntryBox<float>> NumericBox = StaticCastSharedPtr<SNumericEntryBox<float>>(Widget.Pin());
-
-		if (NumericBox.IsValid())
+		if (TSharedPtr<UE::ColorGrading::SColorGradingComponentViewer> PinnedComponentViewer = ComponentViewer.Pin())
 		{
-			TSharedPtr<SSpinBox<float>> SpinBox = StaticCastSharedPtr<SSpinBox<float>>(NumericBox->GetSpinBox());
-
-			if (SpinBox.IsValid())
+			if (PinnedComponentViewer->GetMaxSliderValue() > BestMaxSliderValue)
 			{
-				if (SpinBox->GetMaxSliderValue() > BestMaxSliderValue) 
-				{
-					BestMaxSliderValue = SpinBox->GetMaxSliderValue();
-				}
+				BestMaxSliderValue = PinnedComponentViewer->GetMaxSliderValue();
+			}
 
-				if (SpinBox->GetMinSliderValue() < BestMinSliderValue)
-				{
-					BestMinSliderValue = SpinBox->GetMinSliderValue();
-				}
+			if (PinnedComponentViewer->GetMinSliderValue() < BestMinSliderValue)
+			{
+				BestMinSliderValue = PinnedComponentViewer->GetMinSliderValue();
 			}
 		}
 	}
@@ -1081,49 +859,49 @@ void FColorGradingCustomBuilder::GenerateHeaderRowContent(FDetailWidgetRow& Node
 	ParentGroupName.ReplaceInline(TEXT("|"), TEXT("_"));
 
 	GConfig->GetBool(TEXT("ColorGrading"), *FString::Printf(TEXT("%s_%s_IsRGB"), *ParentGroupName, *ColorGradingPropertyHandle.Pin()->GetPropertyDisplayName().ToString()), RGBMode, GEditorPerProjectIni);
-	OnChangeColorModeClicked(ECheckBoxState::Checked, RGBMode ? ColorModeType::RGB : ColorModeType::HSV);
+	OnChangeColorModeClicked(RGBMode ? UE::ColorGrading::EColorGradingColorDisplayMode::RGB : UE::ColorGrading::EColorGradingColorDisplayMode::HSV);
 
 	// Register to update when an undo/redo operation has been called to update our list of actors
 	GEditor->RegisterForUndo(this);
 }
 
-FText FColorGradingCustomBuilder::OnChangeColorModeText(ColorModeType ModeType) const
+FText FColorGradingCustomBuilder::OnChangeColorModeText(UE::ColorGrading::EColorGradingColorDisplayMode ModeType) const
 {
 	FText Text;
 
 	switch (ModeType)
 	{
-		case FColorGradingCustomBuilder::ColorModeType::RGB: Text = LOCTEXT("ChangeColorModeRGB", "RGB"); break;
-		case FColorGradingCustomBuilder::ColorModeType::HSV: Text = LOCTEXT("ChangeColorModeHSV", "HSV"); break;
+		case UE::ColorGrading::EColorGradingColorDisplayMode::RGB: Text = LOCTEXT("ChangeColorModeRGB", "RGB"); break;
+		case UE::ColorGrading::EColorGradingColorDisplayMode::HSV: Text = LOCTEXT("ChangeColorModeHSV", "HSV"); break;
 	}
 
 	return Text;
 }
 
-FText FColorGradingCustomBuilder::OnChangeColorModeToolTipText(ColorModeType ModeType) const
+FText FColorGradingCustomBuilder::OnChangeColorModeToolTipText(UE::ColorGrading::EColorGradingColorDisplayMode ModeType) const
 {
 	FText Text;
 
 	switch (ModeType)
 	{
-		case FColorGradingCustomBuilder::ColorModeType::RGB: Text = LOCTEXT("ChangeColorModeRGBToolTips", "Change to RGB color mode"); break;
-		case FColorGradingCustomBuilder::ColorModeType::HSV: Text = LOCTEXT("ChangeColorModeHSVToolTips", "Change to HSV color mode"); break;
+		case UE::ColorGrading::EColorGradingColorDisplayMode::RGB: Text = LOCTEXT("ChangeColorModeRGBToolTips", "Change to RGB color mode"); break;
+		case UE::ColorGrading::EColorGradingColorDisplayMode::HSV: Text = LOCTEXT("ChangeColorModeHSVToolTips", "Change to HSV color mode"); break;
 	}
 
 	return Text;
 }
 
-EVisibility FColorGradingCustomBuilder::OnGetRGBHSVButtonVisibility(ColorModeType ModeType) const
+EVisibility FColorGradingCustomBuilder::OnGetRGBHSVButtonVisibility(UE::ColorGrading::EColorGradingColorDisplayMode ModeType) const
 {
-	return GetColorGradingMode() == EColorGradingModes::Offset ? EVisibility::Hidden : EVisibility::Visible;
+	return GetColorGradingMode() == UE::ColorGrading::EColorGradingModes::Offset ? EVisibility::Hidden : EVisibility::Visible;
 }
 
 EVisibility FColorGradingCustomBuilder::OnGetGradientVisibility() const
 {
-	return GetColorGradingMode() == EColorGradingModes::Offset ? EVisibility::Hidden : EVisibility::Visible;
+	return GetColorGradingMode() == UE::ColorGrading::EColorGradingModes::Offset ? EVisibility::Hidden : EVisibility::Visible;
 }
 
-void FColorGradingCustomBuilder::OnChangeColorModeClicked(ECheckBoxState NewValue, ColorModeType ModeType)
+void FColorGradingCustomBuilder::OnChangeColorModeClicked(UE::ColorGrading::EColorGradingColorDisplayMode ModeType)
 {
 	FVector4 CurrentValueVector;
 	if (ColorGradingPropertyHandle.Pin()->GetValue(CurrentValueVector) != FPropertyAccess::Success)
@@ -1133,11 +911,7 @@ void FColorGradingCustomBuilder::OnChangeColorModeClicked(ECheckBoxState NewValu
 
 	bool NewIsRGBMode = true;
 
-	switch (ModeType)
-	{
-		case FColorGradingCustomBuilder::ColorModeType::RGB: NewIsRGBMode = NewValue == ECheckBoxState::Checked ? true : false; break;
-		case FColorGradingCustomBuilder::ColorModeType::HSV: NewIsRGBMode = NewValue == ECheckBoxState::Checked ? false : true; break;
-	}
+	NewIsRGBMode = ModeType == UE::ColorGrading::EColorGradingColorDisplayMode::RGB;
 
 	if (NewIsRGBMode != IsRGBMode)
 	{
@@ -1153,36 +927,13 @@ void FColorGradingCustomBuilder::OnChangeColorModeClicked(ECheckBoxState NewValu
 
 		OnCurrentHSVColorChanged.Broadcast(CurrentHSVColor, true);
 
-		// Force Refresh of the internal cache of the spinner
-		for (int32 ColorIndex = 0; ColorIndex < NumericEntryBoxWidgetList.Num(); ++ColorIndex)
-		{
-			TSharedPtr<SNumericEntryBox<float>> NumericBox = StaticCastSharedPtr<SNumericEntryBox<float>>(NumericEntryBoxWidgetList[ColorIndex].Pin());
-
-			if (NumericBox.IsValid())
-			{
-				TSharedPtr<SSpinBox<float>> SpinBox = StaticCastSharedPtr<SSpinBox<float>>(NumericBox->GetSpinBox());
-
-				if (SpinBox.IsValid())
-				{
-					SpinBox->SetValue(SpinBox->GetValueAttribute());
-				}
-			}
-		}
-
 		OnColorModeChanged.Broadcast(IsRGBMode);
 	}
 }
 
-ECheckBoxState FColorGradingCustomBuilder::OnGetChangeColorMode(ColorModeType ModeType) const
+UE::ColorGrading::EColorGradingColorDisplayMode FColorGradingCustomBuilder::OnGetChangeColorMode() const
 {
-	switch (ModeType)
-	{
-		case FColorGradingCustomBuilder::ColorModeType::RGB: return IsRGBMode ? ECheckBoxState::Checked : ECheckBoxState::Unchecked;
-		case FColorGradingCustomBuilder::ColorModeType::HSV: return !IsRGBMode ? ECheckBoxState::Checked : ECheckBoxState::Unchecked;
-		default: break;
-	}
-
-	return ECheckBoxState::Unchecked;
+	return IsRGBMode ? UE::ColorGrading::EColorGradingColorDisplayMode::RGB : UE::ColorGrading::EColorGradingColorDisplayMode::HSV;
 }
 
 void FColorGradingCustomBuilder::OnColorGradingPickerChanged(FVector4& NewValue, bool ShouldCommitValueChanges)
@@ -1209,11 +960,6 @@ void FColorGradingCustomBuilder::OnColorGradingPickerChanged(FVector4& NewValue,
 	OnCurrentHSVColorChangedDelegate(NewHSVColor, true);
 }
 
-bool FColorGradingCustomBuilder::GetCurrentColorGradingValue(FVector4& OutCurrentValue)
-{
-	return ColorGradingPropertyHandle.Pin()->GetValue(OutCurrentValue) == FPropertyAccess::Success;
-}
-
 
 void FColorGradingCustomBuilder::GenerateChildContent(IDetailChildrenBuilder& ChildrenBuilder)
 {
diff --git a/Engine/Source/Editor/DetailCustomizations/Private/ColorGradingVectorCustomization.h b/Engine/Source/Editor/DetailCustomizations/Private/ColorGradingVectorCustomization.h
--- a/Engine/Source/Editor/DetailCustomizations/Private/ColorGradingVectorCustomization.h
+++ b/Engine/Source/Editor/DetailCustomizations/Private/ColorGradingVectorCustomization.h
@@ -5,6 +5,7 @@
 #include "Containers/Array.h"
 #include "Delegates/Delegate.h"
 #include "EditorUndoClient.h"
+#include "Framework/ColorGrading/ColorGradingCommon.h"
 #include "HAL/Platform.h"
 #include "IDetailCustomNodeBuilder.h"
 #include "Internationalization/Text.h"
@@ -23,11 +24,14 @@ class IDetailChildrenBuilder;
 class IDetailGroup;
 class IPropertyHandle;
 class IPropertyTypeCustomizationUtils;
-class SColorGradingPicker;
 class SWidget;
-template <typename NumericType> class SNumericEntryBox;
 
+namespace UE::ColorGrading
+{
+class SColorGradingPicker;
+class SColorGradingComponentViewer;
 enum class EColorGradingModes;
+}
 
 struct FColorGradingMinMaxSliderValue
 {
@@ -65,9 +69,6 @@ public:
 	/** Callback handling HSV color changed */
 	void OnCurrentHSVColorChangedDelegate(FLinearColor NewHSVColor, bool Originator);
 
-	/** Callback returning the desired gradiant color for a specified color index */
-	TArray<FLinearColor> GetGradientColor(int32 ColorIndex) const;
-
 	//~ Begin FEditorUndoClient Interface
 	virtual void PostUndo(bool bSuccess) override;
 	virtual void PostRedo(bool bSuccess) override;
@@ -75,15 +76,15 @@ public:
 
 protected:
 	bool IsInRGBMode() const;
-	EColorGradingModes GetColorGradingMode() const;
+	UE::ColorGrading::EColorGradingModes GetColorGradingMode() const;
 
-	TSharedRef<SNumericEntryBox<float>> MakeNumericEntryBox(int32 ColorIndex, TOptional<float>& MinValue, TOptional<float>& MaxValue, TOptional<float>& SliderMinValue, TOptional<float>& SliderMaxValue, float& SliderExponent, float& Delta, float& ShiftMultiplier, float &CtrlMultiplier, bool& SupportDynamicSliderMaxValue, bool& SupportDynamicSliderMinValue);
+	TSharedRef<UE::ColorGrading::SColorGradingComponentViewer> MakeComponentViewer(int32 ColorIndex, TOptional<float>& MinValue, TOptional<float>& MaxValue, TOptional<float>& SliderMinValue, TOptional<float>& SliderMaxValue, float& SliderExponent, float& Delta, float& ShiftMultiplier, float &CtrlMultiplier, bool& SupportDynamicSliderMaxValue, bool& SupportDynamicSliderMinValue, bool UseCompactDisplay);
 
-	/** Callback returning the color label text to display for a specified color index */
-	FText OnGetColorLabelText(FText DefaultText, int32 ColorIndex) const;
-	
-	/** Callback returning the color label text tooltip to display for a specified color index */
-	FText OnGetColorLabelToolTipsText(FText DefaultText, int32 ColorIndex) const;
+	/** Callback returning the component associated with a specified color index */
+	UE::ColorGrading::EColorGradingComponent OnGetColorComponent(int32 ColorIndex) const;
+
+	/** Get the current color value being edited/displayed by this widget */
+	bool GetCurrentColorGradingValue(FVector4& OutCurrentValue);
 	
 	/** Callback returning the min/max slider value for a specified color index */
 	TOptional<float> OnGetMaxSliderValue(TOptional<float> DefaultMaxSliderValue, int32 ColorIndex) const;
@@ -108,17 +109,12 @@ protected:
 	/** Callback returning if an entry box should be enabled */
 	bool IsEntryBoxEnabled(int32 ColorIndex) const;
 
-	/** Helper function used to compute desired gradient color for a requested color index */
-	FLinearColor GetGradientEndColor(int32 ColorIndex) const;
-	FLinearColor GetGradientStartColor(int32 ColorIndex) const;
-	FLinearColor GetGradientFillerColor(int32 ColorIndex) const;
-
 protected:
 	/** Min/Max slider value that can change dynamically */
 	FColorGradingMinMaxSliderValue SpinBoxMinMaxSliderValues;
 	
-	/** Registered numeric entry box list */
-	TArray<TWeakPtr<SWidget>> NumericEntryBoxWidgetList;
+	/** List of registered color component viewers */
+	TArray<TWeakPtr<UE::ColorGrading::SColorGradingComponentViewer>> ComponentViewers;
 	
 	/** The color grading property we're editing */
 	TWeakPtr<IPropertyHandle> ColorGradingPropertyHandle;
@@ -162,7 +158,8 @@ public:
 
 private:
 	/** Will return the color of the color block displayed in the header */
-	FLinearColor OnGetHeaderColorBlock() const;
+	FLinearColor OnGetColorForHeaderColorBlock() const;
+
 	EVisibility GetMultipleValuesTextVisibility() const;
 
 	/** Represent the custom builder associated with the color grading property */
@@ -175,13 +172,6 @@ public:
 	/** Notification when we change color mode (RGB <-> HSV) */
 	DECLARE_MULTICAST_DELEGATE_OneParam(FOnColorModeChanged, bool);
 
-	/** Supported color mode */
-	enum class ColorModeType
-	{
-		RGB,
-		HSV
-	};
-
 	FColorGradingCustomBuilder(TWeakPtr<IPropertyHandle> InColorGradingPropertyHandle, const TArray<TWeakPtr<IPropertyHandle>>& InSortedChildArray, 
 							   TSharedRef<FColorGradingVectorCustomization> InColorGradingCustomization, IDetailGroup* InParentGroup);
 	virtual ~FColorGradingCustomBuilder();
@@ -206,7 +196,6 @@ private:
 
 	/* Local UI Handlers */
 	void OnColorGradingPickerChanged(FVector4 &NewValue, bool ShouldCommitValueChanges);
-	bool GetCurrentColorGradingValue(FVector4 &OutCurrentValue);
 
 	void OnBeginMainValueSliderMovement();
 	void OnEndMainValueSliderMovement();
@@ -218,28 +207,28 @@ private:
 	void OnDetailGroupReset();
 
 	/** Callback returning which color mode text we should display */
-	FText OnChangeColorModeText(ColorModeType ModeType) const;
+	FText OnChangeColorModeText(UE::ColorGrading::EColorGradingColorDisplayMode ModeType) const;
 	
 	/** Callback returning which color mode text tooltip we should display */
-	FText OnChangeColorModeToolTipText(ColorModeType ModeType) const;
+	FText OnChangeColorModeToolTipText(UE::ColorGrading::EColorGradingColorDisplayMode ModeType) const;
 
 	/** Callback returning which color mode is checked */	
-	ECheckBoxState OnGetChangeColorMode(ColorModeType ModeType) const;
+	UE::ColorGrading::EColorGradingColorDisplayMode OnGetChangeColorMode() const;
 
 	/** Callback returning if the RGB/HSV button should be visible */
-	EVisibility OnGetRGBHSVButtonVisibility(ColorModeType ModeType) const;
+	EVisibility OnGetRGBHSVButtonVisibility(UE::ColorGrading::EColorGradingColorDisplayMode ModeType) const;
 
 	/** Callback returning if the Gradient should be visible */
 	EVisibility OnGetGradientVisibility() const;
 	
 	/** Callback Called when user click a color mode change checkbox */
-	void OnChangeColorModeClicked(ECheckBoxState NewValue, ColorModeType ModeType);
+	void OnChangeColorModeClicked(UE::ColorGrading::EColorGradingColorDisplayMode ModeType);
 
 	/** Called to rebuild the children of the detail tree */
 	FSimpleDelegate OnRebuildChildren;
 
 	/** Color Picker widget */
-	TWeakPtr<SColorGradingPicker> ColorGradingPickerWidget;
+	TWeakPtr<UE::ColorGrading::SColorGradingPicker> ColorGradingPickerWidget;
 
 	/** Parent of this custom builder (required to communicate with FColorGradingVectorCustomization) */
 	TSharedPtr<FColorGradingVectorCustomization> ColorGradingCustomization;
diff --git a/Engine/Source/Runtime/AdvancedWidgets/Private/Styling/AdvancedWidgetsStyle.cpp b/Engine/Source/Runtime/AdvancedWidgets/Private/Styling/AdvancedWidgetsStyle.cpp
--- a/Engine/Source/Runtime/AdvancedWidgets/Private/Styling/AdvancedWidgetsStyle.cpp
+++ b/Engine/Source/Runtime/AdvancedWidgets/Private/Styling/AdvancedWidgetsStyle.cpp
@@ -5,14 +5,19 @@
 #include "Framework/PropertyViewer/FieldIconFinder.h"
 #include "Math/ColorList.h"
 #include "Styling/AppStyle.h"
+#include "Styling/ColorGradingSpinBoxStyle.h"
+#include "Styling/CoreStyle.h"
 #include "Styling/SlateColor.h"
 #include "Styling/SlateStyle.h"
 #include "Styling/SlateStyleRegistry.h"
 #include "Styling/SlateTypes.h"
+#include "Styling/StyleColors.h"
 #include "UObject/Class.h"
 
 
 #define BORDER_BRUSH( RelativePath, ... ) FSlateBorderBrush( Instance->RootToContentDir( RelativePath, TEXT(".png") ), __VA_ARGS__ )
+#define BOX_BRUSH( RelativePath, ... ) FSlateBoxBrush( Instance->RootToContentDir( RelativePath, TEXT(".png") ), __VA_ARGS__ )
+#define IMAGE_BRUSH_SVG( RelativePath, ... ) FSlateVectorImageBrush( Instance->RootToContentDir(RelativePath, TEXT(".svg") ), __VA_ARGS__)
 
 namespace UE::AdvancedWidgets
 {
@@ -55,6 +60,39 @@ void FAdvancedWidgetsStyle::Create()
 		.SetDownArrowPadding(FMargin(0.0f))
 		);
 
+	Instance->Set("ColorGradingWheel.Cross", new IMAGE_BRUSH_SVG("Starship/Common/color-grading-cross", FVector2f(336.f, 336.f)));
+	Instance->Set("ColorGradingWheel.Selector", new IMAGE_BRUSH_SVG("Starship/Common/color-grading-selector", FVector2f(16.f, 16.f)));
+
+	const FSlateBrush* NoBrush = FCoreStyle::Get().GetBrush("NoBrush");
+	check(NoBrush);
+
+	{
+		FSpinBoxStyle NumericEntrySpinBoxStyle = FCoreStyle::Get().GetWidgetStyle<FSpinBoxStyle>("NumericEntrySpinBox");
+		NumericEntrySpinBoxStyle.ActiveFillBrush = *NoBrush;
+		NumericEntrySpinBoxStyle.HoveredFillBrush = *NoBrush;
+		NumericEntrySpinBoxStyle.InactiveFillBrush = *NoBrush;
+
+		Instance->Set("ColorGradingComponentViewer.NumericEntry", NumericEntrySpinBoxStyle);
+		Instance->Set("ColorGradingComponentViewer.NumericEntry.TextBox", FCoreStyle::Get().GetWidgetStyle<FEditableTextBoxStyle>("NormalEditableTextBox"));
+
+		Instance->Set("ColorGradingSlider.Border", new FSlateRoundedBoxBrush(
+			FStyleColors::Input, CoreStyleConstants::InputFocusRadius,
+			FStyleColors::Secondary, CoreStyleConstants::InputFocusThickness)
+		);
+	}
+
+	{
+		FSpinBoxStyle DefaultSpinBox = FCoreStyle::Get().GetWidgetStyle<FSpinBoxStyle>("SpinBox");
+
+		Instance->Set("ColorGradingSpinBox", FColorGradingSpinBoxStyle()
+			.SetBorderBrush(FSlateRoundedBoxBrush(FStyleColors::Transparent, CoreStyleConstants::InputFocusRadius, FStyleColors::InputOutline, CoreStyleConstants::InputFocusThickness))
+			.SetActiveBorderBrush(FSlateRoundedBoxBrush(FStyleColors::Transparent, CoreStyleConstants::InputFocusRadius, FStyleColors::Primary, CoreStyleConstants::InputFocusThickness))
+			.SetHoveredBorderBrush(FSlateRoundedBoxBrush(FStyleColors::Transparent, CoreStyleConstants::InputFocusRadius, FStyleColors::Hover, CoreStyleConstants::InputFocusThickness))
+			.SetSelectorBrush(BOX_BRUSH("Starship/Common/color-grading-spinbox-selector", FMargin(1/3.f, 0.0f)))
+			.SetSelectorWidth(3.f)
+		);
+	}
+
 	FLinearColor& VectorColor = ColorSettings.StructColors.FindOrAdd(TBaseStructure<FVector>::Get()->GetStructPathName().ToString());
 	VectorColor = FColorList::Yellow;
 	FLinearColor& RotatorColor = ColorSettings.StructColors.FindOrAdd(TBaseStructure<FRotator>::Get()->GetStructPathName().ToString());
diff --git a/Engine/Source/Runtime/AdvancedWidgets/Private/Styling/ColorGradingSpinBoxStyle.cpp b/Engine/Source/Runtime/AdvancedWidgets/Private/Styling/ColorGradingSpinBoxStyle.cpp
new file mode 100644
--- /dev/null
+++ b/Engine/Source/Runtime/AdvancedWidgets/Private/Styling/ColorGradingSpinBoxStyle.cpp
@@ -0,0 +1,26 @@
+// Copyright Epic Games, Inc. All Rights Reserved.
+
+#include "Styling/ColorGradingSpinBoxStyle.h"
+#include "Brushes/SlateNoResource.h"
+
+FColorGradingSpinBoxStyle::FColorGradingSpinBoxStyle()
+	: ActiveBorderBrush(FSlateOptionalBrush())
+	, SelectorWidth(1.f)
+{
+}
+
+void FColorGradingSpinBoxStyle::GetResources(TArray< const FSlateBrush* >& OutBrushes) const
+{
+	OutBrushes.Add(&BorderBrush);
+	OutBrushes.Add(&HoveredBorderBrush);
+	OutBrushes.Add(&ActiveBorderBrush);
+	OutBrushes.Add(&SelectorBrush);
+}
+
+const FName FColorGradingSpinBoxStyle::TypeName(TEXT("FColorGradingSpinBoxStyle"));
+
+const FColorGradingSpinBoxStyle& FColorGradingSpinBoxStyle::GetDefault()
+{
+	static FColorGradingSpinBoxStyle Default;
+	return Default;
+}
\ No newline at end of file
diff --git a/Engine/Source/Runtime/AdvancedWidgets/Private/Util/ColorGradingUtil.cpp b/Engine/Source/Runtime/AdvancedWidgets/Private/Util/ColorGradingUtil.cpp
new file mode 100644
--- /dev/null
+++ b/Engine/Source/Runtime/AdvancedWidgets/Private/Util/ColorGradingUtil.cpp
@@ -0,0 +1,52 @@
+// Copyright Epic Games, Inc. All Rights Reserved.
+
+#include "Util/ColorGradingUtil.h"
+#include "Misc/AssertionMacros.h"
+
+namespace UE::ColorGrading
+{
+
+EColorGradingComponent GetColorGradingComponent(EColorGradingColorDisplayMode DisplayMode, int Index)
+{
+	ensureMsgf(Index >= 0 && Index <= 3, TEXT("Color grading component index %d is outside of expected range (0 to 3)"), Index);
+
+	if (Index == 3)
+	{
+		return EColorGradingComponent::Luminance;
+	}
+
+	switch (DisplayMode)
+	{
+	case EColorGradingColorDisplayMode::RGB:
+		switch (Index)
+		{
+		case 0:
+			return EColorGradingComponent::Red;
+
+		case 1:
+			return EColorGradingComponent::Green;
+
+		case 2:
+			return EColorGradingComponent::Blue;
+		}
+		break;
+
+	case EColorGradingColorDisplayMode::HSV:
+		switch (Index)
+		{
+		case 0:
+			return EColorGradingComponent::Hue;
+
+		case 1:
+			return EColorGradingComponent::Saturation;
+
+		case 2:
+			return EColorGradingComponent::Value;
+		}
+		break;
+	}
+
+	return EColorGradingComponent::Red;
+}
+
+} //namespace
diff --git a/Engine/Source/Runtime/AdvancedWidgets/Private/Widgets/ColorGrading/SColorGradingComponentSpinBox.cpp b/Engine/Source/Runtime/AdvancedWidgets/Private/Widgets/ColorGrading/SColorGradingComponentSpinBox.cpp
new file mode 100644
--- /dev/null
+++ b/Engine/Source/Runtime/AdvancedWidgets/Private/Widgets/ColorGrading/SColorGradingComponentSpinBox.cpp
@@ -0,0 +1,728 @@
+// Copyright Epic Games, Inc. All Rights Reserved.
+
+#include "Widgets/ColorGrading/SColorGradingComponentSpinBox.h"
+
+namespace UE::ColorGrading
+{
+
+TArray<FLinearColor> SColorGradingComponentSpinBox::HueGradientColors = {};
+
+SColorGradingComponentSpinBox::~SColorGradingComponentSpinBox()
+{
+	if (bDragging || PointerDraggingSliderIndex != INDEX_NONE)
+	{
+		CancelMouseCapture();
+	}
+}
+
+void SColorGradingComponentSpinBox::Construct(const FArguments& InArgs)
+{
+	check(InArgs._Style);
+
+	Style = InArgs._Style;
+	Component = InArgs._Component;
+	ColorGradingMode = InArgs._ColorGradingMode;
+	bAllowSpin = InArgs._AllowSpin;
+
+	Interface = InArgs._TypeInterface.IsValid() ? InArgs._TypeInterface : MakeShareable(new TDefaultNumericTypeInterface<float>);
+
+	ValueAttribute = InArgs._Value;
+	OnValueChanged = InArgs._OnValueChanged;
+	OnBeginSliderMovement = InArgs._OnBeginSliderMovement;
+	OnEndSliderMovement = InArgs._OnEndSliderMovement;
+	OnQueryCurrentColor = InArgs._OnQueryCurrentColor;
+
+	MinValue = InArgs._MinValue;
+	MaxValue = InArgs._MaxValue;
+	MinSliderValue = (InArgs._MinSliderValue.Get().IsSet()) ? InArgs._MinSliderValue : MinValue;
+	MaxSliderValue = (InArgs._MaxSliderValue.Get().IsSet()) ? InArgs._MaxSliderValue : MaxValue;
+
+	AlwaysUsesDeltaSnap = InArgs._AlwaysUsesDeltaSnap;
+
+	SupportDynamicSliderMaxValue = InArgs._SupportDynamicSliderMaxValue;
+	SupportDynamicSliderMinValue = InArgs._SupportDynamicSliderMinValue;
+	OnDynamicSliderMaxValueChanged = InArgs._OnDynamicSliderMaxValueChanged;
+	OnDynamicSliderMinValueChanged = InArgs._OnDynamicSliderMinValueChanged;
+
+	CachedExternalValue = ValueAttribute.Get();
+
+	InternalValue = (double)CachedExternalValue;
+
+	if (SupportDynamicSliderMaxValue.Get() && CachedExternalValue > GetMaxSliderValue())
+	{
+		ApplySliderMaxValueChanged(float(CachedExternalValue - GetMaxSliderValue()), true);
+	}
+	else if (SupportDynamicSliderMinValue.Get() && CachedExternalValue < GetMinSliderValue())
+	{
+		ApplySliderMinValueChanged(float(CachedExternalValue - GetMinSliderValue()), true);
+	}
+
+	UpdateIsSpinRangeUnlimited();
+
+	SliderExponent = InArgs._SliderExponent;
+
+	SliderExponentNeutralValue = InArgs._SliderExponentNeutralValue;
+
+	DistanceDragged = 0.0f;
+	PreDragValue = float();
+
+	Delta = InArgs._Delta;
+	ShiftMultiplier = InArgs._ShiftMultiplier;
+	CtrlMultiplier = InArgs._CtrlMultiplier;
+	Sensitivity = InArgs._Sensitivity;
+	LinearDeltaSensitivity = InArgs._LinearDeltaSensitivity;
+
+	BorderHoveredBrush = &InArgs._Style->HoveredBorderBrush;
+	BorderBrush = &InArgs._Style->BorderBrush;
+	BorderActiveBrush = InArgs._Style->ActiveBorderBrush.IsSet() ? &InArgs._Style->ActiveBorderBrush : BorderHoveredBrush;
+	SelectorBrush = &InArgs._Style->SelectorBrush;
+	SelectorWidth = &InArgs._Style->SelectorWidth;
+
+	bDragging = false;
+	PointerDraggingSliderIndex = INDEX_NONE;
+
+	bIsTextChanging = false;
+}
+
+int32 SColorGradingComponentSpinBox::OnPaint(const FPaintArgs& Args, const FGeometry& AllottedGeometry, const FSlateRect& MyCullingRect, FSlateWindowElementList& OutDrawElements, int32 LayerId, const FWidgetStyle& InWidgetStyle, bool bParentEnabled) const
+{
+	const bool bEnabled = ShouldBeEnabled(bParentEnabled);
+	const ESlateDrawEffect DrawEffects = (bEnabled ? ESlateDrawEffect::None : ESlateDrawEffect::DisabledEffect);
+
+	const bool bActiveFeedback = bDragging;
+
+	const FSlateBrush* BorderImage = bActiveFeedback ? BorderActiveBrush : (IsHovered() && bAllowSpin) ?
+		BorderHoveredBrush :
+		BorderBrush;
+
+	// Gradient
+	{
+		const TArray<FLinearColor> Colors = GetGradientColors();
+		int32 NumColors = Colors.Num();
+
+		if (NumColors > 0)
+		{
+			TArray<FSlateGradientStop> GradientStops;
+
+			for (int32 ColorIndex = 0; ColorIndex < NumColors; ++ColorIndex)
+			{
+				GradientStops.Add(FSlateGradientStop(AllottedGeometry.GetLocalSize() * (float(ColorIndex) / (NumColors - 1)), Colors[ColorIndex]));
+			}
+
+			FSlateDrawElement::MakeGradient(
+				OutDrawElements,
+				LayerId++,
+				AllottedGeometry.ToPaintGeometry(),
+				GradientStops,
+				EOrientation::Orient_Vertical,
+				DrawEffects | ESlateDrawEffect::NoGamma, // Gradient colors are in linear space, so disable gamma to let them blend properly
+				FVector4f(BorderBrush->OutlineSettings.CornerRadii)
+			);
+		}
+	}
+
+	// Selector
+	if (bAllowSpin)
+	{
+		const float SelectorLayer = LayerId++;
+		const float Value = ValueAttribute.Get();
+
+		float FractionFilled = Fraction((double)Value, (double)GetMinSliderValue(), (double)GetMaxSliderValue());
+		const float CachedSliderExponent = SliderExponent.Get();
+		if (!FMath::IsNearlyEqual(CachedSliderExponent, 1.f))
+		{
+			if (SliderExponentNeutralValue.IsSet() && SliderExponentNeutralValue.Get() > GetMinSliderValue() && SliderExponentNeutralValue.Get() < GetMaxSliderValue())
+			{
+				//Compute a log curve on both side of the neutral value
+				float StartFractionFilled = Fraction((double)SliderExponentNeutralValue.Get(), (double)GetMinSliderValue(), (double)GetMaxSliderValue());
+				FractionFilled = SpinBoxComputeExponentSliderFraction(FractionFilled, StartFractionFilled, CachedSliderExponent);
+			}
+			else
+			{
+				FractionFilled = 1.0f - FMath::Pow(1.0f - FractionFilled, CachedSliderExponent);
+			}
+		}
+
+		const FVector2f AllottedGeometrySize = AllottedGeometry.GetLocalSize();
+		const FVector2f SelectorSize(*SelectorWidth, AllottedGeometrySize.Y - (BorderBrush->OutlineSettings.Width * 2));
+		const float SelectorRange = AllottedGeometrySize.X - SelectorSize.X;
+
+		const FVector2f SelectorOffset(
+			SelectorRange * FractionFilled,
+			1.0f
+		);
+
+		// Draw the selector's center
+		FSlateDrawElement::MakeBox(
+			OutDrawElements,
+			SelectorLayer,
+			AllottedGeometry.ToPaintGeometry(
+				SelectorSize,
+				FSlateLayoutTransform(SelectorOffset)
+			),
+			SelectorBrush,
+			DrawEffects,
+			SelectorBrush->GetTint(InWidgetStyle) * InWidgetStyle.GetColorAndOpacityTint()
+		);
+	}
+
+	// Border
+	FSlateDrawElement::MakeBox(
+		OutDrawElements,
+		LayerId++,
+		AllottedGeometry.ToPaintGeometry(),
+		BorderImage,
+		DrawEffects,
+		BorderImage->GetTint(InWidgetStyle) * InWidgetStyle.GetColorAndOpacityTint()
+	);
+
+	return LayerId;
+}
+
+const bool SColorGradingComponentSpinBox::CommitWithMultiplier(const FPointerEvent& MouseEvent)
+{
+	return MouseEvent.IsShiftDown() || MouseEvent.IsControlDown();
+}
+
+FReply SColorGradingComponentSpinBox::OnMouseButtonDown(const FGeometry& MyGeometry, const FPointerEvent& MouseEvent)
+{
+	if (bAllowSpin && MouseEvent.GetEffectingButton() == EKeys::LeftMouseButton && PointerDraggingSliderIndex == INDEX_NONE)
+	{
+		DistanceDragged = 0.f;
+		PreDragValue = ValueAttribute.Get();
+		InternalValue = (double)PreDragValue;
+		PointerDraggingSliderIndex = MouseEvent.GetPointerIndex();
+		CachedMousePosition = MouseEvent.GetScreenSpacePosition().IntPoint();
+
+		FReply ReturnReply = FReply::Handled().CaptureMouse(SharedThis(this)).UseHighPrecisionMouseMovement(SharedThis(this)).SetUserFocus(SharedThis(this), EFocusCause::Mouse);
+		if (bPreventThrottling)
+		{
+			ReturnReply.PreventThrottling();
+		}
+		return ReturnReply;
+	}
+	
+	return FReply::Unhandled();
+}
+
+FReply SColorGradingComponentSpinBox::OnMouseButtonUp(const FGeometry& MyGeometry, const FPointerEvent& MouseEvent)
+{
+	if (bAllowSpin && MouseEvent.GetEffectingButton() == EKeys::LeftMouseButton && PointerDraggingSliderIndex == MouseEvent.GetPointerIndex())
+	{
+		if (!this->HasMouseCapture())
+		{
+			// Lost Capture - ensure reset
+			bDragging = false;
+			PointerDraggingSliderIndex = INDEX_NONE;
+
+			return FReply::Unhandled();
+		}
+
+		if (bDragging)
+		{
+			float CurrentDelta = Delta.Get();
+			if (CurrentDelta != float() && !CommitWithMultiplier(MouseEvent))
+			{
+				InternalValue = FMath::GridSnap(InternalValue, (double)CurrentDelta);
+			}
+
+			const float CurrentValue = RoundIfIntegerValue(InternalValue);
+			NotifyValueCommitted(CurrentValue);
+		}
+
+		bDragging = false;
+		PointerDraggingSliderIndex = INDEX_NONE;
+
+		FReply Reply = FReply::Handled().ReleaseMouseCapture();
+
+		if (!MouseEvent.IsTouchEvent())
+		{
+			Reply.SetMousePos(CachedMousePosition);
+		}
+
+		return Reply;
+	}
+
+	return FReply::Unhandled();
+}
+
+void SColorGradingComponentSpinBox::ApplySliderMaxValueChanged(float SliderDeltaToAdd, bool UpdateOnlyIfHigher)
+{
+	check(SupportDynamicSliderMaxValue.Get());
+
+	float NewMaxSliderValue = std::numeric_limits<float>::min();
+
+	if (MaxSliderValue.IsSet() && MaxSliderValue.Get().IsSet())
+	{
+		NewMaxSliderValue = GetMaxSliderValue();
+
+		if ((NewMaxSliderValue + (float)SliderDeltaToAdd > GetMaxSliderValue() && UpdateOnlyIfHigher) || !UpdateOnlyIfHigher)
+		{
+			NewMaxSliderValue += (float)SliderDeltaToAdd;
+
+			if (!MaxSliderValue.IsBound()) // simple value so we can update it without breaking the mechanic otherwise it must be handle by the callback implementer
+			{
+				SetMaxSliderValue(NewMaxSliderValue);
+			}
+		}
+	}
+
+	if (OnDynamicSliderMaxValueChanged.IsBound())
+	{
+		OnDynamicSliderMaxValueChanged.Execute(NewMaxSliderValue, TWeakPtr<SWidget>(AsShared()), true, UpdateOnlyIfHigher);
+	}
+}
+
+void SColorGradingComponentSpinBox::ApplySliderMinValueChanged(float SliderDeltaToAdd, bool UpdateOnlyIfLower)
+{
+	check(SupportDynamicSliderMaxValue.Get());
+
+	float NewMinSliderValue = std::numeric_limits<float>::min();
+
+	if (MinSliderValue.IsSet() && MinSliderValue.Get().IsSet())
+	{
+		NewMinSliderValue = GetMinSliderValue();
+
+		if ((NewMinSliderValue + (float)SliderDeltaToAdd < GetMinSliderValue() && UpdateOnlyIfLower) || !UpdateOnlyIfLower)
+		{
+			NewMinSliderValue += (float)SliderDeltaToAdd;
+
+			if (!MinSliderValue.IsBound()) // simple value so we can update it without breaking the mechanic otherwise it must be handle by the callback implementer
+			{
+				SetMinSliderValue(NewMinSliderValue);
+			}
+		}
+	}
+
+	if (OnDynamicSliderMinValueChanged.IsBound())
+	{
+		OnDynamicSliderMinValueChanged.Execute(NewMinSliderValue, TWeakPtr<SWidget>(AsShared()), true, UpdateOnlyIfLower);
+	}
+}
+
+/**
+* The system calls this method to notify the widget that a mouse moved within it. This event is bubbled.
+*
+* @param MyGeometry The Geometry of the widget receiving the event
+* @param MouseEvent Information about the input event
+* @return Whether the event was handled along with possible requests for the system to take action.
+*/
+
+FReply SColorGradingComponentSpinBox::OnMouseMove(const FGeometry& MyGeometry, const FPointerEvent& MouseEvent)
+{
+	if (bAllowSpin && PointerDraggingSliderIndex == MouseEvent.GetPointerIndex())
+	{
+		if (!this->HasMouseCapture())
+		{
+			// Lost the mouse capture - ensure reset
+			bDragging = false;
+			PointerDraggingSliderIndex = INDEX_NONE;
+
+			return FReply::Unhandled();
+		}
+
+		if (!bDragging)
+		{
+			DistanceDragged += (float)FMath::Abs(MouseEvent.GetCursorDelta().X);
+			if (DistanceDragged > FSlateApplication::Get().GetDragTriggerDistance())
+			{
+				bDragging = true;
+				OnBeginSliderMovement.ExecuteIfBound();
+			}
+
+			// Cache the mouse, even if not dragging cache it.
+			CachedMousePosition = MouseEvent.GetScreenSpacePosition().IntPoint();
+		}
+		else
+		{
+			double NewValue = 0.0;
+
+			// Increments the spin based on delta mouse movement.
+
+			// A minimum slider width to use for calculating deltas in the slider-range space
+			const float MinSliderWidth = 100.f;
+			float SliderWidthInSlateUnits = FMath::Max((float)MyGeometry.GetDrawSize().X, MinSliderWidth);
+
+			if (MouseEvent.IsAltDown())
+			{
+				float DeltaToAdd = (float)MouseEvent.GetCursorDelta().X / SliderWidthInSlateUnits;
+
+				if (SupportDynamicSliderMaxValue.Get() && (float)InternalValue == GetMaxSliderValue())
+				{
+					ApplySliderMaxValueChanged(DeltaToAdd, false);
+				}
+				else if (SupportDynamicSliderMinValue.Get() && (float)InternalValue == GetMinSliderValue())
+				{
+					ApplySliderMinValueChanged(DeltaToAdd, false);
+				}
+			}
+
+			ECommitMethod CommitMethod = CommittedViaSpin;
+
+			double Step = 0.1;
+
+			if (MouseEvent.IsControlDown())
+			{
+				Step *= CtrlMultiplier.Get();
+				CommitMethod = CommittedViaSpinMultiplier;
+			}
+			else if (MouseEvent.IsShiftDown())
+			{
+				Step *= ShiftMultiplier.Get();
+				CommitMethod = CommittedViaSpinMultiplier;
+			}
+
+			const float MouseXMovement = MouseEvent.GetCursorDelta().X * Sensitivity.Get(1.f);
+
+			//if we have a range to draw in
+			if (!bUnlimitedSpinRange)
+			{
+				bool HasValidExponentNeutralValue = SliderExponentNeutralValue.IsSet() && SliderExponentNeutralValue.Get() > GetMinSliderValue() && SliderExponentNeutralValue.Get() < GetMaxSliderValue();
+
+				const float CachedSliderExponent = SliderExponent.Get();
+				// The amount currently filled in the spinbox, needs to be calculated to do deltas correctly.
+				float FractionFilled = Fraction(InternalValue, (double)GetMinSliderValue(), (double)GetMaxSliderValue());
+
+				if (!FMath::IsNearlyEqual(CachedSliderExponent, 1.0f))
+				{
+					if (HasValidExponentNeutralValue)
+					{
+						//Compute a log curve on both side of the neutral value
+						float StartFractionFilled = Fraction((double)SliderExponentNeutralValue.Get(), (double)GetMinSliderValue(), (double)GetMaxSliderValue());
+						FractionFilled = SpinBoxComputeExponentSliderFraction(FractionFilled, StartFractionFilled, CachedSliderExponent);
+					}
+					else
+					{
+						FractionFilled = 1.0f - FMath::Pow(1.0f - FractionFilled, CachedSliderExponent);
+					}
+				}
+				FractionFilled *= SliderWidthInSlateUnits;
+
+				// Now add the delta to the fraction filled, this causes the spin.
+				FractionFilled += (float)(MouseXMovement * Step);
+
+				// Clamp the fraction to be within the bounds of the geometry.
+				FractionFilled = FMath::Clamp(FractionFilled, 0.0f, SliderWidthInSlateUnits);
+
+				// Convert the fraction filled to a percent.
+				float Percent = FMath::Clamp(FractionFilled / SliderWidthInSlateUnits, 0.0f, 1.0f);
+				if (!FMath::IsNearlyEqual(CachedSliderExponent, 1.0f))
+				{
+					// Have to convert the percent to the proper value due to the exponent component to the spin.
+					if (HasValidExponentNeutralValue)
+					{
+						//Compute a log curve on both side of the neutral value
+						float StartFractionFilled = Fraction(SliderExponentNeutralValue.Get(), GetMinSliderValue(), GetMaxSliderValue());
+						Percent = SpinBoxComputeExponentSliderFraction(Percent, StartFractionFilled, 1.0f / CachedSliderExponent);
+					}
+					else
+					{
+						Percent = 1.0f - FMath::Pow(1.0f - Percent, 1.0f / CachedSliderExponent);
+					}
+
+
+				}
+
+				NewValue = FMath::LerpStable<double>((double)GetMinSliderValue(), (double)GetMaxSliderValue(), Percent);
+			}
+			else
+			{
+				// If this control has a specified delta and sensitivity then we use that instead of the current value for determining how much to change.
+				const double Sign = MouseXMovement ? 1.0 : -1.0;
+
+				if (LinearDeltaSensitivity.IsSet() && LinearDeltaSensitivity.Get() != 0 && Delta.IsSet() && Delta.Get() > 0)
+				{
+					const double MouseDelta = FMath::Abs(MouseXMovement / (float)LinearDeltaSensitivity.Get());
+					NewValue = InternalValue + (Sign * MouseDelta * FMath::Pow((double)Delta.Get(), (double)SliderExponent.Get())) * Step;
+				}
+				else
+				{
+					const double MouseDelta = FMath::Abs(MouseXMovement / SliderWidthInSlateUnits);
+					const double CurrentValue = FMath::Clamp<double>(FMath::Abs(InternalValue), 1.0, (double)std::numeric_limits<float>::max());
+					NewValue = InternalValue + (Sign * MouseDelta * FMath::Pow((double)CurrentValue, (double)SliderExponent.Get())) * Step;
+				}
+			}
+
+			float RoundedNewValue = RoundIfIntegerValue(NewValue);
+			CommitValue(RoundedNewValue, NewValue, CommitMethod, ETextCommit::OnEnter);
+		}
+
+		return FReply::Handled();
+	}
+
+	return FReply::Unhandled();
+}
+
+FCursorReply SColorGradingComponentSpinBox::OnCursorQuery(const FGeometry& MyGeometry, const FPointerEvent& CursorEvent) const
+{
+	if (!bAllowSpin)
+	{
+		return FCursorReply::Cursor(EMouseCursor::Default);
+	}
+
+	return bDragging ?
+		FCursorReply::Cursor(EMouseCursor::None) :
+		FCursorReply::Cursor(EMouseCursor::ResizeLeftRight);
+}
+
+void SColorGradingComponentSpinBox::CommitValue(float NewValue, double NewSpinValue, ECommitMethod CommitMethod, ETextCommit::Type OriginalCommitInfo)
+{
+	if (!bAllowSpin)
+	{
+		return;
+	}
+
+	if (CommitMethod == CommittedViaSpin)
+	{
+		const float LocalMinSliderValue = GetMinSliderValue();
+		const float LocalMaxSliderValue = GetMaxSliderValue();
+		NewValue = FMath::Clamp<float>(NewValue, LocalMinSliderValue, LocalMaxSliderValue);
+		NewSpinValue = FMath::Clamp<double>(NewSpinValue, (double)LocalMinSliderValue, (double)LocalMaxSliderValue);
+	}
+
+	{
+		const float LocalMinValue = GetMinValue();
+		const float LocalMaxValue = GetMaxValue();
+		NewValue = FMath::Clamp<float>(NewValue, LocalMinValue, LocalMaxValue);
+		NewSpinValue = FMath::Clamp<double>(NewSpinValue, (double)LocalMinValue, (double)LocalMaxValue);
+	}
+
+	if (!ValueAttribute.IsBound())
+	{
+		ValueAttribute.Set(NewValue);
+	}
+
+	// If not in spin mode, there is no need to jump to the value from the external source, continue to use the committed value.
+	if (CommitMethod == CommittedViaSpin)
+	{
+		const float CurrentValue = ValueAttribute.Get();
+		// This will detect if an external force has changed the value. Internally it will abandon the delta calculated this tick and update the internal value instead.
+		if (CurrentValue != CachedExternalValue)
+		{
+			NewValue = CurrentValue;
+			NewSpinValue = (double)CurrentValue;
+		}
+	}
+
+	// Update the internal value, this needs to be done before rounding.
+	InternalValue = NewSpinValue;
+
+	const bool bAlwaysUsesDeltaSnap = GetAlwaysUsesDeltaSnap();
+	// If needed, round this value to the delta. Internally the value is not held to the Delta but externally it appears to be.
+	if (CommitMethod == CommittedViaSpin || bAlwaysUsesDeltaSnap)
+	{
+		float CurrentDelta = Delta.Get();
+		if (CurrentDelta != float())
+		{
+			NewValue = FMath::GridSnap<float>(NewValue, CurrentDelta); // snap numeric point value to nearest Delta
+		}
+	}
+
+	// Update the max slider value based on the current value if we're in dynamic mode
+	if (SupportDynamicSliderMaxValue.Get() && ValueAttribute.Get() > GetMaxSliderValue())
+	{
+		ApplySliderMaxValueChanged(float(ValueAttribute.Get() - GetMaxSliderValue()), true);
+	}
+	else if (SupportDynamicSliderMinValue.Get() && ValueAttribute.Get() < GetMinSliderValue())
+	{
+		ApplySliderMinValueChanged(float(ValueAttribute.Get() - GetMinSliderValue()), true);
+	}
+
+	OnValueChanged.ExecuteIfBound(NewValue);
+
+	if (!ValueAttribute.IsBound())
+	{
+		ValueAttribute.Set(NewValue);
+	}
+
+	// Update the cache of the external value to what the user believes the value is now.
+	const float CurrentValue = ValueAttribute.Get();
+	if (CachedExternalValue != CurrentValue)
+	{
+		CachedExternalValue = ValueAttribute.Get();
+	}
+
+	// This ensures that dragging is cleared if focus has been removed from this widget in one of the delegate calls, such as when spawning a modal dialog.
+	if (!this->HasMouseCapture())
+	{
+		bDragging = false;
+		PointerDraggingSliderIndex = INDEX_NONE;
+	}
+}
+
+void SColorGradingComponentSpinBox::NotifyValueCommitted(float CurrentValue) const
+{
+	// The internal value will have been clamped and rounded to the delta at this point, but integer values may still need to be rounded
+	// if the delta is 0.
+	OnEndSliderMovement.ExecuteIfBound(CurrentValue);
+}
+
+float SColorGradingComponentSpinBox::Fraction(double InValue, double InMinValue, double InMaxValue)
+{
+	const double HalfMax = InMaxValue * 0.5;
+	const double HalfMin = InMinValue * 0.5;
+	const double HalfVal = InValue * 0.5;
+
+	return (float)FMath::Clamp((HalfVal - HalfMin) / (HalfMax - HalfMin), 0.0, 1.0);
+}
+
+const TArray<FLinearColor>& SColorGradingComponentSpinBox::GetHueGradientColors()
+{
+	if (HueGradientColors.IsEmpty())
+	{
+		for (int32 i = 0; i < 7; ++i)
+		{
+			HueGradientColors.Add(FLinearColor((i % 6) * 60.f, 1.f, 1.f).HSVToLinearRGB());
+		}
+	}
+
+	return HueGradientColors;
+}
+
+float SColorGradingComponentSpinBox::RoundIfIntegerValue(double ValueToRound) const
+{
+	constexpr bool bIsIntegral = TIsIntegral<float>::Value;
+	constexpr bool bCanBeRepresentedInDouble = std::numeric_limits<double>::digits >= std::numeric_limits<float>::digits;
+	if (bIsIntegral && !bCanBeRepresentedInDouble)
+	{
+		return (float)FMath::Clamp<double>(FMath::FloorToDouble(ValueToRound + 0.5), -1.0 * (double)(1ll << std::numeric_limits<double>::digits), (double)(1ll << std::numeric_limits<double>::digits));
+	}
+	else if (bIsIntegral)
+	{
+		return (float)FMath::Clamp<double>(FMath::FloorToDouble(ValueToRound + 0.5), (double)std::numeric_limits<float>::lowest(), (double)std::numeric_limits<float>::max());
+	}
+	else
+	{
+		return (float)FMath::Clamp<double>(ValueToRound, (double)std::numeric_limits<float>::lowest(), (double)std::numeric_limits<float>::max());
+	}
+}
+
+void SColorGradingComponentSpinBox::CancelMouseCapture()
+{
+	bDragging = false;
+	PointerDraggingSliderIndex = INDEX_NONE;
+
+	InternalValue = (double)PreDragValue;
+	NotifyValueCommitted(PreDragValue);
+}
+
+TArray<FLinearColor> SColorGradingComponentSpinBox::GetGradientColors() const
+{
+	const EColorGradingComponent DisplayedComponent = Component.Get();
+	const bool bIsOffset = ColorGradingMode.Get() == EColorGradingModes::Offset;
+
+	// Create gradients
+	switch (DisplayedComponent)
+	{
+	case EColorGradingComponent::Red:
+		if (bIsOffset)
+		{
+			return
+			{
+				FLinearColor(0, 1, 1, 1),
+				FLinearColor::Black,
+				FLinearColor::Red
+			};
+		}
+
+		return
+		{
+			FLinearColor::Black,
+			FLinearColor::Red
+		};
+
+	case EColorGradingComponent::Green:
+		if (bIsOffset)
+		{
+			return
+			{
+				FLinearColor(1, 0, 1, 1),
+				FLinearColor::Black,
+				FLinearColor::Green
+			};
+		}
+
+		return
+		{
+			FLinearColor::Black,
+			FLinearColor::Green
+		};
+
+	case EColorGradingComponent::Blue:
+		if (bIsOffset)
+		{
+			return
+			{
+				FLinearColor(1, 1, 0, 1),
+				FLinearColor::Black,
+				FLinearColor::Blue
+			};
+		}
+
+		return
+		{
+			FLinearColor::Black,
+			FLinearColor::Blue
+		};
+
+	case EColorGradingComponent::Luminance:
+		{
+			const FLinearColor RGBColor = bIsOffset ? FLinearColor::White : GetCurrentRGBColor();
+			return
+			{
+				FLinearColor::Black,
+				RGBColor
+			};
+		}
+
+	case EColorGradingComponent::Hue:
+		return GetHueGradientColors();
+
+	case EColorGradingComponent::Saturation:
+		{
+			const FLinearColor HSVColor = GetCurrentHSVColor();
+			return
+			{
+				FLinearColor(HSVColor.R, 0.f, HSVColor.B).HSVToLinearRGB(),
+				FLinearColor(HSVColor.R, 0.5f, HSVColor.B).HSVToLinearRGB(),
+				FLinearColor(HSVColor.R, 1.f, HSVColor.B).HSVToLinearRGB()
+			};
+		}
+
+	case EColorGradingComponent::Value:
+		{
+			const FLinearColor HSVColor = GetCurrentHSVColor();
+			return
+			{
+				FLinearColor(HSVColor.R, HSVColor.G, 0.f).HSVToLinearRGB(),
+				FLinearColor(HSVColor.R, HSVColor.G, 1.f).HSVToLinearRGB(),
+				FLinearColor(HSVColor.R, HSVColor.G, 2.f).HSVToLinearRGB()
+			};
+		}
+	}
+
+	return TArray<FLinearColor>();
+}
+
+FLinearColor SColorGradingComponentSpinBox::GetCurrentRGBColor() const
+{
+	if (OnQueryCurrentColor.IsBound())
+	{
+		FVector4 ColorComponents;
+		OnQueryCurrentColor.Execute(ColorComponents);
+
+		// This component is luminance, but our RGB representation will use it as alpha, so set it to 100%
+		ColorComponents.W = 1;
+
+		return FLinearColor(ColorComponents);
+	}
+
+	return FLinearColor::White;
+}
+
+FLinearColor SColorGradingComponentSpinBox::GetCurrentHSVColor() const
+{
+	return GetCurrentRGBColor().LinearRGBToHSV();
+}
+
+} //namespace
\ No newline at end of file
diff --git a/Engine/Source/Runtime/AdvancedWidgets/Private/Widgets/ColorGrading/SColorGradingComponentSpinBox.h b/Engine/Source/Runtime/AdvancedWidgets/Private/Widgets/ColorGrading/SColorGradingComponentSpinBox.h
new file mode 100644
--- /dev/null
+++ b/Engine/Source/Runtime/AdvancedWidgets/Private/Widgets/ColorGrading/SColorGradingComponentSpinBox.h
@@ -0,0 +1,359 @@
+// Copyright Epic Games, Inc. All Rights Reserved.
+
+#pragma once
+
+#include "Framework/ColorGrading/ColorGradingCommon.h"
+#include "Input/CursorReply.h"
+#include "Input/Reply.h"
+#include "InputCoreTypes.h"
+#include "HAL/Platform.h"
+#include "Misc/Attribute.h"
+#include "Styling/AdvancedWidgetsStyle.h"
+#include "Styling/ColorGradingSpinBoxStyle.h"
+#include "Styling/CoreStyle.h"
+#include "Widgets/DeclarativeSyntaxSupport.h"
+#include "Widgets/Input/NumericTypeInterface.h"
+#include "Widgets/Input/SSpinBox.h"
+#include "Widgets/SCompoundWidget.h"
+
+namespace UE::ColorGrading
+{
+
+/**
+ * A modified SpinBox that shows a selector along a gradient to indicate a color component for color grading.
+ * Allows spinning via mouse drag if enabled, but does not allow text entry like a standard SpinBox.
+ */
+class SColorGradingComponentSpinBox
+	: public SCompoundWidget
+{
+public:
+
+	/** Notification for numeric value change */
+	DECLARE_DELEGATE_OneParam(FOnValueChanged, float);
+
+	/** Notification when the max/min spinner values are changed (only apply if SupportDynamicSliderMaxValue or SupportDynamicSliderMinValue are true) */
+	DECLARE_DELEGATE_FourParams(FOnDynamicSliderMinMaxValueChanged, float, TWeakPtr<SWidget>, bool, bool);
+
+	SLATE_BEGIN_ARGS(SColorGradingComponentSpinBox)
+		: _Style(&UE::AdvancedWidgets::FAdvancedWidgetsStyle::Get().GetWidgetStyle<FColorGradingSpinBoxStyle>("ColorGradingSpinBox"))
+		, _Value(0)
+		, _Component(EColorGradingComponent::Red)
+		, _ColorGradingMode(EColorGradingModes::Invalid)
+		, _MinValue(0)
+		, _MaxValue(2)
+		, _Delta(0)
+		, _ShiftMultiplier(10.f)
+		, _CtrlMultiplier(0.1f)
+		, _Sensitivity(1.0f)
+		, _SupportDynamicSliderMaxValue(false)
+		, _SupportDynamicSliderMinValue(false)
+		, _SliderExponent(1.f)
+		, _OnValueChanged()
+		, _OnQueryCurrentColor()
+		, _AllowSpin(true)
+		{}
+
+		/** The style used to draw this spinbox */
+		SLATE_STYLE_ARGUMENT(FColorGradingSpinBoxStyle, Style)
+
+		/** The value to display */
+		SLATE_ATTRIBUTE(float, Value)
+		/** The component being displayed */
+		SLATE_ATTRIBUTE(EColorGradingComponent, Component)
+		/** The mode of the associated color grading wheel */
+		SLATE_ATTRIBUTE(EColorGradingModes, ColorGradingMode)
+		/** The minimum value that can be entered into the text edit box */
+		SLATE_ATTRIBUTE(TOptional<float>, MinValue)
+		/** The maximum value that can be entered into the text edit box */
+		SLATE_ATTRIBUTE(TOptional<float>, MaxValue)
+		/** The minimum value that can be specified by using the slider, defaults to MinValue */
+		SLATE_ATTRIBUTE(TOptional<float>, MinSliderValue)
+		/** The maximum value that can be specified by using the slider, defaults to MaxValue */
+		SLATE_ATTRIBUTE(TOptional<float>, MaxSliderValue)
+		/** Whether typed values should use delta snapping, defaults to false */
+		SLATE_ATTRIBUTE(bool, AlwaysUsesDeltaSnap)
+		/** Delta to increment the value as the slider moves.  If not specified will determine automatically */
+		SLATE_ATTRIBUTE(float, Delta)
+		/** Multiplier to use when shift is held down */
+		SLATE_ATTRIBUTE(float, ShiftMultiplier)
+		/** Multiplier to use when ctrl is held down */
+		SLATE_ATTRIBUTE(float, CtrlMultiplier)
+		/** Multiplier to apply to all mouse movement */
+		SLATE_ATTRIBUTE(float, Sensitivity)
+		/** If we're an unbounded spinbox, what value do we divide mouse movement by before multiplying by Delta. Requires Delta to be set. */
+		SLATE_ATTRIBUTE(int32, LinearDeltaSensitivity)
+		/** Tell us if we want to support dynamically changing of the max value using alt */
+		SLATE_ATTRIBUTE(bool, SupportDynamicSliderMaxValue)
+		/** Tell us if we want to support dynamically changing of the min value using alt */
+		SLATE_ATTRIBUTE(bool, SupportDynamicSliderMinValue)
+		/** Called right after the max slider value is changed (only relevant if SupportDynamicSliderMaxValue is true) */
+		SLATE_EVENT(FOnDynamicSliderMinMaxValueChanged, OnDynamicSliderMaxValueChanged)
+		/** Called right after the min slider value is changed (only relevant if SupportDynamicSliderMinValue is true) */
+		SLATE_EVENT(FOnDynamicSliderMinMaxValueChanged, OnDynamicSliderMinValueChanged)
+		/** Use exponential scale for the slider */
+		SLATE_ATTRIBUTE(float, SliderExponent)
+		/** When use exponential scale for the slider which is the neutral value */
+		SLATE_ATTRIBUTE(float, SliderExponentNeutralValue)
+		/** Step to increment or decrement the value by when scrolling the mouse wheel. If not specified will determine automatically */
+		SLATE_ATTRIBUTE(TOptional<float>, WheelStep)
+		/** Called when the value is changed by slider or typing */
+		SLATE_EVENT(FOnValueChanged, OnValueChanged)
+		/** Called right before the slider begins to move */
+		SLATE_EVENT(FSimpleDelegate, OnBeginSliderMovement)
+		/** Called right after the slider handle is released by the user */
+		SLATE_EVENT(FOnValueChanged, OnEndSliderMovement)
+		/** Callback to get the current FVector4 color value (used to update the background gradient) */
+		SLATE_EVENT(FOnGetCurrentVector4Value, OnQueryCurrentColor)
+		/** Provide custom type conversion functionality to this spin box */
+		SLATE_ARGUMENT(TSharedPtr<INumericTypeInterface<float>>, TypeInterface)
+		/** Whether or not the user should be able to change the value by dragging with the mouse cursor */
+		SLATE_ARGUMENT(bool, AllowSpin)
+
+	SLATE_END_ARGS()
+
+	SColorGradingComponentSpinBox() = default;
+
+	virtual ~SColorGradingComponentSpinBox();
+
+	/**
+	 * Construct the widget
+	 *
+	 * @param InArgs   A declaration from which to construct the widget
+	 */
+	void Construct(const FArguments& InArgs);
+
+	virtual int32 OnPaint(const FPaintArgs& Args, const FGeometry& AllottedGeometry, const FSlateRect& MyCullingRect, FSlateWindowElementList& OutDrawElements, int32 LayerId, const FWidgetStyle& InWidgetStyle, bool bParentEnabled) const override;
+
+	const bool CommitWithMultiplier(const FPointerEvent& MouseEvent);
+
+	/**
+	 * The system calls this method to notify the widget that a mouse button was pressed within it. This event is bubbled.
+	 *
+	 * @param MyGeometry The Geometry of the widget receiving the event
+	 * @param MouseEvent Information about the input event
+	 * @return Whether the event was handled along with possible requests for the system to take action.
+	 */
+	virtual FReply OnMouseButtonDown(const FGeometry& MyGeometry, const FPointerEvent& MouseEvent) override;
+
+	/**
+	 * The system calls this method to notify the widget that a mouse button was release within it. This event is bubbled.
+	 *
+	 * @param MyGeometry The Geometry of the widget receiving the event
+	 * @param MouseEvent Information about the input event
+	 * @return Whether the event was handled along with possible requests for the system to take action.
+	 */
+	virtual FReply OnMouseButtonUp(const FGeometry& MyGeometry, const FPointerEvent& MouseEvent) override;
+
+	void ApplySliderMaxValueChanged(float SliderDeltaToAdd, bool UpdateOnlyIfHigher);
+
+	void ApplySliderMinValueChanged(float SliderDeltaToAdd, bool UpdateOnlyIfLower);
+
+	/**
+	 * The system calls this method to notify the widget that a mouse moved within it. This event is bubbled.
+	 *
+	 * @param MyGeometry The Geometry of the widget receiving the event
+	 * @param MouseEvent Information about the input event
+	 * @return Whether the event was handled along with possible requests for the system to take action.
+	 */
+	virtual FReply OnMouseMove(const FGeometry& MyGeometry, const FPointerEvent& MouseEvent) override;
+
+	virtual FCursorReply OnCursorQuery(const FGeometry& MyGeometry, const FPointerEvent& CursorEvent) const override;
+
+	/** Return the Value attribute */
+	TAttribute<float> GetValueAttribute() const { return ValueAttribute; }
+
+	/** See the Value attribute */
+	float GetValue() const { return ValueAttribute.Get(); }
+	void SetValue(const TAttribute<float>& InValueAttribute)
+	{
+		ValueAttribute = InValueAttribute;
+		const float LocalValueAttribute = ValueAttribute.Get();
+		CommitValue(LocalValueAttribute, (double)LocalValueAttribute, ECommitMethod::CommittedViaCode, ETextCommit::Default);
+	}
+
+	/** See the MinValue attribute */
+	float GetMinValue() const { return MinValue.Get().Get(std::numeric_limits<float>::lowest()); }
+	void SetMinValue(const TAttribute<TOptional<float>>& InMinValue)
+	{
+		MinValue = InMinValue;
+		UpdateIsSpinRangeUnlimited();
+	}
+
+	/** See the MaxValue attribute */
+	float GetMaxValue() const { return MaxValue.Get().Get(std::numeric_limits<float>::max()); }
+	void SetMaxValue(const TAttribute<TOptional<float>>& InMaxValue)
+	{
+		MaxValue = InMaxValue;
+		UpdateIsSpinRangeUnlimited();
+	}
+
+	/** See the MinSliderValue attribute */
+	bool IsMinSliderValueBound() const { return MinSliderValue.IsBound(); }
+
+	float GetMinSliderValue() const { return MinSliderValue.Get().Get(std::numeric_limits<float>::lowest()); }
+	void SetMinSliderValue(const TAttribute<TOptional<float>>& InMinSliderValue)
+	{
+		MinSliderValue = (InMinSliderValue.Get().IsSet()) ? InMinSliderValue : MinValue;
+		UpdateIsSpinRangeUnlimited();
+	}
+
+	/** See the MaxSliderValue attribute */
+	bool IsMaxSliderValueBound() const { return MaxSliderValue.IsBound(); }
+
+	float GetMaxSliderValue() const { return MaxSliderValue.Get().Get(std::numeric_limits<float>::max()); }
+	void SetMaxSliderValue(const TAttribute<TOptional<float>>& InMaxSliderValue)
+	{
+		MaxSliderValue = (InMaxSliderValue.Get().IsSet()) ? InMaxSliderValue : MaxValue;;
+		UpdateIsSpinRangeUnlimited();
+	}
+
+	/** See the AlwaysUsesDeltaSnap attribute */
+	bool GetAlwaysUsesDeltaSnap() const { return AlwaysUsesDeltaSnap.Get(); }
+	void SetAlwaysUsesDeltaSnap(bool bNewValue) { AlwaysUsesDeltaSnap.Set(bNewValue); }
+
+	/** See the Delta attribute */
+	float GetDelta() const { return Delta.Get(); }
+	void SetDelta(float InDelta) { Delta.Set(InDelta); }
+
+	/** See the SliderExponent attribute */
+	float GetSliderExponent() const { return SliderExponent.Get(); }
+	void SetSliderExponent(const TAttribute<float>& InSliderExponent) { SliderExponent = InSliderExponent; }
+
+	const FColorGradingSpinBoxStyle* GetWidgetStyle() const { return Style; }
+	void SetWidgetStyle(const FColorGradingSpinBoxStyle* InStyle) { Style = InStyle; }
+	void InvalidateStyle() { Invalidate(EInvalidateWidgetReason::Layout); }
+
+protected:
+	/** How user changed the value in the spinbox */
+	enum ECommitMethod
+	{
+		CommittedViaSpin,
+		CommittedViaCode,
+		CommittedViaSpinMultiplier
+	};
+
+	/**
+	 * Call this method when the user's interaction has changed the value
+	 *
+	 * @param NewValue               Value resulting from the user's interaction
+	 * @param CommitMethod           Did the user type in the value or spin to it.
+	 * @param OriginalCommitInfo	 If the user typed in the value, information about the source of the commit
+	 */
+	void CommitValue(float NewValue, double NewSpinValue, ECommitMethod CommitMethod, ETextCommit::Type OriginalCommitInfo);
+
+	void NotifyValueCommitted(float CurrentValue) const;
+
+	/** Calculates range fraction. Possible to use on full numeric range  */
+	static float Fraction(double InValue, double InMinValue, double InMaxValue);
+
+private:
+
+	/** Get the gradient stops for a hue slider. These are lazily generated and cached for future calls. */
+	static const TArray<FLinearColor>& GetHueGradientColors();
+	static TArray<FLinearColor> HueGradientColors;
+
+	TAttribute<float> ValueAttribute;
+	FOnValueChanged OnValueChanged;
+	FSimpleDelegate OnBeginSliderMovement;
+	FOnValueChanged OnEndSliderMovement;
+	FOnGetCurrentVector4Value OnQueryCurrentColor;
+
+	/** Interface that defines conversion functionality for the templated type */
+	TSharedPtr<INumericTypeInterface<float>> Interface;
+
+	/** True when no range is specified, spinner can be spun indefinitely */
+	bool bUnlimitedSpinRange;
+	void UpdateIsSpinRangeUnlimited()
+	{
+		bUnlimitedSpinRange = !((MinValue.Get().IsSet() && MaxValue.Get().IsSet()) || (MinSliderValue.Get().IsSet() && MaxSliderValue.Get().IsSet()));
+	}
+
+	const FColorGradingSpinBoxStyle* Style;
+
+	const FSlateBrush* BorderHoveredBrush;
+	const FSlateBrush* BorderActiveBrush;
+	const FSlateBrush* BorderBrush;
+	const FSlateBrush* SelectorBrush;
+	const float* SelectorWidth;
+
+	bool bAllowSpin;
+	float DistanceDragged;
+	TAttribute<float> Delta;
+	TAttribute<float> ShiftMultiplier;
+	TAttribute<float> CtrlMultiplier;
+	TAttribute<float> Sensitivity;
+	TAttribute<int32> LinearDeltaSensitivity;
+	TAttribute<float> SliderExponent;
+	TAttribute<float> SliderExponentNeutralValue;
+	TAttribute<TOptional<float>> MinValue;
+	TAttribute<TOptional<float>> MaxValue;
+	TAttribute<TOptional<float>> MinSliderValue;
+	TAttribute<TOptional<float>> MaxSliderValue;
+	TAttribute<bool> AlwaysUsesDeltaSnap;
+	TAttribute<bool> SupportDynamicSliderMaxValue;
+	TAttribute<bool> SupportDynamicSliderMinValue;
+	FOnDynamicSliderMinMaxValueChanged OnDynamicSliderMaxValueChanged;
+	FOnDynamicSliderMinMaxValueChanged OnDynamicSliderMinValueChanged;
+	TAttribute<EColorGradingComponent> Component;
+	TAttribute<EColorGradingModes> ColorGradingMode;
+
+	/**
+	 * Rounds the submitted value to the correct value if it's an integer.
+	 * For int64, not all values can be represented by a double. We can only round until we reach that limit.
+	 * This function should only be used when we drag the value. We accept that we can't drag huge numbers.
+	 */
+	float RoundIfIntegerValue(double ValueToRound) const;
+
+	void CancelMouseCapture();
+
+	/** Generate the gradient stops to display in the background based on the current color and viewed component. */
+	TArray<FLinearColor> GetGradientColors() const;
+
+	/** Get the currently selected color in linear RGB space. */
+	FLinearColor GetCurrentRGBColor() const;
+
+	/** Get the currently selected color in HSV space. */
+	FLinearColor GetCurrentHSVColor() const;
+
+	/** Tracks which cursor is currently dragging the slider (e.g., the mouse cursor or a specific finger) */
+	int32 PointerDraggingSliderIndex;
+
+	/** Cached mouse position to restore after scrolling. */
+	FIntPoint CachedMousePosition;
+
+	/**
+	 * This value represents what the spinbox believes the value to be, regardless of delta and the user binding to an int.
+	 * The spinbox will always count using floats between values, this is important to keep it flowing smoothly and feeling right,
+	 * and most importantly not conflicting with the user truncating the value to an int.
+	 */
+	double InternalValue;
+
+	/** The state of InternalValue before a drag operation was started */
+	float PreDragValue;
+
+	/**
+	 * This is the cached value the user believes it to be (usually different due to truncation to an int). Used for identifying
+	 * external forces on the spinbox and syncing the internal value to them. Synced when a value is committed to the spinbox.
+	 */
+	float CachedExternalValue;
+
+	/** Whether the user is dragging the slider */
+	bool bDragging;
+
+	/** Re-entrant guard for the text changed handler */
+	bool bIsTextChanging;
+
+	/*
+	 * Holds whether or not to prevent throttling during mouse capture
+	 * When true, the viewport will be updated with every single change to the value during dragging
+	 */
+	bool bPreventThrottling;
+
+	/** Does this spin box have the mouse wheel feature enabled? */
+	bool bEnableWheel = true;
+
+	/** True to broadcast every time we type. */
+	bool bBroadcastValueChangesPerKey = false;
+};
+
+} //namespace
diff --git a/Engine/Source/Runtime/AdvancedWidgets/Private/Widgets/ColorGrading/SColorGradingComponentViewer.cpp b/Engine/Source/Runtime/AdvancedWidgets/Private/Widgets/ColorGrading/SColorGradingComponentViewer.cpp
new file mode 100644
--- /dev/null
+++ b/Engine/Source/Runtime/AdvancedWidgets/Private/Widgets/ColorGrading/SColorGradingComponentViewer.cpp
@@ -0,0 +1,246 @@
+// Copyright Epic Games, Inc. All Rights Reserved.
+
+#include "Widgets/ColorGrading/SColorGradingComponentViewer.h"
+
+#include "SColorGradingComponentSpinBox.h"
+#include "Math/NumericLimits.h"
+#include "Styling/AdvancedWidgetsStyle.h"
+#include "Widgets/Layout/SBox.h"
+
+#define LOCTEXT_NAMESPACE "ColorGradingEditor"
+
+namespace UE::ColorGrading
+{
+
+SColorGradingComponentViewer::~SColorGradingComponentViewer()
+{
+}
+
+void SColorGradingComponentViewer::Construct(const FArguments& InArgs)
+{
+	OptionalValue = InArgs._Value;
+	Component = InArgs._Component;
+	MinValue = InArgs._MinValue;
+	MaxValue = InArgs._MaxValue;
+	MinSliderValue = InArgs._MinSliderValue;
+	MaxSliderValue = InArgs._MaxSliderValue;
+
+	const bool bUseCompactDisplay = InArgs._UseCompactDisplay;
+	const auto& AdvancedWidgetsStyle = UE::AdvancedWidgets::FAdvancedWidgetsStyle::Get();
+
+	TSharedPtr<SHorizontalBox> HorizontalBox;
+
+	this->ChildSlot
+	[
+		SNew(SBox)
+		.HeightOverride(22)
+		.HAlign(HAlign_Fill)
+		.ToolTipText(this, &SColorGradingComponentViewer::GetComponentToolTipText)
+		[
+			SAssignNew(HorizontalBox, SHorizontalBox)
+		]
+	];
+
+	HorizontalBox->AddSlot()
+		.AutoWidth()
+		.HAlign(HAlign_Left)
+		.VAlign(VAlign_Center)
+		[
+			SNew(SBox)
+			.WidthOverride(8)
+			[
+				SNew(STextBlock)
+				.Text(this, &SColorGradingComponentViewer::GetComponentLabelText)
+			]
+		];
+
+	if (!bUseCompactDisplay)
+	{
+		HorizontalBox->AddSlot()
+			.FillWidth(1.0f)
+			.HAlign(HAlign_Fill)
+			.VAlign(VAlign_Fill)
+			.Padding(8.f, 0.f, 0.f, 0.f)
+			[
+				SAssignNew(GradientSpinBox, SColorGradingComponentSpinBox)
+				.Value(this, &SColorGradingComponentViewer::GetValue)
+				.Component(InArgs._Component)
+				.ColorGradingMode(InArgs._ColorGradingMode)
+				.OnValueChanged(InArgs._OnValueChanged)
+				.OnBeginSliderMovement(InArgs._OnBeginSliderMovement)
+				.OnEndSliderMovement(InArgs._OnEndSliderMovement)
+				.ShiftMultiplier(InArgs._ShiftMultiplier)
+				.CtrlMultiplier(InArgs._CtrlMultiplier)
+				.SupportDynamicSliderMinValue(InArgs._SupportDynamicSliderMinValue)
+				.SupportDynamicSliderMaxValue(InArgs._SupportDynamicSliderMaxValue)
+				.OnDynamicSliderMinValueChanged(InArgs._OnDynamicSliderMinValueChanged)
+				.OnDynamicSliderMaxValueChanged(InArgs._OnDynamicSliderMaxValueChanged)
+				.OnQueryCurrentColor(InArgs._OnQueryCurrentColor)
+				.MinValue(InArgs._MinValue)
+				.MaxValue(InArgs._MaxValue)
+				.MinSliderValue(InArgs._MinSliderValue)
+				.MaxSliderValue(InArgs._MaxSliderValue)
+				.SliderExponent(InArgs._SliderExponent)
+				.SliderExponentNeutralValue(InArgs._SliderExponentNeutralValue)
+				.Delta(InArgs._Delta)
+				.TypeInterface(InArgs._TypeInterface)
+				.IsEnabled(this, &SWidget::IsEnabled)
+				.AllowSpin(InArgs._AllowSpin)
+			];
+	}
+
+	TSharedPtr<SBox> NumericEntryContainer;
+	SAssignNew(NumericEntryContainer, SBox)
+	[
+		SAssignNew(NumericEntryBox, SNumericEntryBox<float>)
+		.SpinBoxStyle(&AdvancedWidgetsStyle.GetWidgetStyle<FSpinBoxStyle>("ColorGradingComponentViewer.NumericEntry"))
+		.EditableTextBoxStyle(&AdvancedWidgetsStyle.GetWidgetStyle<FEditableTextBoxStyle>("ColorGradingComponentViewer.NumericEntry.TextBox"))
+		.UndeterminedString(NSLOCTEXT("PropertyEditor", "MultipleValues", "Multiple Values"))
+		.Value(InArgs._Value)
+		.OnValueChanged(InArgs._OnValueChanged)
+		.OnBeginSliderMovement(InArgs._OnBeginSliderMovement)
+		.OnEndSliderMovement(InArgs._OnEndSliderMovement)
+		.AllowSpin(true)
+		.ShiftMultiplier(InArgs._ShiftMultiplier)
+		.CtrlMultiplier(InArgs._CtrlMultiplier)
+		.SupportDynamicSliderMinValue(InArgs._SupportDynamicSliderMinValue)
+		.SupportDynamicSliderMaxValue(InArgs._SupportDynamicSliderMaxValue)
+		.OnDynamicSliderMinValueChanged(InArgs._OnDynamicSliderMinValueChanged)
+		.OnDynamicSliderMaxValueChanged(InArgs._OnDynamicSliderMaxValueChanged)
+		.MinValue(InArgs._MinValue)
+		.MaxValue(InArgs._MaxValue)
+		.MinSliderValue(InArgs._MinSliderValue)
+		.MaxSliderValue(InArgs._MaxSliderValue)
+		.MinFractionalDigits(InArgs._MinFractionalDigits)
+		.MaxFractionalDigits(InArgs._MaxFractionalDigits)
+		.SliderExponent(InArgs._SliderExponent)
+		.SliderExponentNeutralValue(InArgs._SliderExponentNeutralValue)
+		.Delta(InArgs._Delta)
+		.TypeInterface(InArgs._TypeInterface)
+		.IsEnabled(this, &SWidget::IsEnabled)
+		.AllowSpin(InArgs._AllowSpin)
+	];
+
+
+	const FMargin NumericEntryPadding(bUseCompactDisplay ? 4.f : 8.f, 0.f, 0.f, 0.f);
+
+	if (bUseCompactDisplay)
+	{
+		HorizontalBox->AddSlot()
+			.HAlign(HAlign_Fill)
+			.VAlign(VAlign_Fill)
+			.FillWidth(1.0f)
+			.Padding(NumericEntryPadding)
+			[
+				NumericEntryContainer.ToSharedRef()
+			];
+	}
+	else
+	{
+		HorizontalBox->AddSlot()
+			.HAlign(HAlign_Fill)
+			.VAlign(VAlign_Fill)
+			.AutoWidth()
+			.Padding(NumericEntryPadding)
+			[
+				NumericEntryContainer.ToSharedRef()
+			];
+
+		NumericEntryContainer->SetWidthOverride(56);
+	}
+}
+
+float SColorGradingComponentViewer::GetMaxSliderValue() const
+{
+	if (MaxSliderValue.IsSet() && MaxSliderValue.Get().IsSet())
+	{
+		return MaxSliderValue.Get().GetValue();
+	}
+
+	return MaxValue.Get().Get(TNumericLimits<float>::Max());
+}
+
+float SColorGradingComponentViewer::GetMinSliderValue() const
+{
+	if (MinSliderValue.IsSet() && MinSliderValue.Get().IsSet())
+	{
+		return MinSliderValue.Get().GetValue();
+	}
+
+	return MinValue.Get().Get(TNumericLimits<float>::Min());
+}
+
+float SColorGradingComponentViewer::GetValue() const
+{
+	if (OptionalValue.IsSet())
+	{
+		if (const float* Value = OptionalValue.Get().GetPtrOrNull())
+		{
+			return *Value;
+		}
+	}
+
+	return 0.0f;
+}
+
+FText SColorGradingComponentViewer::GetComponentLabelText() const
+{
+	switch (Component.Get())
+	{
+	case EColorGradingComponent::Red:
+		return LOCTEXT("ColorWheel_RedComponentLabel", "R");
+
+	case EColorGradingComponent::Green:
+		return LOCTEXT("ColorWheel_GreenComponentLabel", "G");
+
+	case EColorGradingComponent::Blue:
+		return LOCTEXT("ColorWheel_BlueComponentLabel", "B");
+
+	case EColorGradingComponent::Hue:
+		return LOCTEXT("ColorWheel_HueComponentLabel", "H");
+
+	case EColorGradingComponent::Saturation:
+		return LOCTEXT("ColorWheel_SaturationComponentLabel", "S");
+
+	case EColorGradingComponent::Value:
+		return LOCTEXT("ColorWheel_ValueComponentLabel", "V");
+
+	case EColorGradingComponent::Luminance:
+		return LOCTEXT("ColorWheel_LuminanceComponentLabel", "Y");
+	}
+
+	return FText::GetEmpty();
+}
+
+FText SColorGradingComponentViewer::GetComponentToolTipText() const
+{
+	switch (Component.Get())
+	{
+	case EColorGradingComponent::Red:
+		return LOCTEXT("ColorWheel_RedComponentToolTip", "Red");
+
+	case EColorGradingComponent::Green:
+		return LOCTEXT("ColorWheel_GreenComponentToolTip", "Green");
+
+	case EColorGradingComponent::Blue:
+		return LOCTEXT("ColorWheel_BlueComponentToolTip", "Blue");
+
+	case EColorGradingComponent::Hue:
+		return LOCTEXT("ColorWheel_HueComponentToolTip", "Hue");
+
+	case EColorGradingComponent::Saturation:
+		return LOCTEXT("ColorWheel_SaturationComponentToolTip", "Saturation");
+
+	case EColorGradingComponent::Value:
+		return LOCTEXT("ColorWheel_ValueComponentToolTip", "Value");
+
+	case EColorGradingComponent::Luminance:
+		return LOCTEXT("ColorWheel_LuminanceComponentToolTip", "Luminance");
+	}
+
+	return FText::GetEmpty();
+}
+
+} //namespace
+
+#undef LOCTEXT_NAMESPACE
\ No newline at end of file
diff --git a/Engine/Source/Runtime/AdvancedWidgets/Private/Widgets/ColorGrading/SColorGradingPicker.cpp b/Engine/Source/Runtime/AdvancedWidgets/Private/Widgets/ColorGrading/SColorGradingPicker.cpp
new file mode 100644
--- /dev/null
+++ b/Engine/Source/Runtime/AdvancedWidgets/Private/Widgets/ColorGrading/SColorGradingPicker.cpp
@@ -0,0 +1,340 @@
+// Copyright Epic Games, Inc. All Rights Reserved.
+#include "Widgets/ColorGrading/SColorGradingPicker.h"
+#include "Widgets/ColorGrading/SColorGradingWheel.h"
+#include "Widgets/Input/SNumericEntryBox.h"
+#include "Widgets/Layout/SBox.h"
+
+
+#define LOCTEXT_NAMESPACE "ColorGradingWheel"
+
+namespace UE::ColorGrading
+{
+
+SColorGradingPicker::~SColorGradingPicker()
+{
+}
+
+void SColorGradingPicker::Construct(const FArguments& InArgs)
+{
+	if (!InArgs._SliderValueMin.IsSet())
+	{
+		SliderValueMin = TNumericLimits<float>::Lowest();
+	}
+	else
+	{
+		SliderValueMin = InArgs._SliderValueMin.GetValue();
+	}
+
+	if (!InArgs._SliderValueMax.IsSet())
+	{
+		SliderValueMax = TNumericLimits<float>::Max();
+	}
+	else
+	{
+		SliderValueMax = InArgs._SliderValueMax.GetValue();
+	}
+
+	check(SliderValueMin < SliderValueMax);
+
+	MainDelta = InArgs._MainDelta;
+	MainShiftMultiplier = InArgs._MainShiftMultiplier;
+	MainCtrlMultiplier = InArgs._MainCtrlMultiplier;
+	ColorGradingModes = InArgs._ColorGradingModes;
+	OnColorCommitted = InArgs._OnColorCommitted;
+	OnQueryCurrentColor = InArgs._OnQueryCurrentColor;
+
+	float ColorGradingWheelExponent = 2.4f;
+	if (ColorGradingModes == EColorGradingModes::Offset)
+	{
+		ColorGradingWheelExponent = 3.0f;
+	}
+
+	ExternalBeginSliderMovementDelegate = InArgs._OnBeginSliderMovement;
+	ExternalEndSliderMovementDelegate = InArgs._OnEndSliderMovement;
+	ExternalBeginMouseCaptureDelegate = InArgs._OnBeginMouseCapture;
+	ExternalEndMouseCaptureDelegate = InArgs._OnEndMouseCapture;
+
+	ChildSlot
+		[
+			SNew(SVerticalBox)
+				+ SVerticalBox::Slot()
+				.FillHeight(1.0)
+				.HAlign(HAlign_Fill)
+				[
+					SNew(SBox)
+						.HAlign(HAlign_Fill)
+						.VAlign(VAlign_Fill)
+						.MinAspectRatio(1.0f)
+						.MaxAspectRatio(1.0f)
+						[
+							SNew(SColorGradingWheel)
+								.SelectedColor(TAttribute<FLinearColor>::Create(TAttribute<FLinearColor>::FGetter::CreateSP(this, &SColorGradingPicker::GetCurrentLinearColor)))
+								.ExponentDisplacement(ColorGradingWheelExponent)
+								.OnValueChanged(this, &SColorGradingPicker::HandleCurrentColorValueChanged, false)
+								.OnMouseCaptureBegin(this, &SColorGradingPicker::HandleColorWheelMouseCaptureBegin)
+								.OnMouseCaptureEnd(this, &SColorGradingPicker::HandleColorWheelMouseCaptureEnd)
+						]
+				]
+				+ SVerticalBox::Slot()
+				.AutoHeight()
+				.HAlign(HAlign_Center)
+				[
+					SNew(SBox)
+						.HAlign(HAlign_Fill)
+						.MaxDesiredWidth(400.f)
+						.MinDesiredWidth(400.f)
+						[
+							SAssignNew(NumericEntryBoxWidget, SNumericEntryBox<float>)
+								.EditableTextBoxStyle(&FCoreStyle::Get().GetWidgetStyle<FEditableTextBoxStyle>("DarkEditableTextBox"))
+								.Value(this, &SColorGradingPicker::OnGetMainValue)
+								.OnValueCommitted(this, &SColorGradingPicker::OnMainValueCommitted)
+								.OnValueChanged(this, &SColorGradingPicker::OnMainValueChanged, false)
+								.AllowSpin(InArgs._AllowSpin.Get())
+								.SupportDynamicSliderMaxValue(InArgs._SupportDynamicSliderMaxValue)
+								.SupportDynamicSliderMinValue(InArgs._SupportDynamicSliderMinValue)
+								.OnDynamicSliderMaxValueChanged(this, &SColorGradingPicker::OnDynamicSliderMaxValueChanged)
+								.OnDynamicSliderMinValueChanged(this, &SColorGradingPicker::OnDynamicSliderMinValueChanged)
+								.MinValue(InArgs._ValueMin)
+								.MaxValue(InArgs._ValueMax)
+								.MinSliderValue(SliderValueMin)
+								.MaxSliderValue(SliderValueMax)
+								.Delta(MainDelta)
+								.ShiftMultiplier(MainShiftMultiplier)
+								.CtrlMultiplier(MainCtrlMultiplier)
+								.OnBeginSliderMovement(this, &SColorGradingPicker::OnBeginSliderMovement)
+								.OnEndSliderMovement(this, &SColorGradingPicker::OnEndSliderMovement)
+								.UndeterminedString(NSLOCTEXT("PropertyEditor", "MultipleValues", "Multiple Values"))
+								.IsEnabled(this, &SColorGradingPicker::IsEntryBoxEnabled)
+						]
+				]
+		];
+}
+
+bool SColorGradingPicker::IsEntryBoxEnabled() const
+{
+	return OnGetMainValue() != TOptional<float>();
+}
+
+void SColorGradingPicker::HandleColorWheelMouseCaptureBegin(const FLinearColor& InValue)
+{
+	bIsMouseDragging = true;
+	ExternalBeginMouseCaptureDelegate.ExecuteIfBound();
+}
+
+void SColorGradingPicker::HandleColorWheelMouseCaptureEnd(const FLinearColor& InValue)
+{
+	bIsMouseDragging = false;
+	HandleCurrentColorValueChanged(InValue, false);
+	ExternalEndMouseCaptureDelegate.ExecuteIfBound();
+}
+
+void SColorGradingPicker::OnBeginSliderMovement()
+{
+	ExternalBeginSliderMovementDelegate.ExecuteIfBound();
+	bIsMouseDragging = true;
+	//Keep the current value so we can always keep the ratio during the whole drag.
+
+	if (OnQueryCurrentColor.IsBound())
+	{
+		if (OnQueryCurrentColor.Execute(StartDragRatio))
+		{
+			TransformColorGradingRangeToLinearColorRange(StartDragRatio);
+			float MaxCurrentValue = FMath::Max3(StartDragRatio.X, StartDragRatio.Y, StartDragRatio.Z);
+			FVector4 RatioValue(1.0f, 1.0f, 1.0f, 1.0f);
+			if (MaxCurrentValue > SMALL_NUMBER)
+			{
+				RatioValue.X = StartDragRatio.X / MaxCurrentValue;
+				RatioValue.Y = StartDragRatio.Y / MaxCurrentValue;
+				RatioValue.Z = StartDragRatio.Z / MaxCurrentValue;
+			}
+			StartDragRatio = RatioValue;
+		}
+	}
+}
+
+void SColorGradingPicker::OnEndSliderMovement(float NewValue)
+{
+	bIsMouseDragging = false;
+	StartDragRatio.X = 1.0f;
+	StartDragRatio.Y = 1.0f;
+	StartDragRatio.Z = 1.0f;
+
+	OnMainValueChanged(NewValue, false);
+
+	ExternalEndSliderMovementDelegate.ExecuteIfBound();
+}
+
+void SColorGradingPicker::AdjustRatioValue(FVector4& NewValue)
+{
+	if (!bIsMouseDragging)
+	{
+		return;
+	}
+	float MaxCurrentValue = FMath::Max3(NewValue.X, NewValue.Y, NewValue.Z);
+	if (MaxCurrentValue > SMALL_NUMBER)
+	{
+		NewValue.X = StartDragRatio.X * MaxCurrentValue;
+		NewValue.Y = StartDragRatio.Y * MaxCurrentValue;
+		NewValue.Z = StartDragRatio.Z * MaxCurrentValue;
+	}
+}
+
+void SColorGradingPicker::OnMainValueChanged(float InValue, bool ShouldCommitValueChanges)
+{
+	if (bIsMouseDragging || ShouldCommitValueChanges)
+	{
+		TransformColorGradingRangeToLinearColorRange(InValue);
+
+		FVector4 CurrentValue(0.0f, 0.0f, 0.0f, 0.0f);
+
+		if (OnQueryCurrentColor.IsBound())
+		{
+			if (OnQueryCurrentColor.Execute(CurrentValue))
+			{
+				TransformColorGradingRangeToLinearColorRange(CurrentValue);
+
+				//The MainValue is the maximum of any channel value
+				float MaxCurrentValue = FMath::Max3(CurrentValue.X, CurrentValue.Y, CurrentValue.Z);
+				if (MaxCurrentValue <= SMALL_NUMBER)
+				{
+					//We need the neutral value for the type of color grading, currently only offset is an addition(0.0) all other are multiplier(1.0)
+					CurrentValue = FVector4(InValue, InValue, InValue, CurrentValue.W);
+				}
+				else
+				{
+					float Ratio = InValue / MaxCurrentValue;
+					CurrentValue *= FVector4(Ratio, Ratio, Ratio, 1.0f);
+					AdjustRatioValue(CurrentValue);
+				}
+				TransformLinearColorRangeToColorGradingRange(CurrentValue);
+				OnColorCommitted.ExecuteIfBound(CurrentValue, ShouldCommitValueChanges);
+			}
+		}
+	}
+}
+
+void SColorGradingPicker::OnMainValueCommitted(float InValue, ETextCommit::Type CommitType)
+{
+	OnMainValueChanged(InValue, true);
+}
+
+TOptional<float> SColorGradingPicker::OnGetMainValue() const
+{
+	FVector4 CurrentValue(0.0f, 0.0f, 0.0f, 0.0f);
+
+	if (OnQueryCurrentColor.IsBound())
+	{
+		if (OnQueryCurrentColor.Execute(CurrentValue))
+		{
+			//The MainValue is the maximum of any channel value
+			TOptional<float> CurrentValueOption = FMath::Max3(CurrentValue.X, CurrentValue.Y, CurrentValue.Z);
+			return CurrentValueOption;
+		}
+	}
+
+	return TOptional<float>();
+}
+
+void SColorGradingPicker::TransformLinearColorRangeToColorGradingRange(FVector4& VectorValue) const
+{
+	float Ratio = (SliderValueMax - SliderValueMin);
+	VectorValue *= Ratio;
+	VectorValue += FVector4(SliderValueMin, SliderValueMin, SliderValueMin, SliderValueMin);
+}
+
+void SColorGradingPicker::TransformColorGradingRangeToLinearColorRange(FVector4& VectorValue) const
+{
+	VectorValue += -FVector4(SliderValueMin, SliderValueMin, SliderValueMin, SliderValueMin);
+	float Ratio = 1.0f / (SliderValueMax - SliderValueMin);
+	VectorValue *= Ratio;
+}
+
+void SColorGradingPicker::TransformColorGradingRangeToLinearColorRange(float& FloatValue)
+{
+	FloatValue -= SliderValueMin;
+	float Ratio = 1.0f / (SliderValueMax - SliderValueMin);
+	FloatValue *= Ratio;
+}
+
+FLinearColor SColorGradingPicker::GetCurrentLinearColor()
+{
+	FLinearColor CurrentColor = FLinearColor::Black;
+	FVector4 CurrentValue;
+
+	if (OnQueryCurrentColor.IsBound())
+	{
+		if (OnQueryCurrentColor.Execute(CurrentValue))
+		{
+			TransformColorGradingRangeToLinearColorRange(CurrentValue);
+			CurrentColor = FLinearColor(CurrentValue.X, CurrentValue.Y, CurrentValue.Z);
+		}
+	}
+
+	return CurrentColor.LinearRGBToHSV();
+}
+
+void SColorGradingPicker::HandleCurrentColorValueChanged(const FLinearColor& NewValue, bool ShouldCommitValueChanges)
+{
+	//Query the current vector4 so we can pass back the w value
+	FVector4 CurrentValue(0.0f, 0.0f, 0.0f, 0.0f);
+	if (OnQueryCurrentColor.IsBound())
+	{
+		if (OnQueryCurrentColor.Execute(CurrentValue))
+		{
+			FLinearColor NewValueRGB = NewValue.HSVToLinearRGB();
+			FVector4 NewValueVector(NewValueRGB.R, NewValueRGB.G, NewValueRGB.B);
+			TransformLinearColorRangeToColorGradingRange(NewValueVector);
+			//Set the W with the original value
+			NewValueVector.W = CurrentValue.W;
+			OnColorCommitted.ExecuteIfBound(NewValueVector, ShouldCommitValueChanges);
+		}
+	}
+}
+
+void SColorGradingPicker::OnDynamicSliderMaxValueChanged(float NewMaxSliderValue, TWeakPtr<SWidget> InValueChangedSourceWidget, bool IsOriginator, bool UpdateOnlyIfHigher)
+{
+	TSharedPtr<SSpinBox<float>> SpinBox = StaticCastSharedPtr<SSpinBox<float>>(NumericEntryBoxWidget->GetSpinBox());
+
+	if (SpinBox.IsValid())
+	{
+		if (SpinBox != InValueChangedSourceWidget)
+		{
+			if ((NewMaxSliderValue > SpinBox->GetMaxSliderValue() && UpdateOnlyIfHigher) || !UpdateOnlyIfHigher)
+			{
+				SpinBox->SetMaxSliderValue(NewMaxSliderValue);
+				SliderValueMax = NewMaxSliderValue;
+			}
+		}
+	}
+
+	if (IsOriginator)
+	{
+		OnNumericEntryBoxDynamicSliderMaxValueChanged.Broadcast(NewMaxSliderValue, InValueChangedSourceWidget, false, UpdateOnlyIfHigher);
+	}
+}
+
+void SColorGradingPicker::OnDynamicSliderMinValueChanged(float NewMinSliderValue, TWeakPtr<SWidget> InValueChangedSourceWidget, bool IsOriginator, bool UpdateOnlyIfLower)
+{
+	TSharedPtr<SSpinBox<float>> SpinBox = StaticCastSharedPtr<SSpinBox<float>>(NumericEntryBoxWidget->GetSpinBox());
+
+	if (SpinBox.IsValid())
+	{
+		if (SpinBox != InValueChangedSourceWidget)
+		{
+			if ((NewMinSliderValue < SpinBox->GetMinSliderValue() && UpdateOnlyIfLower) || !UpdateOnlyIfLower)
+			{
+				SpinBox->SetMinSliderValue(NewMinSliderValue);
+				SliderValueMin = NewMinSliderValue;
+			}
+		}
+	}
+
+	if (IsOriginator)
+	{
+		OnNumericEntryBoxDynamicSliderMinValueChanged.Broadcast(NewMinSliderValue, InValueChangedSourceWidget, false, UpdateOnlyIfLower);
+	}
+}
+
+} //namespace
+
+#undef LOCTEXT_NAMESPACE
diff --git a/Engine/Source/Runtime/AdvancedWidgets/Private/Widgets/ColorGrading/SColorGradingWheel.cpp b/Engine/Source/Runtime/AdvancedWidgets/Private/Widgets/ColorGrading/SColorGradingWheel.cpp
new file mode 100644
--- /dev/null
+++ b/Engine/Source/Runtime/AdvancedWidgets/Private/Widgets/ColorGrading/SColorGradingWheel.cpp
@@ -0,0 +1,291 @@
+// Copyright Epic Games, Inc. All Rights Reserved.
+#include "Widgets/ColorGrading/SColorGradingWheel.h"
+
+#include "Materials/Material.h"
+#include "SlateMaterialBrush.h"
+#include "Styling/CoreStyle.h"
+#include "Styling/AdvancedWidgetsStyle.h"
+#include "Rendering/RenderingCommon.h"
+#include "Rendering/DrawElements.h"
+
+namespace UE::ColorGrading
+{
+
+SLATE_IMPLEMENT_WIDGET(SColorGradingWheel)
+void SColorGradingWheel::PrivateRegisterAttributes(FSlateAttributeInitializer& AttributeInitializer)
+{
+	SLATE_ADD_MEMBER_ATTRIBUTE_DEFINITION_WITH_NAME(AttributeInitializer, "SelectedColor", SelectedColorAttribute, EInvalidateWidgetReason::Paint);
+	SLATE_ADD_MEMBER_ATTRIBUTE_DEFINITION_WITH_NAME(AttributeInitializer, "DesiredWheelSize", DesiredWheelSizeAttribute, EInvalidateWidgetReason::Layout);
+	SLATE_ADD_MEMBER_ATTRIBUTE_DEFINITION_WITH_NAME(AttributeInitializer, "ExponentDisplacement", ExponentDisplacementAttribute, EInvalidateWidgetReason::Paint);
+}
+
+SColorGradingWheel::SColorGradingWheel()
+	: SelectedColorAttribute(*this, FLinearColor(ForceInit))
+	, DesiredWheelSizeAttribute(*this)
+	, ExponentDisplacementAttribute(*this)
+	, Union_IsAttributeSet(0)
+{
+}
+
+/* SColorGradingWheel methods
+ *****************************************************************************/
+
+void SColorGradingWheel::Construct(const FArguments& InArgs)
+{
+	// Create the brush here since we need to dynamically load the material for it
+	BackgroundMaterial = LoadObject<UMaterial>(nullptr, TEXT("Material'/Engine/EngineMaterials/ColorGradingWheel.ColorGradingWheel'"));
+	check(BackgroundMaterial);
+	BackgroundImage = MakeShared<FSlateMaterialBrush>(*BackgroundMaterial.Get(), FVector2f(400.f, 400.f));
+
+	CrossImage = UE::AdvancedWidgets::FAdvancedWidgetsStyle::Get().GetBrush("ColorGradingWheel.Cross");
+	SelectorImage = UE::AdvancedWidgets::FAdvancedWidgetsStyle::Get().GetBrush("ColorGradingWheel.Selector");
+	SetSelectedColorAttribute(InArgs._SelectedColor);
+	SetDesiredWheelSizeAttribute(InArgs._DesiredWheelSize);
+	SetExponentDisplacementAttribute(InArgs._ExponentDisplacement);
+	OnMouseCaptureBegin = InArgs._OnMouseCaptureBegin;
+	OnMouseCaptureEnd = InArgs._OnMouseCaptureEnd;
+	OnValueChanged = InArgs._OnValueChanged;
+}
+
+void SColorGradingWheel::SetSelectedColorAttribute(TAttribute<FLinearColor> InSelectedColor)
+{
+	SelectedColorAttribute.Assign(*this, MoveTemp(InSelectedColor));
+}
+
+void SColorGradingWheel::SetDesiredWheelSizeAttribute(TAttribute<int32> InDesiredWheelSize)
+{
+	const bool bDesiredWheelSizeSetChanged = (bIsAttributeDesiredWheelSizeSet != InDesiredWheelSize.IsSet());
+	bIsAttributeDesiredWheelSizeSet = InDesiredWheelSize.IsSet();
+	const bool bAttributeWasAssgined = DesiredWheelSizeAttribute.Assign(*this, MoveTemp(InDesiredWheelSize));
+
+	// If the assign didn't invalidate the widget but the attribute set changed, then invalidate the widget.
+	if (bDesiredWheelSizeSetChanged && !bAttributeWasAssgined)
+	{
+		Invalidate(EInvalidateWidgetReason::Layout);
+	}
+}
+
+void SColorGradingWheel::SetExponentDisplacementAttribute(TAttribute<float> InExponentDisplacement)
+{
+	const bool bExponentDisplacementSetChanged = (bIsAttributeExponentDisplacementSet != InExponentDisplacement.IsSet());
+	bIsAttributeExponentDisplacementSet = InExponentDisplacement.IsSet();
+	const bool bAttributeWasAssgined = ExponentDisplacementAttribute.Assign(*this, MoveTemp(InExponentDisplacement), 1.f);
+
+	// If the assign didn't invalidate the widget but the attribute set changed, then invalidate the widget.
+	if (bExponentDisplacementSetChanged && !bAttributeWasAssgined)
+	{
+		Invalidate(EInvalidateWidgetReason::Paint);
+	}
+}
+
+FString SColorGradingWheel::GetReferencerName() const
+{
+	return TEXT("SColorGradingWheel");
+}
+
+void SColorGradingWheel::AddReferencedObjects(FReferenceCollector& InCollector)
+{
+	InCollector.AddReferencedObject(BackgroundMaterial);
+}
+
+
+/* SWidget overrides
+ *****************************************************************************/
+
+FVector2D SColorGradingWheel::ComputeDesiredSize(float) const
+{
+	if (bIsAttributeDesiredWheelSizeSet)
+	{
+		int32 CachedDesiredWheelSize = DesiredWheelSizeAttribute.Get();
+		return FVector2D(CachedDesiredWheelSize, CachedDesiredWheelSize);
+	}
+	return FVector2D(BackgroundImage->ImageSize);
+}
+
+
+FReply SColorGradingWheel::OnMouseButtonDoubleClick(const FGeometry& InMyGeometry, const FPointerEvent& InMouseEvent)
+{
+	return FReply::Handled();
+}
+
+
+FReply SColorGradingWheel::OnMouseButtonDown(const FGeometry& MyGeometry, const FPointerEvent& MouseEvent)
+{
+	if (MouseEvent.GetEffectingButton() == EKeys::LeftMouseButton)
+	{
+		if (OnMouseCaptureBegin.IsBound())
+		{
+			SelectedColorAttribute.UpdateNow(*this);
+			OnMouseCaptureBegin.Execute(SelectedColorAttribute.Get());
+		}
+
+		if (!ProcessMouseAction(MyGeometry, MouseEvent, false))
+		{
+			if (OnMouseCaptureEnd.IsBound())
+			{
+				SelectedColorAttribute.UpdateNow(*this);
+				OnMouseCaptureEnd.Execute(SelectedColorAttribute.Get());
+			}
+			return FReply::Unhandled();
+		}
+
+		return FReply::Handled().CaptureMouse(SharedThis(this));
+	}
+
+	return FReply::Unhandled();
+}
+
+
+FReply SColorGradingWheel::OnMouseButtonUp(const FGeometry& MyGeometry, const FPointerEvent& MouseEvent)
+{
+	if (MouseEvent.GetEffectingButton() == EKeys::LeftMouseButton && HasMouseCapture())
+	{
+		if (OnMouseCaptureEnd.IsBound())
+		{
+			SelectedColorAttribute.UpdateNow(*this);
+			OnMouseCaptureEnd.Execute(SelectedColorAttribute.Get());
+		}
+
+		return FReply::Handled().ReleaseMouseCapture();
+	}
+
+	return FReply::Unhandled();
+}
+
+
+FReply SColorGradingWheel::OnMouseMove(const FGeometry& MyGeometry, const FPointerEvent& MouseEvent)
+{
+	if (!HasMouseCapture())
+	{
+		return FReply::Unhandled();
+	}
+
+	ProcessMouseAction(MyGeometry, MouseEvent, true);
+
+	return FReply::Handled();
+}
+
+
+int32 SColorGradingWheel::OnPaint(const FPaintArgs& Args, const FGeometry& AllottedGeometry, const FSlateRect& MyCullingRect, FSlateWindowElementList& OutDrawElements, int32 LayerId, const FWidgetStyle& InWidgetStyle, bool bParentEnabled) const
+{
+	const bool bIsEnabled = ShouldBeEnabled(bParentEnabled);
+	const ESlateDrawEffect DrawEffects = bIsEnabled ? ESlateDrawEffect::None : ESlateDrawEffect::DisabledEffect;
+
+	FVector2f AllottedGeometrySize;
+	if (bIsAttributeDesiredWheelSizeSet)
+	{
+		int32 CachedDesiredWheelSize = DesiredWheelSizeAttribute.Get();
+		AllottedGeometrySize = FVector2f(CachedDesiredWheelSize, CachedDesiredWheelSize);
+	}
+	else
+	{
+		AllottedGeometrySize = AllottedGeometry.GetLocalSize();
+	}
+
+	const float SelectorShrinkThreshold = 150.f;
+	FVector2f SelectorSize = SelectorImage->ImageSize;
+	if (AllottedGeometrySize.X < SelectorShrinkThreshold || AllottedGeometrySize.Y < SelectorShrinkThreshold)
+	{
+		SelectorSize *= 0.75f;
+	}
+
+	const FVector2f CircleSize = AllottedGeometrySize - SelectorSize;
+
+	FSlateDrawElement::MakeBox(
+		OutDrawElements,
+		LayerId,
+		AllottedGeometry.ToPaintGeometry(CircleSize, FSlateLayoutTransform(0.5f * SelectorSize)),
+		BackgroundImage.Get(),
+		DrawEffects,
+		InWidgetStyle.GetColorAndOpacityTint() * BackgroundImage->GetTint(InWidgetStyle)
+	);
+
+	FSlateDrawElement::MakeBox(
+		OutDrawElements,
+		LayerId + 1,
+		AllottedGeometry.ToPaintGeometry(CircleSize, FSlateLayoutTransform(0.5f * SelectorSize)),
+		CrossImage,
+		DrawEffects,
+		InWidgetStyle.GetColorAndOpacityTint() * CrossImage->GetTint(InWidgetStyle)
+	);
+
+	FSlateDrawElement::MakeBox(
+		OutDrawElements,
+		LayerId + 2,
+		AllottedGeometry.ToPaintGeometry(SelectorSize, FSlateLayoutTransform(0.5f * (AllottedGeometrySize + CalcRelativePositionFromCenter() * CircleSize - SelectorSize))),
+		SelectorImage,
+		DrawEffects,
+		InWidgetStyle.GetColorAndOpacityTint() * SelectorImage->GetTint(InWidgetStyle)
+	);
+
+	return LayerId + 1;
+}
+
+
+/* SColorGradingWheel implementation
+ *****************************************************************************/
+
+UE::Slate::FDeprecateVector2DResult SColorGradingWheel::CalcRelativePositionFromCenter() const
+{
+	float Hue = SelectedColorAttribute.Get().R;
+	float Saturation = SelectedColorAttribute.Get().G;
+	if (bIsAttributeExponentDisplacementSet && ExponentDisplacementAttribute.Get() != 1.0f && !FMath::IsNearlyEqual(ExponentDisplacementAttribute.Get(), 0.0f, 0.00001f))
+	{
+		//Use log curve to set the distance G value
+		Saturation = FMath::Pow(Saturation, 1.0f / ExponentDisplacementAttribute.Get());
+	}
+	float Angle = HueAngleOffset - (Hue / 180.0f * PI);
+	float Radius = Saturation;
+
+	return FVector2f(FMath::Cos(Angle), FMath::Sin(Angle)) * Radius;
+}
+
+
+bool SColorGradingWheel::ProcessMouseAction(const FGeometry& MyGeometry, const FPointerEvent& MouseEvent, bool bProcessWhenOutsideColorWheel)
+{
+	FVector2f GeometrySize = MyGeometry.GetLocalSize();
+	if (bIsAttributeDesiredWheelSizeSet)
+	{
+		DesiredWheelSizeAttribute.UpdateNow(*this);
+		int32 CachedDesiredWheelSize = DesiredWheelSizeAttribute.Get();
+		GeometrySize.X = CachedDesiredWheelSize;
+		GeometrySize.Y = CachedDesiredWheelSize;
+	}
+
+	const FVector2f LocalMouseCoordinate = MyGeometry.AbsoluteToLocal(MouseEvent.GetScreenSpacePosition());
+	const FVector2f RelativePositionFromCenter = (2.0f * LocalMouseCoordinate - GeometrySize) / (GeometrySize - SelectorImage->ImageSize);
+	const float RelativeRadius = RelativePositionFromCenter.Size();
+
+	if (RelativeRadius <= 1.0f || bProcessWhenOutsideColorWheel)
+	{
+		float Angle = FMath::Atan2(-RelativePositionFromCenter.Y, RelativePositionFromCenter.X) + HueAngleOffset;
+
+		if (Angle < 0.0f)
+		{
+			Angle += 2.0f * PI;
+		}
+
+		SelectedColorAttribute.UpdateNow(*this);
+		FLinearColor NewColor = SelectedColorAttribute.Get();
+		{
+			NewColor.R = Angle * 180.0f * INV_PI;
+			float LinearRadius = FMath::Min(RelativeRadius, 1.0f);
+			if (bIsAttributeExponentDisplacementSet && ExponentDisplacementAttribute.Get() != 1.0f)
+			{
+				//Use log curve to set the distance G value
+				float LogDistance = FMath::Pow(LinearRadius, ExponentDisplacementAttribute.Get());
+				NewColor.G = LogDistance;
+			}
+			else
+			{
+				NewColor.G = LinearRadius;
+			}
+		}
+
+		OnValueChanged.ExecuteIfBound(NewColor);
+	}
+
+	return (RelativeRadius <= 1.0f);
+}
+
+} //namespace
diff --git a/Engine/Source/Runtime/AdvancedWidgets/Public/Framework/ColorGrading/ColorGradingCommon.h b/Engine/Source/Runtime/AdvancedWidgets/Public/Framework/ColorGrading/ColorGradingCommon.h
new file mode 100644
--- /dev/null
+++ b/Engine/Source/Runtime/AdvancedWidgets/Public/Framework/ColorGrading/ColorGradingCommon.h
@@ -0,0 +1,54 @@
+// Copyright Epic Games, Inc. All Rights Reserved.
+
+#pragma once
+
+#include "Delegates/Delegate.h"
+#include "HAL/Platform.h"
+
+namespace UE::ColorGrading
+{
+
+/** Callback to get the current FVector4 value */
+DECLARE_DELEGATE_RetVal_OneParam(bool, FOnGetCurrentVector4Value, FVector4&);
+
+/** Color components that can be controlled by a color grading widget. */
+enum class ADVANCEDWIDGETS_API EColorGradingComponent
+{
+	Red,
+	Green,
+	Blue,
+	Hue,
+	Saturation,
+	Value,
+	Luminance
+};
+
+/** Types of color grading modes */
+enum class ADVANCEDWIDGETS_API EColorGradingWheelType
+{
+	/** A wheel that controls a standard color value */
+	Standard,
+
+	/** A wheel that applies an offset to colors */
+	Offset
+};
+
+/** Modes that colors can be displayed in for the color grading panel */
+enum class ADVANCEDWIDGETS_API EColorGradingColorDisplayMode
+{
+	RGB,
+	HSV
+};
+
+/** Enumerates color picker modes */
+enum class ADVANCEDWIDGETS_API EColorGradingModes
+{
+	Saturation,
+	Contrast,
+	Gamma,
+	Gain,
+	Offset,
+	Invalid
+};
+
+} //namespace
diff --git a/Engine/Source/Runtime/AdvancedWidgets/Public/Styling/ColorGradingSpinBoxStyle.h b/Engine/Source/Runtime/AdvancedWidgets/Public/Styling/ColorGradingSpinBoxStyle.h
new file mode 100644
--- /dev/null
+++ b/Engine/Source/Runtime/AdvancedWidgets/Public/Styling/ColorGradingSpinBoxStyle.h
@@ -0,0 +1,66 @@
+// Copyright Epic Games, Inc. All Rights Reserved.
+
+#pragma once
+
+#include "CoreMinimal.h"
+#include "UObject/ObjectMacros.h"
+#include "Styling/SlateTypes.h"
+
+#include "ColorGradingSpinBoxStyle.generated.h"
+
+/**
+ * Represents the appearance of a color grading spin box
+ */
+USTRUCT(BlueprintType)
+struct FColorGradingSpinBoxStyle : public FSlateWidgetStyle
+{
+	GENERATED_USTRUCT_BODY()
+
+	ADVANCEDWIDGETS_API FColorGradingSpinBoxStyle();
+
+	virtual ~FColorGradingSpinBoxStyle() {}
+
+	ADVANCEDWIDGETS_API virtual void GetResources(TArray< const FSlateBrush* >& OutBrushes) const override;
+
+	static ADVANCEDWIDGETS_API const FName TypeName;
+	virtual const FName GetTypeName() const override { return TypeName; };
+
+	static ADVANCEDWIDGETS_API const FColorGradingSpinBoxStyle& GetDefault();
+
+	/** Brush used to draw the border of the spinbox */
+	UPROPERTY(EditAnywhere, BlueprintReadWrite, Category = Appearance)
+	FSlateBrush BorderBrush;
+	FColorGradingSpinBoxStyle& SetBorderBrush(const FSlateBrush& InBorderBrush) { BorderBrush = InBorderBrush; return *this; }
+
+	/** Brush used to draw the border of the spinbox when it's in active use by the user */
+	UPROPERTY(EditAnywhere, BlueprintReadWrite, Category = Appearance)
+	FSlateBrush ActiveBorderBrush;
+	FColorGradingSpinBoxStyle& SetActiveBorderBrush(const FSlateBrush& InActiveBorderBrush) { ActiveBorderBrush = InActiveBorderBrush; return *this; }
+
+	/** Brush used to draw the border of the spinbox when it's hovered over */
+	UPROPERTY(EditAnywhere, BlueprintReadWrite, Category = Appearance)
+	FSlateBrush HoveredBorderBrush;
+	FColorGradingSpinBoxStyle& SetHoveredBorderBrush(const FSlateBrush& InHoveredBorderBrush) { HoveredBorderBrush = InHoveredBorderBrush; return *this; }
+
+	/** Brush used to draw the selector indicating the current value */
+	UPROPERTY(EditAnywhere, BlueprintReadWrite, Category = Appearance)
+	FSlateBrush SelectorBrush;
+	FColorGradingSpinBoxStyle& SetSelectorBrush(const FSlateBrush& InSelectorBrush) { SelectorBrush = InSelectorBrush; return *this; }
+
+	/** Width of the selector */
+	UPROPERTY(EditAnywhere, BlueprintReadWrite, Category = Appearance)
+	float SelectorWidth;
+	FColorGradingSpinBoxStyle& SetSelectorWidth(float InSelectorWidth) { SelectorWidth = InSelectorWidth; return *this; }
+
+	/**
+	 * Unlinks all colors in this style.
+	 * @see FSlateColor::Unlink
+	 */
+	void UnlinkColors()
+	{
+		BorderBrush.UnlinkColors();
+		HoveredBorderBrush.UnlinkColors();
+		ActiveBorderBrush.UnlinkColors();
+		SelectorBrush.UnlinkColors();
+	}
+};
diff --git a/Engine/Source/Runtime/AdvancedWidgets/Public/Util/ColorGradingUtil.h b/Engine/Source/Runtime/AdvancedWidgets/Public/Util/ColorGradingUtil.h
new file mode 100644
--- /dev/null
+++ b/Engine/Source/Runtime/AdvancedWidgets/Public/Util/ColorGradingUtil.h
@@ -0,0 +1,13 @@
+// Copyright Epic Games, Inc. All Rights Reserved.
+
+#pragma once
+
+#include "Framework/ColorGrading/ColorGradingCommon.h"
+
+namespace UE::ColorGrading
+{
+
+/** Get the color grading component typically associated with the given display mode and component index. */
+EColorGradingComponent ADVANCEDWIDGETS_API GetColorGradingComponent(EColorGradingColorDisplayMode DisplayMode, int Index);
+
+}
diff --git a/Engine/Source/Runtime/AdvancedWidgets/Public/Widgets/ColorGrading/SColorGradingComponentViewer.h b/Engine/Source/Runtime/AdvancedWidgets/Public/Widgets/ColorGrading/SColorGradingComponentViewer.h
new file mode 100644
--- /dev/null
+++ b/Engine/Source/Runtime/AdvancedWidgets/Public/Widgets/ColorGrading/SColorGradingComponentViewer.h
@@ -0,0 +1,152 @@
+// Copyright Epic Games, Inc. All Rights Reserved.
+
+#pragma once
+
+#include "CoreMinimal.h"
+#include "Widgets/DeclarativeSyntaxSupport.h"
+#include "Widgets/Input/NumericTypeInterface.h"
+#include "Widgets/Input/SNumericEntryBox.h"
+#include "Widgets/SCompoundWidget.h"
+#include "Framework/ColorGrading/ColorGradingCommon.h"
+
+namespace UE::ColorGrading
+{
+
+class SColorGradingComponentSpinBox;
+
+/**
+ * A slider displayed in tandem with a color grading wheel.
+ * Allows mouse-based control or direct input via a textbox.
+ * Shows a gradient representing the controlled color component.
+ */
+class ADVANCEDWIDGETS_API SColorGradingComponentViewer
+	: public SCompoundWidget
+{
+public:
+
+	/** Notification for numeric value change */
+	DECLARE_DELEGATE_OneParam(FOnValueChanged, float);
+
+	/** Notification when the max/min spinner values are changed (only apply if SupportDynamicSliderMaxValue or SupportDynamicSliderMinValue are true) */
+	DECLARE_DELEGATE_FourParams(FOnDynamicSliderMinMaxValueChanged, float, TWeakPtr<SWidget>, bool, bool);
+
+	SLATE_BEGIN_ARGS(SColorGradingComponentViewer)
+		: _Value(0)
+		, _MinValue(0)
+		, _MaxValue(2)
+		, _MinFractionalDigits(DefaultMinFractionalDigits)
+		, _MaxFractionalDigits(DefaultMaxFractionalDigits)
+		, _Delta(0)
+		, _ShiftMultiplier(10.f)
+		, _CtrlMultiplier(0.1f)
+		, _SupportDynamicSliderMaxValue(false)
+		, _SupportDynamicSliderMinValue(false)
+		, _SliderExponent(1.f)
+		, _OnValueChanged()
+		, _OnQueryCurrentColor()
+		, _UseCompactDisplay(false)
+		, _AllowSpin(true)
+		{}
+
+		/** The value to display */
+		SLATE_ATTRIBUTE(TOptional<float>, Value)
+		/** The component being displayed */
+		SLATE_ATTRIBUTE(EColorGradingComponent, Component)
+		/** The mode of the associated color grading wheel */
+		SLATE_ATTRIBUTE(EColorGradingModes, ColorGradingMode)
+		/** The minimum value that can be entered into the text edit box */
+		SLATE_ATTRIBUTE(TOptional<float>, MinValue)
+		/** The maximum value that can be entered into the text edit box */
+		SLATE_ATTRIBUTE(TOptional<float>, MaxValue)
+		/** The minimum value that can be specified by using the slider, defaults to MinValue */
+		SLATE_ATTRIBUTE(TOptional<float>, MinSliderValue)
+		/** The maximum value that can be specified by using the slider, defaults to MaxValue */
+		SLATE_ATTRIBUTE(TOptional<float>, MaxSliderValue)
+		/** The minimum fractional digits the spin box displays, defaults to 1 */
+		SLATE_ATTRIBUTE(TOptional<int32>, MinFractionalDigits)
+		/** The maximum fractional digits the spin box displays, defaults to 3 */
+		SLATE_ATTRIBUTE(TOptional<int32>, MaxFractionalDigits)
+		/** Delta to increment the value as the slider moves.  If not specified will determine automatically */
+		SLATE_ATTRIBUTE(float, Delta)
+		/** Multiplier to use when shift is held down */
+		SLATE_ATTRIBUTE(float, ShiftMultiplier)
+		/** Multiplier to use when ctrl is held down */
+		SLATE_ATTRIBUTE(float, CtrlMultiplier)
+		/** If we're an unbounded spinbox, what value do we divide mouse movement by before multiplying by Delta. Requires Delta to be set. */
+		SLATE_ATTRIBUTE(int32, LinearDeltaSensitivity)
+		/** Tell us if we want to support dynamically changing of the max value using alt */
+		SLATE_ATTRIBUTE(bool, SupportDynamicSliderMaxValue)
+		/** Tell us if we want to support dynamically changing of the min value using alt */
+		SLATE_ATTRIBUTE(bool, SupportDynamicSliderMinValue)
+		/** Called right after the max slider value is changed (only relevant if SupportDynamicSliderMaxValue is true) */
+		SLATE_EVENT(FOnDynamicSliderMinMaxValueChanged, OnDynamicSliderMaxValueChanged)
+		/** Called right after the min slider value is changed (only relevant if SupportDynamicSliderMinValue is true) */
+		SLATE_EVENT(FOnDynamicSliderMinMaxValueChanged, OnDynamicSliderMinValueChanged)
+		/** Use exponential scale for the slider */
+		SLATE_ATTRIBUTE(float, SliderExponent)
+		/** When use exponential scale for the slider which is the neutral value */
+		SLATE_ATTRIBUTE(float, SliderExponentNeutralValue)
+		/** Step to increment or decrement the value by when scrolling the mouse wheel. If not specified will determine automatically */
+		SLATE_ATTRIBUTE(TOptional<float>, WheelStep)
+		/** Called when the value is changed by slider or typing */
+		SLATE_EVENT(FOnValueChanged, OnValueChanged)
+		/** Called right before the slider begins to move */
+		SLATE_EVENT(FSimpleDelegate, OnBeginSliderMovement)
+		/** Called right after the slider handle is released by the user */
+		SLATE_EVENT(FOnValueChanged, OnEndSliderMovement)
+		/** Callback to get the current FVector4 color value (used to update the background gradient) */
+		SLATE_EVENT(FOnGetCurrentVector4Value, OnQueryCurrentColor)
+		/** Provide custom type conversion functionality to the spin box */
+		SLATE_ARGUMENT(TSharedPtr<INumericTypeInterface<float>>, TypeInterface)
+		/** If true, reduce padding and hide the gradient spinbox, leaving only the numeric entry. */
+		SLATE_ARGUMENT(bool, UseCompactDisplay)
+		/** Whether or not the user should be able to change the value by dragging with the mouse cursor */
+		SLATE_ARGUMENT(bool, AllowSpin)
+
+	SLATE_END_ARGS()
+
+	SColorGradingComponentViewer() = default;
+
+	virtual ~SColorGradingComponentViewer();
+
+	/**
+	 * Construct the widget
+	 *
+	 * @param InArgs   A declaration from which to construct the widget
+	 */
+	void Construct(const FArguments& InArgs);
+
+	/** Get the maximum slider value */
+	float GetMaxSliderValue() const;
+
+	/** Get the minimum slider value */
+	float GetMinSliderValue() const;
+
+private:
+
+	/** Get the value as a non-optional float */
+	float GetValue() const;
+
+	/** Get the component's name as shown on a short label */
+	FText GetComponentLabelText() const;
+
+	/** Get the component's name as shown in a tooltip */
+	FText GetComponentToolTipText() const;
+
+	/** The default minimum fractional digits */
+	static constexpr int32 DefaultMinFractionalDigits = 1;
+
+	/** The default maximum fractional digits */
+	static constexpr int32 DefaultMaxFractionalDigits = 3;
+
+	TAttribute<TOptional<float>> MinValue;
+	TAttribute<TOptional<float>> MaxValue;
+	TAttribute<TOptional<float>> MinSliderValue;
+	TAttribute<TOptional<float>> MaxSliderValue;
+	TAttribute<TOptional<float>> OptionalValue;
+	TSharedPtr<SNumericEntryBox<float>> NumericEntryBox;
+	TSharedPtr<SColorGradingComponentSpinBox> GradientSpinBox;
+	TAttribute<EColorGradingComponent> Component;
+};
+
+} //namespace
diff --git a/Engine/Source/Runtime/AdvancedWidgets/Public/Widgets/ColorGrading/SColorGradingPicker.h b/Engine/Source/Runtime/AdvancedWidgets/Public/Widgets/ColorGrading/SColorGradingPicker.h
new file mode 100644
--- /dev/null
+++ b/Engine/Source/Runtime/AdvancedWidgets/Public/Widgets/ColorGrading/SColorGradingPicker.h
@@ -0,0 +1,156 @@
+// Copyright Epic Games, Inc. All Rights Reserved.
+
+#pragma once
+
+#include "CoreMinimal.h"
+#include "Misc/Attribute.h"
+#include "Input/Reply.h"
+#include "Widgets/DeclarativeSyntaxSupport.h"
+#include "Widgets/SLeafWidget.h"
+#include "Framework/ColorGrading/ColorGradingCommon.h"
+#include "Framework/SlateDelegates.h"
+#include "Widgets/SCompoundWidget.h"
+#include "Widgets/Input/SNumericEntryBox.h"
+
+namespace UE::ColorGrading
+{
+
+/**
+ * Class for placing a color picker. If all you need is a standalone color picker,
+ * use the functions OpenColorGradingWheel and DestroyColorGradingWheel, since they hold a static
+ * instance of the color picker.
+ */
+class SColorGradingPicker
+	: public SCompoundWidget
+{
+public:
+	/** Notification when the max/min spinner values are changed (only apply if SupportDynamicSliderMaxValue or SupportDynamicSliderMinValue are true) */
+	DECLARE_MULTICAST_DELEGATE_FourParams(FOnNumericEntryBoxDynamicSliderMinMaxValueChanged, float, TWeakPtr<SWidget>, bool, bool);
+
+	// Delegate called when the widget Color Data changed
+	DECLARE_DELEGATE_TwoParams(FOnColorGradingPickerValueChanged, FVector4&, bool);
+
+
+	SLATE_BEGIN_ARGS(SColorGradingPicker)
+		: _AllowSpin(true)
+		, _SupportDynamicSliderMaxValue(false)
+		, _SupportDynamicSliderMinValue(false)
+		, _MainDelta(0.01f)
+		, _MainShiftMultiplier(10.f)
+		, _MainCtrlMultiplier(0.1f)
+		, _ColorGradingModes(EColorGradingModes::Saturation)
+		, _OnColorCommitted()
+		, _OnQueryCurrentColor()
+		{ }
+
+		SLATE_ARGUMENT(TOptional<float>, ValueMin)
+		SLATE_ARGUMENT(TOptional<float>, ValueMax)
+		SLATE_ARGUMENT(TOptional<float>, SliderValueMin)
+		SLATE_ARGUMENT(TOptional<float>, SliderValueMax)
+		SLATE_ATTRIBUTE(bool, AllowSpin)
+
+		/** Tell us if we want to support dynamically changing of the max value using ctrl */
+		SLATE_ATTRIBUTE(bool, SupportDynamicSliderMaxValue)
+		/** Tell us if we want to support dynamically changing of the min value using ctrl */
+		SLATE_ATTRIBUTE(bool, SupportDynamicSliderMinValue)
+
+		SLATE_ARGUMENT(float, MainDelta)
+
+		/** Multiplier to use when shift is held down */
+		SLATE_ARGUMENT(float, MainShiftMultiplier)
+		/** Multiplier to use when ctrl is held down */
+		SLATE_ARGUMENT(float, MainCtrlMultiplier)
+
+		SLATE_ARGUMENT_DEPRECATED(int32, MainShiftMouseMovePixelPerDelta, 5.4, "Shift Mouse Move Pixel Per Delta is deprecated and incrementing by a fixed delta per pixel is no longer supported. Please use ShiftMultiplier and CtrlMultiplier which will multiply the step per mouse move")
+
+		SLATE_ARGUMENT(EColorGradingModes, ColorGradingModes)
+
+		/** The event called when the color is committed */
+		SLATE_EVENT(FOnColorGradingPickerValueChanged, OnColorCommitted)
+
+		/** Callback to get the current FVector4 value */
+		SLATE_EVENT(FOnGetCurrentVector4Value, OnQueryCurrentColor)
+
+		/** Called right before the slider begins to move */
+		SLATE_EVENT(FSimpleDelegate, OnBeginSliderMovement)
+
+		/** Called right after the slider handle is released by the user */
+		SLATE_EVENT(FSimpleDelegate, OnEndSliderMovement)
+
+		/** Called when the mouse captures starts on the color wheel */
+		SLATE_EVENT(FSimpleDelegate, OnBeginMouseCapture)
+
+		/** Called when the mouse captures ends for the color wheel */
+		SLATE_EVENT(FSimpleDelegate, OnEndMouseCapture)
+
+	SLATE_END_ARGS()
+
+	/**	Destructor. */
+	ADVANCEDWIDGETS_API ~SColorGradingPicker();
+
+public:
+
+	/**
+	 * Construct the widget
+	 *
+	 * @param InArgs Declaration from which to construct the widget.
+	 */
+	ADVANCEDWIDGETS_API void Construct(const FArguments& InArgs);
+
+	FOnNumericEntryBoxDynamicSliderMinMaxValueChanged& GetOnNumericEntryBoxDynamicSliderMaxValueChangedDelegate() { return OnNumericEntryBoxDynamicSliderMaxValueChanged; }
+	FOnNumericEntryBoxDynamicSliderMinMaxValueChanged& GetOnNumericEntryBoxDynamicSliderMinValueChangedDelegate() { return OnNumericEntryBoxDynamicSliderMinValueChanged; }
+
+	/** Callback when the max/min spinner value are changed (only apply if SupportDynamicSliderMaxValue or SupportDynamicSliderMinValue are true) */
+	ADVANCEDWIDGETS_API void OnDynamicSliderMaxValueChanged(float NewMaxSliderValue, TWeakPtr<SWidget> InValueChangedSourceWidget, bool IsOriginator, bool UpdateOnlyIfHigher);
+	ADVANCEDWIDGETS_API void OnDynamicSliderMinValueChanged(float NewMinSliderValue, TWeakPtr<SWidget> InValueChangedSourceWidget, bool IsOriginator, bool UpdateOnlyIfLower);
+
+protected:
+
+	ADVANCEDWIDGETS_API void TransformLinearColorRangeToColorGradingRange(FVector4& VectorValue) const;
+	ADVANCEDWIDGETS_API void TransformColorGradingRangeToLinearColorRange(FVector4& VectorValue) const;
+	ADVANCEDWIDGETS_API void TransformColorGradingRangeToLinearColorRange(float& FloatValue);
+
+	ADVANCEDWIDGETS_API TOptional<float> OnGetMainValue() const;
+	ADVANCEDWIDGETS_API void OnMainValueChanged(float InValue, bool ShouldCommitValueChanges);
+	ADVANCEDWIDGETS_API void OnMainValueCommitted(float InValue, ETextCommit::Type CommitType);
+
+	ADVANCEDWIDGETS_API FLinearColor GetCurrentLinearColor();
+
+	ADVANCEDWIDGETS_API bool IsEntryBoxEnabled() const;
+
+	// Callback for value changes in the color spectrum picker.
+	ADVANCEDWIDGETS_API void HandleCurrentColorValueChanged(const FLinearColor& NewValue, bool ShouldCommitValueChanges);
+
+	ADVANCEDWIDGETS_API void HandleColorWheelMouseCaptureBegin(const FLinearColor& InValue);
+	ADVANCEDWIDGETS_API void HandleColorWheelMouseCaptureEnd(const FLinearColor& InValue);
+
+	ADVANCEDWIDGETS_API void OnBeginSliderMovement();
+	ADVANCEDWIDGETS_API void OnEndSliderMovement(float NewValue);
+	ADVANCEDWIDGETS_API void AdjustRatioValue(FVector4& NewValue);
+
+	bool bIsMouseDragging;
+	FVector4 StartDragRatio;
+
+	float SliderValueMin;
+	float SliderValueMax;
+	float MainDelta;
+	float MainShiftMultiplier;
+	float MainCtrlMultiplier;
+	EColorGradingModes ColorGradingModes;
+
+	TSharedPtr<SNumericEntryBox<float>> NumericEntryBoxWidget;
+
+	/** Invoked when a new value is selected on the color wheel */
+	FOnColorGradingPickerValueChanged OnColorCommitted;
+
+	FOnGetCurrentVector4Value OnQueryCurrentColor;
+	FOnNumericEntryBoxDynamicSliderMinMaxValueChanged OnNumericEntryBoxDynamicSliderMaxValueChanged;
+	FOnNumericEntryBoxDynamicSliderMinMaxValueChanged OnNumericEntryBoxDynamicSliderMinValueChanged;
+
+	FSimpleDelegate ExternalBeginSliderMovementDelegate;
+	FSimpleDelegate ExternalEndSliderMovementDelegate;
+	FSimpleDelegate ExternalBeginMouseCaptureDelegate;
+	FSimpleDelegate ExternalEndMouseCaptureDelegate;
+};
+
+} //namespace
\ No newline at end of file
diff --git a/Engine/Source/Runtime/AdvancedWidgets/Public/Widgets/ColorGrading/SColorGradingWheel.h b/Engine/Source/Runtime/AdvancedWidgets/Public/Widgets/ColorGrading/SColorGradingWheel.h
new file mode 100644
--- /dev/null
+++ b/Engine/Source/Runtime/AdvancedWidgets/Public/Widgets/ColorGrading/SColorGradingWheel.h
@@ -0,0 +1,161 @@
+// Copyright Epic Games, Inc. All Rights Reserved.
+
+#pragma once
+
+#include "CoreMinimal.h"
+#include "Misc/Attribute.h"
+#include "Input/Reply.h"
+#include "Widgets/DeclarativeSyntaxSupport.h"
+#include "Widgets/SLeafWidget.h"
+#include "Framework/SlateDelegates.h"
+
+class UMaterial;
+
+namespace UE::ColorGrading
+{
+
+/**
+ * Implements the color wheel widget.
+ */
+class SColorGradingWheel
+	: public SLeafWidget, public FGCObject
+{
+	SLATE_DECLARE_WIDGET_API(SColorGradingWheel, SLeafWidget, ADVANCEDWIDGETS_API)
+
+public:
+
+	DECLARE_DELEGATE_OneParam(FOnColorGradingWheelMouseCapture, const FLinearColor&);
+	DECLARE_DELEGATE_OneParam(FOnColorGradingWheelValueChanged, const FLinearColor&);
+
+	SLATE_BEGIN_ARGS(SColorGradingWheel)
+		: _SelectedColor()
+		, _DesiredWheelSize()
+		, _ExponentDisplacement()
+		, _OnMouseCaptureBegin()
+		, _OnMouseCaptureEnd()
+		, _OnValueChanged()
+		{ }
+
+		/** The current color selected by the user. */
+		SLATE_ATTRIBUTE(FLinearColor, SelectedColor)
+
+		SLATE_ATTRIBUTE(int32, DesiredWheelSize)
+
+		SLATE_ATTRIBUTE(float, ExponentDisplacement)
+
+		/** Invoked when the mouse is pressed and a capture begins. */
+		SLATE_EVENT(FOnColorGradingWheelMouseCapture, OnMouseCaptureBegin)
+
+		/** Invoked when the mouse is released and a capture ends. */
+		SLATE_EVENT(FOnColorGradingWheelMouseCapture, OnMouseCaptureEnd)
+
+		/** Invoked when a new value is selected on the color wheel. */
+		SLATE_EVENT(FOnColorGradingWheelValueChanged, OnValueChanged)
+
+	SLATE_END_ARGS()
+
+public:
+	ADVANCEDWIDGETS_API SColorGradingWheel();
+
+	/**
+	 * Construct this widget.
+	 *
+	 * @param InArgs The declaration data for this widget.
+	 */
+	ADVANCEDWIDGETS_API void Construct(const FArguments& InArgs);
+
+public:
+
+	// SWidget overrides
+
+	ADVANCEDWIDGETS_API virtual FVector2D ComputeDesiredSize(float) const override;
+	ADVANCEDWIDGETS_API virtual FReply OnMouseButtonDoubleClick(const FGeometry& InMyGeometry, const FPointerEvent& InMouseEvent) override;
+	ADVANCEDWIDGETS_API virtual FReply OnMouseButtonDown(const FGeometry& MyGeometry, const FPointerEvent& MouseEvent) override;
+	ADVANCEDWIDGETS_API virtual FReply OnMouseButtonUp(const FGeometry& MyGeometry, const FPointerEvent& MouseEvent) override;
+	ADVANCEDWIDGETS_API virtual FReply OnMouseMove(const FGeometry& MyGeometry, const FPointerEvent& MouseEvent) override;
+	ADVANCEDWIDGETS_API virtual int32 OnPaint(const FPaintArgs& Args, const FGeometry& AllottedGeometry, const FSlateRect& MyCullingRect, FSlateWindowElementList& OutDrawElements, int32 LayerId, const FWidgetStyle& InWidgetStyle, bool bParentEnabled) const override;
+
+protected:
+
+	/**
+	 * Calculates the position of the color selection indicator.
+	 *
+	 * @return The position relative to the widget.
+	 */
+	ADVANCEDWIDGETS_API UE::Slate::FDeprecateVector2DResult CalcRelativePositionFromCenter() const;
+
+	/**
+	 * Performs actions according to mouse click / move
+	 *
+	 * @return	True if the mouse action occurred within the color wheel radius
+	 */
+	ADVANCEDWIDGETS_API bool ProcessMouseAction(const FGeometry& MyGeometry, const FPointerEvent& MouseEvent, bool bProcessWhenOutsideColorWheel);
+
+	/** */
+	ADVANCEDWIDGETS_API void SetSelectedColorAttribute(TAttribute<FLinearColor> InSelectedColor);
+
+	/** */
+	ADVANCEDWIDGETS_API void SetDesiredWheelSizeAttribute(TAttribute<int32> InDesiredWheelSize);
+
+	/** */
+	ADVANCEDWIDGETS_API void SetExponentDisplacementAttribute(TAttribute<float> InExponentDisplacement);
+
+	/** @return an attribute reference of SelectedColor */
+	TSlateAttributeRef<FLinearColor> GetSelectedColorAttribute() const { return TSlateAttributeRef<FLinearColor>(SharedThis(this), SelectedColorAttribute); }
+
+	/** @return an attribute reference of DesiredWheelSize */
+	TSlateAttributeRef<int32> GetDesiredWheelSizeAttribute() const { return TSlateAttributeRef<int32>(SharedThis(this), DesiredWheelSizeAttribute); }
+
+	/** @return an attribute reference of ExponentDisplacement */
+	TSlateAttributeRef<float> GetExponentDisplacementAttribute() const { return TSlateAttributeRef<float>(SharedThis(this), ExponentDisplacementAttribute); }
+
+	/** The material to use for the wheel's background. */
+	TObjectPtr<UMaterial> BackgroundMaterial;
+
+	/** The brush to use for the wheel's background. */
+	TSharedPtr<FSlateBrush> BackgroundImage;
+
+	/** The cross to show on top of the color wheel background. */
+	const FSlateBrush* CrossImage;
+
+	/** The color selector image to show. */
+	const FSlateBrush* SelectorImage;
+
+	/** Invoked when the mouse is pressed and a capture begins. */
+	FOnColorGradingWheelMouseCapture OnMouseCaptureBegin;
+
+	/** Invoked when the mouse is let up and a capture ends. */
+	FOnColorGradingWheelMouseCapture OnMouseCaptureEnd;
+
+	/** Invoked when a new value is selected on the color wheel. */
+	FOnColorGradingWheelValueChanged OnValueChanged;
+
+private:
+
+	//~ Begin FGCObject
+	virtual FString GetReferencerName() const override;
+	virtual void AddReferencedObjects(FReferenceCollector& InCollector) override;
+	//~ End FGCObject
+
+	/** The angle in degrees (clockwise from due east) at which hue 0/red is represented on the wheel. */
+	static constexpr float HueAngleOffset = -120.f / 180.f * PI;
+
+	/** The current color selected by the user. */
+	TSlateAttribute<FLinearColor> SelectedColorAttribute;
+
+	TSlateAttribute<int32> DesiredWheelSizeAttribute;
+	TSlateAttribute<float> ExponentDisplacementAttribute;
+
+	/** Flags used to check if the SlateAttribute is set. */
+	union
+	{
+		struct
+		{
+			uint8 bIsAttributeDesiredWheelSizeSet : 1;
+			uint8 bIsAttributeExponentDisplacementSet : 1;
+		};
+		uint8 Union_IsAttributeSet;
+	};
+};
+
+} //namespace
diff --git a/Engine/Source/Runtime/AppFramework/Private/Widgets/Colors/SColorGradingPicker.cpp b/Engine/Source/Runtime/AppFramework/Private/Widgets/Colors/SColorGradingPicker.cpp
--- a/Engine/Source/Runtime/AppFramework/Private/Widgets/Colors/SColorGradingPicker.cpp
+++ b/Engine/Source/Runtime/AppFramework/Private/Widgets/Colors/SColorGradingPicker.cpp
@@ -41,7 +41,9 @@ void SColorGradingPicker::Construct( const FArguments& InArgs )
 	OnQueryCurrentColor = InArgs._OnQueryCurrentColor;
 	
 	float ColorGradingWheelExponent = 2.4f;
+	PRAGMA_DISABLE_DEPRECATION_WARNINGS
 	if (ColorGradingModes == EColorGradingModes::Offset)
+	PRAGMA_ENABLE_DEPRECATION_WARNINGS
 	{
 		ColorGradingWheelExponent = 3.0f;
 	}
@@ -65,7 +67,9 @@ void SColorGradingPicker::Construct( const FArguments& InArgs )
 				.MinDesiredWidth(125.0f)
 				.MaxDesiredWidth(125.0f)
 				[
+					PRAGMA_DISABLE_DEPRECATION_WARNINGS
 					SNew(SColorGradingWheel)
+					PRAGMA_ENABLE_DEPRECATION_WARNINGS
 					.SelectedColor(TAttribute<FLinearColor>::Create(TAttribute<FLinearColor>::FGetter::CreateSP(this, &SColorGradingPicker::GetCurrentLinearColor)))
 					.DesiredWheelSize(125)
 					.ExponentDisplacement(ColorGradingWheelExponent)
diff --git a/Engine/Source/Runtime/AppFramework/Public/Widgets/Colors/SColorGradingPicker.h b/Engine/Source/Runtime/AppFramework/Public/Widgets/Colors/SColorGradingPicker.h
--- a/Engine/Source/Runtime/AppFramework/Public/Widgets/Colors/SColorGradingPicker.h
+++ b/Engine/Source/Runtime/AppFramework/Public/Widgets/Colors/SColorGradingPicker.h
@@ -12,12 +12,13 @@
 #include "Widgets/Input/SNumericEntryBox.h"
 
 /** Callback to get the current FVector4 value */
+UE_DEPRECATED(5.5, "Color grading widgets and associated types have moved to the AdvancedWidgets module.")
 DECLARE_DELEGATE_RetVal_OneParam(bool, FOnGetCurrentVector4Value, FVector4&)
 
 /**
 * Enumerates color picker modes.
 */
-enum class EColorGradingModes
+enum class UE_DEPRECATED(5.5, "Color grading widgets and associated types have moved to the AdvancedWidgets module.") EColorGradingModes
 {
 	Saturation,
 	Contrast,
@@ -33,7 +34,7 @@ enum class EColorGradingModes
  * use the functions OpenColorGradingWheel and DestroyColorGradingWheel, since they hold a static
  * instance of the color picker.
  */
-class SColorGradingPicker
+class UE_DEPRECATED(5.5, "Color grading widgets and associated types have moved to the AdvancedWidgets module.") SColorGradingPicker
 	: public SCompoundWidget
 {
 public:
@@ -51,7 +52,9 @@ public:
 		, _MainDelta(0.01f)
 		, _MainShiftMultiplier(10.f)
 		, _MainCtrlMultiplier(0.1f)
+		PRAGMA_DISABLE_DEPRECATION_WARNINGS
 		, _ColorGradingModes(EColorGradingModes::Saturation)
+		PRAGMA_ENABLE_DEPRECATION_WARNINGS
 		, _OnColorCommitted()
 		, _OnQueryCurrentColor()
 	{ }
@@ -76,13 +79,19 @@ public:
 
 		SLATE_ARGUMENT_DEPRECATED( int32, MainShiftMouseMovePixelPerDelta, 5.4, "Shift Mouse Move Pixel Per Delta is deprecated and incrementing by a fixed delta per pixel is no longer supported. Please use ShiftMultiplier and CtrlMultiplier which will multiply the step per mouse move")
 
+		PRAGMA_DISABLE_DEPRECATION_WARNINGS
 		SLATE_ARGUMENT( EColorGradingModes, ColorGradingModes)
+		PRAGMA_ENABLE_DEPRECATION_WARNINGS
 
 		/** The event called when the color is committed */
+		PRAGMA_DISABLE_DEPRECATION_WARNINGS
 		SLATE_EVENT(FOnColorGradingPickerValueChanged, OnColorCommitted )
+		PRAGMA_ENABLE_DEPRECATION_WARNINGS
 
 		/** Callback to get the current FVector4 value */
+		PRAGMA_DISABLE_DEPRECATION_WARNINGS
 		SLATE_EVENT(FOnGetCurrentVector4Value, OnQueryCurrentColor)
+		PRAGMA_ENABLE_DEPRECATION_WARNINGS
 
 		/** Called right before the slider begins to move */
 		SLATE_EVENT(FSimpleDelegate, OnBeginSliderMovement)
@@ -149,14 +158,19 @@ protected:
 	float MainDelta;
 	float MainShiftMultiplier;
 	float MainCtrlMultiplier;
+
+	PRAGMA_DISABLE_DEPRECATION_WARNINGS
 	EColorGradingModes ColorGradingModes;
+	PRAGMA_ENABLE_DEPRECATION_WARNINGS
 
 	TSharedPtr<SNumericEntryBox<float>> NumericEntryBoxWidget;
 
 	/** Invoked when a new value is selected on the color wheel */
 	FOnColorGradingPickerValueChanged OnColorCommitted;
 
+	PRAGMA_DISABLE_DEPRECATION_WARNINGS
 	FOnGetCurrentVector4Value OnQueryCurrentColor;
+	PRAGMA_ENABLE_DEPRECATION_WARNINGS
 	FOnNumericEntryBoxDynamicSliderMinMaxValueChanged OnNumericEntryBoxDynamicSliderMaxValueChanged;
 	FOnNumericEntryBoxDynamicSliderMinMaxValueChanged OnNumericEntryBoxDynamicSliderMinValueChanged;
 
diff --git a/Engine/Source/Runtime/Slate/Private/Widgets/Colors/SColorGradingWheel.cpp b/Engine/Source/Runtime/Slate/Private/Widgets/Colors/SColorGradingWheel.cpp
--- a/Engine/Source/Runtime/Slate/Private/Widgets/Colors/SColorGradingWheel.cpp
+++ b/Engine/Source/Runtime/Slate/Private/Widgets/Colors/SColorGradingWheel.cpp
@@ -4,7 +4,10 @@
 #include "Rendering/RenderingCommon.h"
 #include "Rendering/DrawElements.h"
 
+PRAGMA_DISABLE_DEPRECATION_WARNINGS
 SLATE_IMPLEMENT_WIDGET(SColorGradingWheel)
+PRAGMA_ENABLE_DEPRECATION_WARNINGS
+
 void SColorGradingWheel::PrivateRegisterAttributes(FSlateAttributeInitializer& AttributeInitializer)
 {
 	SLATE_ADD_MEMBER_ATTRIBUTE_DEFINITION_WITH_NAME(AttributeInitializer, "SelectedColor", SelectedColorAttribute, EInvalidateWidgetReason::Paint);
diff --git a/Engine/Source/Runtime/Slate/Public/Widgets/Colors/SColorGradingWheel.h b/Engine/Source/Runtime/Slate/Public/Widgets/Colors/SColorGradingWheel.h
--- a/Engine/Source/Runtime/Slate/Public/Widgets/Colors/SColorGradingWheel.h
+++ b/Engine/Source/Runtime/Slate/Public/Widgets/Colors/SColorGradingWheel.h
@@ -12,7 +12,7 @@
 /**
  * Implements the color wheel widget.
  */
-class SColorGradingWheel
+class UE_DEPRECATED(5.5, "Color grading widgets and associated types have moved to the AdvancedWidgets module.") SColorGradingWheel
 	: public SLeafWidget
 {
 	SLATE_DECLARE_WIDGET_API(SColorGradingWheel, SLeafWidget, SLATE_API)
