Allow showing/editing spatially loaded flag on actor even when the level is not partitioned.
This allows to control the spatial flag on actors part of a level referenced by a Level Instance used in a partition map.

diff --git a/Engine/Source/Editor/DetailCustomizations/Private/ActorDetails.cpp b/Engine/Source/Editor/DetailCustomizations/Private/ActorDetails.cpp
--- a/Engine/Source/Editor/DetailCustomizations/Private/ActorDetails.cpp
+++ b/Engine/Source/Editor/DetailCustomizations/Private/ActorDetails.cpp
@@ -156,7 +156,6 @@ void FActorDetails::CustomizeDetails( IDetailLayoutBuilder& DetailLayout )
 		if (!bShouldDisplayWorldPartitionProperties)
 		{
 			DetailLayout.HideProperty(DetailLayout.GetProperty(AActor::GetRuntimeGridPropertyName(), AActor::StaticClass()));
-			DetailLayout.HideProperty(DetailLayout.GetProperty(AActor::GetIsSpatiallyLoadedPropertyName(), AActor::StaticClass()));
 			DetailLayout.HideProperty(DetailLayout.GetProperty(AActor::GetDataLayerAssetsPropertyName(), AActor::StaticClass()));
 			DetailLayout.HideProperty(DetailLayout.GetProperty(AActor::GetDataLayerPropertyName(), AActor::StaticClass()));
 			DetailLayout.HideProperty(DetailLayout.GetProperty(AActor::GetHLODLayerPropertyName(), AActor::StaticClass()));
diff --git a/Engine/Source/Runtime/Engine/Private/ActorEditor.cpp b/Engine/Source/Runtime/Engine/Private/ActorEditor.cpp
--- a/Engine/Source/Runtime/Engine/Private/ActorEditor.cpp
+++ b/Engine/Source/Runtime/Engine/Private/ActorEditor.cpp
@@ -111,7 +111,7 @@ bool AActor::CanEditChange(const FProperty* PropertyThatWillChange) const
 			UWorldPartition* OwningWorldPartition = OwningWorld ? OwningWorld->GetWorldPartition() : nullptr;
 			UWorldPartition* OuterWorldPartition = OuterWorld ? OuterWorld->GetWorldPartition() : nullptr;			
 
-			if (!OwningWorldPartition || !OuterWorldPartition)
+			if (!bIsSpatiallyLoadedProperty && (!OwningWorldPartition || !OuterWorldPartition))
 			{
 				return false;
 			}
