[ColorGrading] Load AdvancedWidgets module on startup to avoid crash due to missing styles

diff --git a/Engine/Source/Editor/DetailCustomizations/Private/DetailCustomizations.cpp b/Engine/Source/Editor/DetailCustomizations/Private/DetailCustomizations.cpp
--- a/Engine/Source/Editor/DetailCustomizations/Private/DetailCustomizations.cpp
+++ b/Engine/Source/Editor/DetailCustomizations/Private/DetailCustomizations.cpp
@@ -5,6 +5,7 @@
 #include "AIDataProviderValueDetails.h"
 #include "ActorComponentDetails.h"
 #include "ActorDetails.h"
+#include "AdvancedWidgetsModule.h"
 #include "AmbientSoundDetails.h"
 #include "AnimMontageSegmentDetails.h"
 #include "AnimSequenceDetails.h"
@@ -173,6 +174,9 @@ IMPLEMENT_MODULE( FDetailCustomizationsModule, DetailCustomizations );
 
 void FDetailCustomizationsModule::StartupModule()
 {
+	// Load widgets (e.g. ColorGradingWheel) that some customizations depend on
+	FModuleManager::Get().LoadModuleChecked<FAdvancedWidgetsModule>("AdvancedWidgets");
+
 	FPropertyEditorModule& PropertyModule = FModuleManager::GetModuleChecked<FPropertyEditorModule>("PropertyEditor");
 	
 	RegisterPropertyTypeCustomizations();
