[ColorGrading] Adjust spacing of Details panel color grading controls

diff --git a/Engine/Source/Editor/DetailCustomizations/Private/ColorGradingVectorCustomization.cpp b/Engine/Source/Editor/DetailCustomizations/Private/ColorGradingVectorCustomization.cpp
--- a/Engine/Source/Editor/DetailCustomizations/Private/ColorGradingVectorCustomization.cpp
+++ b/Engine/Source/Editor/DetailCustomizations/Private/ColorGradingVectorCustomization.cpp
@@ -776,7 +776,7 @@ void FColorGradingCustomBuilder::GenerateHeaderRowContent(FDetailWidgetRow& Node
 		SpinBoxMinMaxSliderValues.DefaultMaxSliderValue = SpinBoxMinMaxSliderValues.CurrentMaxSliderValue;
 
 		VerticalBox->AddSlot()
-			.Padding(FMargin(0.0f, 8.0f, 4.0f, 0.0f))
+			.Padding(FMargin(0.0f, 4.0f, 0.0f, 0.0f))
 			.AutoHeight()
 			.HAlign(HAlign_Fill)
 			.VAlign(VAlign_Fill)
