* Fixed assert when copying from spline properties that have different values (prevent copying when this happens)
* Added copy/paste options to the spline point's type property

diff --git a/Engine/Source/Editor/DetailCustomizations/Private/SplineComponentDetails.cpp b/Engine/Source/Editor/DetailCustomizations/Private/SplineComponentDetails.cpp
--- a/Engine/Source/Editor/DetailCustomizations/Private/SplineComponentDetails.cpp
+++ b/Engine/Source/Editor/DetailCustomizations/Private/SplineComponentDetails.cpp
@@ -134,6 +134,8 @@ private:
 			bInitialized = false;
 		}
 
+		bool IsValid() const { return bInitialized; }
+
 		void Add(T InValue)
 		{
 			if (!bInitialized)
@@ -147,6 +149,11 @@ private:
 			}
 		}
 
+		bool HasMultipleValues() const
+		{
+			return !Value.IsSet();
+		}
+
 		TOptional<T> Value;
 		bool bInitialized;
 	};
@@ -179,6 +186,11 @@ private:
 			}
 		}
 
+		bool HasMultipleValues() const 
+		{
+			return !X.IsSet() || !Y.IsSet() || !Z.IsSet();
+		}
+
 		TOptional<float> X;
 		TOptional<float> Y;
 		TOptional<float> Z;
@@ -213,6 +225,11 @@ private:
 			}
 		}
 
+		bool HasMultipleValues() const
+		{
+			return !Roll.IsSet() || !Pitch.IsSet() || !Yaw.IsSet();
+		}
+
 		TOptional<float> Roll;
 		TOptional<float> Pitch;
 		TOptional<float> Yaw;
@@ -265,7 +282,8 @@ private:
 		Rotation,
 		Scale,
 		ArriveTangent,
-		LeaveTangent
+		LeaveTangent, 
+		Type
 	};
 
 	TSharedRef<SWidget> BuildSplinePointPropertyLabel(ESplinePointProperty SplinePointProp);
@@ -278,7 +296,7 @@ private:
 	FUIAction CreateCopyAction(ESplinePointProperty SplinePointProp);
 	FUIAction CreatePasteAction(ESplinePointProperty SplinePointProp);
 
-	bool OnCanCopy(ESplinePointProperty SplinePointProp) const { return true; }
+	bool OnCanCopy(ESplinePointProperty SplinePointProp) const;
 	void OnCopy(ESplinePointProperty SplinePointProp);
 	void OnPaste(ESplinePointProperty SplinePointProp);
 
@@ -511,7 +529,7 @@ void FSplinePointDetails::GenerateChildContent(IDetailChildrenBuilder& ChildrenB
 		[
 			SNew(STextBlock)
 			.Text(LOCTEXT("InputKey", "Input Key"))
-		.Font(IDetailLayoutBuilder::GetDetailFont())
+			.Font(IDetailLayoutBuilder::GetDetailFont())
 		]
 	.ValueContent()
 		.MinDesiredWidth(125.0f)
@@ -717,13 +735,15 @@ void FSplinePointDetails::GenerateChildContent(IDetailChildrenBuilder& ChildrenB
 		ChildrenBuilder.AddCustomRow(LOCTEXT("Type", "Type"))
 			.RowTag(TEXT("Type"))
 			.Visibility(TAttribute<EVisibility>(this, &FSplinePointDetails::IsEnabled))
+			.CopyAction(CreateCopyAction(ESplinePointProperty::Type))
+			.PasteAction(CreatePasteAction(ESplinePointProperty::Type))
 			.NameContent()
 			.HAlign(HAlign_Left)
 			.VAlign(VAlign_Center)
 			[
 				SNew(STextBlock)
 				.Text(LOCTEXT("Type", "Type"))
-			.Font(IDetailLayoutBuilder::GetDetailFont())
+				.Font(IDetailLayoutBuilder::GetDetailFont())
 			]
 		.ValueContent()
 			.MinDesiredWidth(125.0f)
@@ -1246,44 +1266,53 @@ void FSplinePointDetails::OnSplinePointTypeChanged(TSharedPtr<FString> NewValue,
 		return;
 	}
 
+	bool bWasModified = false;
 	// Scope the transaction to only include the value change and none of the derived data changes that might arise from NotifyPropertyModified
 	{
-		const FScopedTransaction Transaction(LOCTEXT("SetSplinePointType", "Set spline point type"));
-		SplineComp->Modify();
-
 		EInterpCurveMode Mode = CIM_Unknown;
-		if (NewValue.IsValid() && SplinePointTypes.Contains(NewValue))
+		if (NewValue.IsValid() && SplinePointTypes.ContainsByPredicate([&NewValue](const TSharedPtr<FString>& InSplinePointType) { return (*InSplinePointType == *NewValue); }))
 		{
 			const UEnum* SplinePointTypeEnum = StaticEnum<ESplinePointType::Type>();
 			check(SplinePointTypeEnum);
 			const int64 SplinePointType = SplinePointTypeEnum->GetValueByNameString(*NewValue);
 
 			Mode = ConvertSplinePointTypeToInterpCurveMode(static_cast<ESplinePointType::Type>(SplinePointType));
-		}
+			check(Mode != CIM_Unknown);
 
-		for (int32 Index : SelectedKeys)
-		{
-			if (Index < 0 || Index >= SplineComp->GetSplinePointsPosition().Points.Num())
+			const FScopedTransaction Transaction(LOCTEXT("SetSplinePointType", "Set spline point type"));
+			SplineComp->Modify();
+
+			for (int32 Index : SelectedKeys)
 			{
-				UE_LOG(LogSplineComponentDetails, Error, TEXT("Set spline point type: invalid index %d in selected points for spline component %s which contains %d spline points."),
-					Index, *SplineComp->GetPathName(), SplineComp->GetSplinePointsPosition().Points.Num());
-				continue;
-			}
+				if (Index < 0 || Index >= SplineComp->GetSplinePointsPosition().Points.Num())
+				{
+					UE_LOG(LogSplineComponentDetails, Error, TEXT("Set spline point type: invalid index %d in selected points for spline component %s which contains %d spline points."),
+						Index, *SplineComp->GetPathName(), SplineComp->GetSplinePointsPosition().Points.Num());
+					continue;
+				}
 
-			SplineComp->GetSplinePointsPosition().Points[Index].InterpMode = Mode;
+				if (SplineComp->GetSplinePointsPosition().Points[Index].InterpMode != Mode)
+				{
+					SplineComp->GetSplinePointsPosition().Points[Index].InterpMode = Mode;
+					bWasModified = true;
+				}
+			}
 		}
 	}
 
-	SplineComp->UpdateSpline();
-	SplineComp->bSplineHasBeenEdited = true;
-	FComponentVisualizer::NotifyPropertyModified(SplineComp, SplineCurvesProperty);
-	if (AActor* Owner = SplineComp->GetOwner())
+	if (bWasModified)
 	{
-		Owner->PostEditMove(true);
-	}
-	UpdateValues();
+		SplineComp->UpdateSpline();
+		SplineComp->bSplineHasBeenEdited = true;
+		FComponentVisualizer::NotifyPropertyModified(SplineComp, SplineCurvesProperty);
+		if (AActor* Owner = SplineComp->GetOwner())
+		{
+			Owner->PostEditMove(true);
+		}
+		UpdateValues();
 
-	GEditor->RedrawLevelEditingViewports(true);
+		GEditor->RedrawLevelEditingViewports(true);
+	}
 }
 
 USplineComponent* FSplinePointDetails::GetSplineComponentToVisualize() const
@@ -1573,6 +1602,9 @@ void FSplinePointDetails::SetSplinePointProperty(ESplinePointProperty SplinePoin
 		OnSetLeaveTangent(NewValue.Y, ETextCommit::Default, EAxis::Y);
 		OnSetLeaveTangent(NewValue.Z, ETextCommit::OnEnter, EAxis::Z);
 		break;
+	case ESplinePointProperty::Type:
+		checkf(false, TEXT("SetSplinePointProperty shouldn't be called for non-vector types"));
+		break;
 	default:
 		break;
 	}
@@ -1594,6 +1626,28 @@ FUIAction FSplinePointDetails::CreatePasteAction(ESplinePointProperty SplinePoin
 		FUIAction(FExecuteAction::CreateSP(this, &FSplinePointDetails::OnPaste, SplinePointProp));
 }
 
+bool FSplinePointDetails::OnCanCopy(ESplinePointProperty SplinePointProp) const
+{
+	// Can't copy if at least one of spline point's values is different (we're editing multiple values) : 
+	switch (SplinePointProp)
+	{
+	case ESplinePointProperty::Location:
+		return Position.IsValid() && !Position.HasMultipleValues();
+	case ESplinePointProperty::Rotation:
+		return Rotation.IsValid() && !Rotation.HasMultipleValues();
+	case ESplinePointProperty::Scale:
+		return Scale.IsValid() && !Scale.HasMultipleValues();
+	case ESplinePointProperty::ArriveTangent:
+		return ArriveTangent.IsValid() && !ArriveTangent.HasMultipleValues();
+	case ESplinePointProperty::LeaveTangent:
+		return LeaveTangent.IsValid() && !LeaveTangent.HasMultipleValues();
+	case ESplinePointProperty::Type:
+		return PointType.IsValid() && !PointType.HasMultipleValues();
+	default:
+		return false;
+	}
+}
+
 void FSplinePointDetails::OnCopy(ESplinePointProperty SplinePointProp)
 {
 	FString CopyStr;
@@ -1614,6 +1668,17 @@ void FSplinePointDetails::OnCopy(ESplinePointProperty SplinePointProp)
 	case ESplinePointProperty::LeaveTangent:
 		CopyStr = FString::Printf(TEXT("(X=%f,Y=%f,Z=%f)"), LeaveTangent.X.GetValue(), LeaveTangent.Y.GetValue(), LeaveTangent.Z.GetValue());
 		break;
+	case ESplinePointProperty::Type:
+	{
+		FString TypeString = UEnum::GetValueAsString(PointType.Value.GetValue());
+		int32 LastColonPos;
+		if (TypeString.FindLastChar(':', LastColonPos))
+		{
+			check((LastColonPos + 1) < TypeString.Len());
+			CopyStr = TypeString.RightChop(LastColonPos + 1);
+		}
+		break;
+	}
 	default:
 		break;
 	}
@@ -1702,6 +1767,12 @@ void FSplinePointDetails::PasteFromText(
 			}
 			break;
 		}
+	case ESplinePointProperty::Type:
+	{
+		ESelectInfo::Type DummySelectInfo = ESelectInfo::Direct;
+		OnSplinePointTypeChanged(MakeShared<FString>(InText), DummySelectInfo);
+		break;
+	}
 	default:
 		break;
 	}
