Make EditConditions work for FGuid

#rnx

diff --git a/Engine/Source/Editor/DetailCustomizations/Private/GuidStructCustomization.cpp b/Engine/Source/Editor/DetailCustomizations/Private/GuidStructCustomization.cpp
--- a/Engine/Source/Editor/DetailCustomizations/Private/GuidStructCustomization.cpp
+++ b/Engine/Source/Editor/DetailCustomizations/Private/GuidStructCustomization.cpp
@@ -85,7 +85,14 @@ void FGuidStructCustomization::CustomizeHeader( TSharedRef<class IPropertyHandle
 					// text box
 					SAssignNew(TextBox, SEditableTextBox)
 						.ClearKeyboardFocusOnCommit(false)
-						.IsEnabled(!PropertyHandle->IsEditConst())
+						.IsEnabled_Lambda([PropertyHandleWeak = PropertyHandle->AsWeak()]()->bool
+						{
+							if (TSharedPtr<IPropertyHandle> PropertyHandleLocal = PropertyHandleWeak.Pin())
+							{
+								return !PropertyHandleLocal->IsEditConst();
+							}
+							return false;
+						})
 						.ForegroundColor(this, &FGuidStructCustomization::HandleTextBoxForegroundColor)
 						.OnTextChanged(this, &FGuidStructCustomization::HandleTextBoxTextChanged)
 						.OnTextCommitted(this, &FGuidStructCustomization::HandleTextBoxTextCommited)
