Added Overscan subcategory in the Camera details customizer where all camera overscan properties are displayed for Camera and Cine Camera.

diff --git a/Engine/Source/Editor/DetailCustomizations/Private/CameraDetails.cpp b/Engine/Source/Editor/DetailCustomizations/Private/CameraDetails.cpp
--- a/Engine/Source/Editor/DetailCustomizations/Private/CameraDetails.cpp
+++ b/Engine/Source/Editor/DetailCustomizations/Private/CameraDetails.cpp
@@ -10,6 +10,7 @@
 #include "DetailLayoutBuilder.h"
 #include "DetailWidgetRow.h"
 #include "EditorCategoryUtils.h"
+#include "IDetailGroup.h"
 #include "Fonts/SlateFontInfo.h"
 #include "Framework/Commands/UIAction.h"
 #include "Framework/MultiBox/MultiBoxBuilder.h"
@@ -88,7 +89,7 @@ void FCameraDetails::CustomizeDetails( IDetailLayoutBuilder& DetailLayout )
 		}
 	}
 
-	DetailLayout.EditCategory("Current Camera Settings", FText::GetEmpty(), ECategoryPriority::Important);
+	IDetailCategoryBuilder& CurrentCameraSettingsCategory = DetailLayout.EditCategory("Current Camera Settings", FText::GetEmpty(), ECategoryPriority::Important);
 
 	if (bCameraSettingsHidden == false)
 	{
@@ -169,6 +170,18 @@ void FCameraDetails::CustomizeDetails( IDetailLayoutBuilder& DetailLayout )
 					]
 				]
 			];
+
+		IDetailGroup& OverscanGroup = CameraCategory.AddGroup(TEXT("Overscan"), LOCTEXT("OverscanGroup", "Overscan"));
+		OverscanGroup.AddPropertyRow(DetailLayout.GetProperty(GET_MEMBER_NAME_CHECKED(UCameraComponent, Overscan)));
+		OverscanGroup.AddPropertyRow(DetailLayout.GetProperty(GET_MEMBER_NAME_CHECKED(UCameraComponent, bScaleResolutionWithOverscan)));
+		OverscanGroup.AddPropertyRow(DetailLayout.GetProperty(GET_MEMBER_NAME_CHECKED(UCameraComponent, bCropOverscan)));
+	}
+	else
+	{
+		IDetailGroup& OverscanGroup = CurrentCameraSettingsCategory.AddGroup(TEXT("Overscan"), LOCTEXT("OverscanGroup", "Overscan"));
+		OverscanGroup.AddPropertyRow(DetailLayout.GetProperty(GET_MEMBER_NAME_CHECKED(UCameraComponent, Overscan)));
+		OverscanGroup.AddPropertyRow(DetailLayout.GetProperty(GET_MEMBER_NAME_CHECKED(UCameraComponent, bScaleResolutionWithOverscan)));
+		OverscanGroup.AddPropertyRow(DetailLayout.GetProperty(GET_MEMBER_NAME_CHECKED(UCameraComponent, bCropOverscan)));
 	}
 
 	IDetailCategoryBuilder& CameraSettingsCategory = DetailLayout.EditCategory( "CameraOptions", FText::GetEmpty(), ECategoryPriority::Important );
