SceneCapture: Fixed missing show flag property reset.

We also update the show flag toggle behavior to remove overrides when the (archetype) instance state matches the default. Undo/redo transactions are now supported as well.

diff --git a/Engine/Source/Editor/DetailCustomizations/Private/SceneCaptureDetails.cpp b/Engine/Source/Editor/DetailCustomizations/Private/SceneCaptureDetails.cpp
--- a/Engine/Source/Editor/DetailCustomizations/Private/SceneCaptureDetails.cpp
+++ b/Engine/Source/Editor/DetailCustomizations/Private/SceneCaptureDetails.cpp
@@ -22,6 +22,7 @@
 #include "Misc/Attribute.h"
 #include "PropertyHandle.h"
 #include "RenderCaptureInterface.h"
+#include "ScopedTransaction.h"
 #include "ShowFlags.h"
 #include "SlotBase.h"
 #include "Styling/SlateTypes.h"
@@ -46,18 +47,74 @@ TSharedRef<IDetailCustomization> FSceneCaptureDetails::MakeInstance()
 	return MakeShareable( new FSceneCaptureDetails );
 }
 
-// Used to sort the show flags inside of their categories in the order of the text that is actually displayed
-inline static bool SortAlphabeticallyByLocalizedText(const FString& ip1, const FString& ip2)
+namespace UE
 {
-	FText LocalizedText1;
-	FEngineShowFlags::FindShowFlagDisplayName(ip1, LocalizedText1);
+	namespace DetailsCustomization
+	{
+		namespace Private
+		{
+			// Used to sort the show flags inside of their categories in the order of the text that is actually displayed
+			inline static bool SortAlphabeticallyByLocalizedText(const FString& ip1, const FString& ip2)
+			{
+				FText LocalizedText1;
+				FEngineShowFlags::FindShowFlagDisplayName(ip1, LocalizedText1);
 
-	FText LocalizedText2;
-	FEngineShowFlags::FindShowFlagDisplayName(ip2, LocalizedText2);
+				FText LocalizedText2;
+				FEngineShowFlags::FindShowFlagDisplayName(ip2, LocalizedText2);
+
+				return LocalizedText1.ToString() < LocalizedText2.ToString();
+			}
+
+			// Utility function to get a show flag state by name.
+			static bool GetShowFlagState(const USceneCaptureComponent* SceneCaptureComponent, FStringView ShowFlagName)
+			{
+				if (!IsValid(SceneCaptureComponent))
+				{
+					return false;
+				}
+
+				const int32 SettingIndex = SceneCaptureComponent->ShowFlags.FindIndexByName(ShowFlagName.GetData());
+				return SettingIndex != INDEX_NONE ? SceneCaptureComponent->ShowFlags.GetSingleFlag(SettingIndex) : false;
+			}
 
-	return LocalizedText1.ToString() < LocalizedText2.ToString();
+			// Update archetype instances so that once the value matches the default, the instance override will be removed. This correctly mimics the behavior of other checkbox states.
+			static void UpdateArchetypeChainShowFlagSettings(USceneCaptureComponent* SceneCaptureComponent, FStringView ShowFlagName, bool bEnabledState)
+			{
+				if (!IsValid(SceneCaptureComponent))
+				{
+					return;
+				}
+
+				FProperty* ShowFlagsSettingsPropertyVariable = FindFieldChecked<FProperty>(USceneCaptureComponent::StaticClass(), GET_MEMBER_NAME_CHECKED(USceneCaptureComponent, ShowFlagSettings));
+				FPropertyChangedEvent PropertyChangedEvent(ShowFlagsSettingsPropertyVariable);
+
+				// Update archetype instances so that once the value matches the default, the instance override will be removed.
+				// This correctly mimics the behavior of other checkbox states.
+				TArray<UObject*> Instances;
+				SceneCaptureComponent->GetArchetypeInstances(Instances);
+				for (UObject* Instance : Instances)
+				{
+					if (USceneCaptureComponent* InstanceComponent = Cast<USceneCaptureComponent>(Instance))
+					{
+						TArray<FEngineShowFlagsSetting> InstanceShowFlagSettings = InstanceComponent->GetShowFlagSettings();
+						const int32 RemovedElements = InstanceShowFlagSettings.RemoveAll([ShowFlagName, bEnabledState](const FEngineShowFlagsSetting& InSetting)
+							{
+								return ShowFlagName.Equals(InSetting.ShowFlagName) && (bEnabledState == InSetting.Enabled);
+							});
+						if (RemovedElements > 0)
+						{
+							InstanceComponent->PreEditChange(ShowFlagsSettingsPropertyVariable);
+							InstanceComponent->SetShowFlagSettings(InstanceShowFlagSettings);
+							InstanceComponent->PostEditChangeProperty(PropertyChangedEvent);
+						}
+					}
+				}
+			}
+		}
+	}
 }
 
+
 void FSceneCaptureDetails::CustomizeDetails( IDetailLayoutBuilder& DetailLayout )
 {
 	IDetailCategoryBuilder& SceneCaptureCategoryBuilder = DetailLayout.EditCategory("SceneCapture");
@@ -145,7 +202,7 @@ void FSceneCaptureDetails::CustomizeDetails( IDetailLayoutBuilder& DetailLayout
 	// Sort the flags in their respective group alphabetically
 	for (TArray<FString>& ShowFlagGroup : ShowFlagsByGroup)
 	{
-		ShowFlagGroup.Sort(SortAlphabeticallyByLocalizedText);
+		ShowFlagGroup.Sort(UE::DetailsCustomization::Private::SortAlphabeticallyByLocalizedText);
 	}
 
 	// Add each group
@@ -213,6 +270,7 @@ void FSceneCaptureDetails::CustomizeDetails( IDetailLayoutBuilder& DetailLayout
 
 				Group.AddWidgetRow()
 					.IsEnabled(true)
+					.OverrideResetToDefault(FResetToDefaultOverride::Create(TAttribute<bool>::CreateSP(this, &FSceneCaptureDetails::GetShowFlagResetVisibility, FlagName), FSimpleDelegate::CreateSP(this, &FSceneCaptureDetails::OnResetShowFlag, FlagName)))
 					.NameContent()
 					[
 						SNew(STextBlock)
@@ -291,24 +349,6 @@ void FSceneCaptureDetails::CustomizeDetails( IDetailLayoutBuilder& DetailLayout
 	}
 }
 
-static bool FindShowFlagSetting(
-	TArray<FEngineShowFlagsSetting>& ShowFlagSettings,
-	FString FlagName,
-	FEngineShowFlagsSetting** ShowFlagSettingOut)
-{
-	bool HasSetting = false;
-	for (int32 ShowFlagSettingsIndex = 0; ShowFlagSettingsIndex < ShowFlagSettings.Num(); ++ShowFlagSettingsIndex)
-	{
-		if (ShowFlagSettings[ShowFlagSettingsIndex].ShowFlagName.Equals(FlagName))
-		{
-			HasSetting = true;
-			*ShowFlagSettingOut = &(ShowFlagSettings[ShowFlagSettingsIndex]);
-			break;
-		}
-	}
-	return HasSetting;
-}
-
 ECheckBoxState FSceneCaptureDetails::OnGetDisplayCheckState(FString ShowFlagName) const
 {
 	TArray<const void*> RawData;
@@ -372,6 +412,8 @@ void FSceneCaptureDetails::OnShowFlagCheckStateChanged(ECheckBoxState InNewRadio
 		return;
 	}
 
+	FScopedTransaction Transaction(LOCTEXT("ShowFlagCheckStateChanged", "Show Flag Check State Changed"));
+
 	ShowFlagSettingsProperty->NotifyPreChange();
 
 	TArray<void*> RawData;
@@ -386,10 +428,9 @@ void FSceneCaptureDetails::OnShowFlagCheckStateChanged(ECheckBoxState InNewRadio
 		void* Data = RawData[ObjectIdx];
 		check(Data);
 
-		const UObject* SceneComp = OuterObjects[ObjectIdx];
+		USceneCaptureComponent* SceneComp = Cast<USceneCaptureComponent>(OuterObjects[ObjectIdx]);
 		const USceneCaptureComponent* SceneCompArchetype = SceneComp ? Cast<USceneCaptureComponent>(SceneComp->GetArchetype()) : nullptr;
-		const int32 SettingIndex = SceneCompArchetype ? SceneCompArchetype->ShowFlags.FindIndexByName(*FlagName) : INDEX_NONE;
-		const bool bDefaultValue = (SettingIndex != INDEX_NONE) ? SceneCompArchetype->ShowFlags.GetSingleFlag(SettingIndex) : false;
+		const bool bDefaultValue = UE::DetailsCustomization::Private::GetShowFlagState(SceneCompArchetype, FlagName);
 
 		TArray<FEngineShowFlagsSetting>& ShowFlagSettings = *reinterpret_cast<TArray<FEngineShowFlagsSetting>*>(Data);
 		if (bNewEnabledState == bDefaultValue)
@@ -414,6 +455,67 @@ void FSceneCaptureDetails::OnShowFlagCheckStateChanged(ECheckBoxState InNewRadio
 				ShowFlagSettings.Add(NewFlagSetting);
 			}
 		}
+
+		UE::DetailsCustomization::Private::UpdateArchetypeChainShowFlagSettings(SceneComp, FlagName, bNewEnabledState);
+	}
+
+	ShowFlagSettingsProperty->NotifyPostChange(EPropertyChangeType::ValueSet);
+	ShowFlagSettingsProperty->NotifyFinishedChangingProperties();
+}
+
+bool FSceneCaptureDetails::GetShowFlagResetVisibility(FString ShowFlagName) const
+{
+	TArray<UObject*> OuterObjects;
+	ShowFlagSettingsProperty->GetOuterObjects(OuterObjects);
+
+	for (UObject* OuterObject : OuterObjects)
+	{
+		if (const USceneCaptureComponent* SceneCaptureComponent = Cast<USceneCaptureComponent>(OuterObject))
+		{
+			if (const USceneCaptureComponent* SceneCompArchetype = Cast<USceneCaptureComponent>(SceneCaptureComponent->GetArchetype()))
+			{
+				return UE::DetailsCustomization::Private::GetShowFlagState(SceneCaptureComponent, ShowFlagName) != UE::DetailsCustomization::Private::GetShowFlagState(SceneCompArchetype, ShowFlagName);
+			}
+		}
+	}
+
+	return false;
+}
+
+void FSceneCaptureDetails::OnResetShowFlag(FString ShowFlagName)
+{
+	FScopedTransaction Transaction(LOCTEXT("OnResetShowFlag", "On Reset Show Flag"));
+
+	FProperty* ShowFlagsSettingsPropertyVariable = FindFieldChecked<FProperty>(USceneCaptureComponent::StaticClass(), GET_MEMBER_NAME_CHECKED(USceneCaptureComponent, ShowFlagSettings));
+	FPropertyChangedEvent PropertyChangedEvent(ShowFlagsSettingsPropertyVariable);
+
+	ShowFlagSettingsProperty->NotifyPreChange();
+
+	TArray<void*> RawData;
+	ShowFlagSettingsProperty->AccessRawData(RawData);
+
+	TArray<UObject*> OuterObjects;
+	ShowFlagSettingsProperty->GetOuterObjects(OuterObjects);
+
+	for (int32 ObjectIdx = 0; ObjectIdx < RawData.Num(); ++ObjectIdx)
+	{
+		void* Data = RawData[ObjectIdx];
+		check(Data);
+
+		if (USceneCaptureComponent* SceneCaptureComponent = Cast<USceneCaptureComponent>(OuterObjects[ObjectIdx]))
+		{
+			// Reset to archetype default
+			TArray<FEngineShowFlagsSetting>& ShowFlagSettings = *reinterpret_cast<TArray<FEngineShowFlagsSetting>*>(Data);
+			ShowFlagSettings.RemoveAll([&ShowFlagName](const FEngineShowFlagsSetting& InSetting)
+				{
+					return InSetting.ShowFlagName.Equals(ShowFlagName);
+				});
+
+			const USceneCaptureComponent* SceneCompArchetype = Cast<USceneCaptureComponent>(SceneCaptureComponent->GetArchetype());
+			const bool bDefaultValue = UE::DetailsCustomization::Private::GetShowFlagState(SceneCompArchetype, ShowFlagName);
+
+			UE::DetailsCustomization::Private::UpdateArchetypeChainShowFlagSettings(SceneCaptureComponent, ShowFlagName, bDefaultValue);
+		}
 	}
 
 	ShowFlagSettingsProperty->NotifyPostChange(EPropertyChangeType::ValueSet);
diff --git a/Engine/Source/Editor/DetailCustomizations/Private/SceneCaptureDetails.h b/Engine/Source/Editor/DetailCustomizations/Private/SceneCaptureDetails.h
--- a/Engine/Source/Editor/DetailCustomizations/Private/SceneCaptureDetails.h
+++ b/Engine/Source/Editor/DetailCustomizations/Private/SceneCaptureDetails.h
@@ -35,4 +35,10 @@ private:
 
 	/** Show flag settings changed, so update the scene capture */
 	void OnShowFlagCheckStateChanged(ECheckBoxState InNewRadioState, FString FlagName);
+
+	/** Gets the visibility of the reset to default property detail. */
+	bool GetShowFlagResetVisibility(FString ShowFlagName) const;
+
+	/** Reset the specified show flag to its (archetype) default. */
+	void OnResetShowFlag(FString ShowFlagName);
 };
diff --git a/Engine/Source/Runtime/Engine/Private/Components/SceneCaptureComponent.cpp b/Engine/Source/Runtime/Engine/Private/Components/SceneCaptureComponent.cpp
--- a/Engine/Source/Runtime/Engine/Private/Components/SceneCaptureComponent.cpp
+++ b/Engine/Source/Runtime/Engine/Private/Components/SceneCaptureComponent.cpp
@@ -220,6 +220,8 @@ void USceneCaptureComponent::PostLoad()
 	{
 		RegisterDelegates();
 	}
+
+	UpdateShowFlags();
 }
 
 const TArray<FEngineShowFlagsSetting>& USceneCaptureComponent::GetShowFlagSettings() const
