Add custom details panel sections to PostProcessVolume, CameraComponent, and CineCamera.
#rnx

diff --git a/Engine/Source/Editor/DetailCustomizations/Private/DetailCustomizations.cpp b/Engine/Source/Editor/DetailCustomizations/Private/DetailCustomizations.cpp
--- a/Engine/Source/Editor/DetailCustomizations/Private/DetailCustomizations.cpp
+++ b/Engine/Source/Editor/DetailCustomizations/Private/DetailCustomizations.cpp
@@ -850,9 +850,38 @@ void FDetailCustomizationsModule::RegisterSectionMappings()
 		{
 			TSharedRef<FPropertySection> Section = PropertyModule.FindOrCreateSection("PostProcessVolume", "General", LOCTEXT("General", "General"));
 			Section->AddCategory("PostProcessVolumeSettings");
-			Section->AddCategory("Rendering Features");
 			Section->RemoveCategory("Brush Settings");
 		}
+		{
+			TSharedRef<FPropertySection> Section = PropertyModule.FindOrCreateSection("PostProcessVolume", "Rendering", LOCTEXT("Rendering", "Rendering"));
+			Section->AddCategory("Rendering Features");
+			Section->RemoveCategory("Path Tracing");
+		}
+		{
+			TSharedRef<FPropertySection> Section = PropertyModule.FindOrCreateSection("PostProcessVolume", "ColorGrading", LOCTEXT("ColorGrading", "Color Grading"));
+			Section->AddCategory("Color Grading");
+		}
+		{
+			TSharedRef<FPropertySection> Section = PropertyModule.FindOrCreateSection("PostProcessVolume", "Film", LOCTEXT("Film", "Film"));
+			Section->AddCategory("Film");
+			Section->AddCategory("Film Grain");
+		}
+		{
+			TSharedRef<FPropertySection> Section = PropertyModule.FindOrCreateSection("PostProcessVolume", "GlobalIllumination", LOCTEXT("GlobalIllumination", "Global Illumination"));
+			Section->AddCategory("Global Illumination");
+		}
+		{
+			TSharedRef<FPropertySection> Section = PropertyModule.FindOrCreateSection("PostProcessVolume", "Lens", LOCTEXT("Lens", "Lens"));
+			Section->AddCategory("Lens");
+		}
+		{
+			TSharedRef<FPropertySection> Section = PropertyModule.FindOrCreateSection("PostProcessVolume", "PathTracing", LOCTEXT("PathTracing", "Path Tracing"));
+			Section->AddCategory("Path Tracing");
+		}
+		{
+			TSharedRef<FPropertySection> Section = PropertyModule.FindOrCreateSection("PostProcessVolume", "Reflections", LOCTEXT("Reflections", "Reflections"));
+			Section->AddCategory("Reflections");
+		}
 	}
 
 	// RuntimeVirtualTextureComponent
@@ -868,6 +897,35 @@ void FDetailCustomizationsModule::RegisterSectionMappings()
 			Section->AddCategory("Volume Bounds");
 		}
 	}
+
+	// CameraComponent
+	{
+		{
+			TSharedRef<FPropertySection> Section = PropertyModule.FindOrCreateSection("CameraComponent", "CameraOptions", LOCTEXT("CameraOptions", "Camera Options"));
+			Section->AddCategory("Camera Options");
+		}
+		{
+			TSharedRef<FPropertySection> Section = PropertyModule.FindOrCreateSection("CameraComponent", "Camera", LOCTEXT("Camera", "Camera"));
+			Section->AddCategory("Camera Settings");
+			Section->AddCategory("Camera");
+		}
+		{
+			TSharedRef<FPropertySection> Section = PropertyModule.FindOrCreateSection("CameraComponent", "PostProcess", LOCTEXT("PostProcess", "Post Process"));
+			Section->AddCategory("Post Process");
+		}
+	}
+
+	// Cinematic Camera
+	{
+		{
+			TSharedRef<FPropertySection> Section = PropertyModule.FindOrCreateSection("CineCameraActor", "CurrentCameraSettings", LOCTEXT("CurrentCameraSettings", "Current Camera Settings"));
+			Section->AddCategory("Current Camera Settings");
+		}
+		{
+			TSharedRef<FPropertySection> Section = PropertyModule.FindOrCreateSection("CineCameraComponent", "CurrentCameraSettings", LOCTEXT("CurrentCameraSettings", "Current Camera Settings"));
+			Section->AddCategory("Current Camera Settings");
+		}
+	}
 }
 
 #undef LOCTEXT_NAMESPACE
