[ColorGrading] [DetailCustomizations] Fix cases where HSV sliders would fail to update
[RN] Fixed cases where Color Grading hue, saturation, and value sliders could fail to update when changed in another panel

diff --git a/Engine/Plugins/Editor/ColorGrading/Source/ColorGradingEditor/Private/SColorGradingColorWheel.cpp b/Engine/Plugins/Editor/ColorGrading/Source/ColorGradingEditor/Private/SColorGradingColorWheel.cpp
--- a/Engine/Plugins/Editor/ColorGrading/Source/ColorGradingEditor/Private/SColorGradingColorWheel.cpp
+++ b/Engine/Plugins/Editor/ColorGrading/Source/ColorGradingEditor/Private/SColorGradingColorWheel.cpp
@@ -4,6 +4,7 @@
 
 #include "Customizations/MathStructCustomizations.h"
 #include "DetailLayoutBuilder.h"
+#include "Editor.h"
 #include "PropertyHandle.h"
 #include "ScopedTransaction.h"
 #include "Util/ColorGradingUtil.h"
@@ -12,30 +13,26 @@
 #include "Widgets/Colors/SComplexGradient.h"
 #include "Widgets/Layout/SBox.h"
 
-#if WITH_EDITOR
-#include "Editor.h"
-#endif
-
 #define LOCTEXT_NAMESPACE "ColorGradingEditor"
 
 SColorGradingColorWheel::SColorGradingColorWheel()
 {
-#if WITH_EDITOR
 	if (GEditor)
 	{
 		GEditor->RegisterForUndo(this);
 	}
-#endif
+
+	FCoreUObjectDelegates::OnObjectPropertyChanged.AddRaw(this, &SColorGradingColorWheel::OnPropertyValueChanged);
 }
 
 SColorGradingColorWheel::~SColorGradingColorWheel()
 {
-#if WITH_EDITOR
 	if (GEditor)
 	{
 		GEditor->UnregisterForUndo(this);
 	}
-#endif
+
+	FCoreUObjectDelegates::OnObjectPropertyChanged.RemoveAll(this);
 }
 
 void SColorGradingColorWheel::Construct(const FArguments& InArgs)
@@ -133,9 +130,9 @@ void SColorGradingColorWheel::SetColorPropertyHandle(TSharedPtr<IPropertyHandle>
 	ComponentSliderDynamicMinValue.Reset();
 	ComponentSliderDynamicMaxValue.Reset();
 
-	ColorPropertyHandle = InColorPropertyHandle;
+	ColorPropertyHandle = FTrackedVector4PropertyHandle(InColorPropertyHandle);
 
-	if (ColorPropertyHandle.IsValid() && ColorPropertyHandle.Pin()->IsValidHandle())
+	if (ColorPropertyHandle.IsValidHandle())
 	{
 		ColorPropertyMetadata = GetColorPropertyMetadata();
 		ComponentSliderDynamicMinValue = ColorPropertyMetadata->MinValue;
@@ -165,7 +162,6 @@ void SColorGradingColorWheel::SetHeaderContent(const TSharedRef<SWidget>& Header
 	}
 }
 
-#if WITH_EDITOR
 void SColorGradingColorWheel::PostUndo(bool bSuccess)
 {
 	RecalculateHSVColor();
@@ -175,11 +171,35 @@ void SColorGradingColorWheel::PostRedo(bool bSuccess)
 {
 	RecalculateHSVColor();
 }
-#endif
+
+void SColorGradingColorWheel::OnPropertyValueChanged(UObject* Object, FPropertyChangedEvent& PropertyChangedEvent)
+{
+	if (ColorPropertyHandle.IsSettingValue())
+	{
+		// If setting our own value, it's already handled (or will be)
+		return;
+	}
+
+	if (TSharedPtr<const IPropertyHandle> PinnedPropertyHandle = ColorPropertyHandle.GetHandle())
+	{
+		uint32 NumChildren;
+		if (PinnedPropertyHandle->GetNumChildren(NumChildren) == FPropertyAccess::Result::Success)
+		{
+			for (uint32 ChildIndex = 0; ChildIndex < NumChildren; ++ChildIndex)
+			{
+				if (PinnedPropertyHandle->GetChildHandle(ChildIndex)->GetProperty() == PropertyChangedEvent.Property)
+				{
+					RecalculateHSVColor();
+					break;
+				}
+			}
+		}
+	}
+}
 
 TSharedRef<SWidget> SColorGradingColorWheel::SColorGradingColorWheel::CreateColorGradingPicker()
 {
-	if (ColorPropertyHandle.IsValid() && ColorPropertyHandle.Pin()->IsValidHandle())
+	if (ColorPropertyHandle.IsValidHandle())
 	{
 		return SNew(UE::ColorGrading::SColorGradingPicker)
 			.DesiredWheelSize(this, &SColorGradingColorWheel::GetMaxWheelWidth)
@@ -195,7 +215,7 @@ TSharedRef<SWidget> SColorGradingColorWheel::SColorGradingColorWheel::CreateColo
 			.ColorGradingModes(ColorPropertyMetadata->ColorGradingMode)
 			.OnColorCommitted(this, &SColorGradingColorWheel::CommitColor)
 			.OnQueryCurrentColor(this, &SColorGradingColorWheel::GetColor)
-			.AllowSpin(ColorPropertyHandle.Pin()->GetNumOuterObjects() == 1) // TODO: May want to find a way to support multiple objects
+			.AllowSpin(ColorPropertyHandle.GetHandle()->GetNumOuterObjects() == 1) // TODO: May want to find a way to support multiple objects
 			.OnBeginSliderMovement(this, &SColorGradingColorWheel::BeginUsingColorPickerSlider)
 			.OnEndSliderMovement(this, &SColorGradingColorWheel::EndUsingColorPickerSlider)
 			.OnBeginMouseCapture(this, &SColorGradingColorWheel::BeginUsingColorPickerSlider)
@@ -208,7 +228,7 @@ TSharedRef<SWidget> SColorGradingColorWheel::SColorGradingColorWheel::CreateColo
 
 TSharedRef<SWidget> SColorGradingColorWheel::CreateColorComponentSliders()
 {
-	if (ColorPropertyHandle.IsValid() && ColorPropertyHandle.Pin()->IsValidHandle())
+	if (ColorPropertyHandle.IsValidHandle())
 	{
 		TSharedRef<SVerticalBox> ColorSlidersVerticalBox = SNew(SVerticalBox);
 
@@ -260,9 +280,9 @@ SColorGradingColorWheel::FColorPropertyMetadata SColorGradingColorWheel::GetColo
 {
 	FColorPropertyMetadata PropertyMetadata;
 
-	if (ColorPropertyHandle.IsValid() && ColorPropertyHandle.Pin()->IsValidHandle())
+	if (ColorPropertyHandle.IsValidHandle())
 	{
-		FProperty* Property = ColorPropertyHandle.Pin()->GetProperty();
+		FProperty* Property = ColorPropertyHandle.GetHandle()->GetProperty();
 		const FString ColorGradingModeString = Property->GetMetaData(TEXT("ColorGradingMode")).ToLower();
 
 		if (!ColorGradingModeString.IsEmpty())
@@ -367,9 +387,9 @@ SColorGradingColorWheel::FColorPropertyMetadata SColorGradingColorWheel::GetColo
 
 bool SColorGradingColorWheel::IsPropertyEnabled() const
 {
-	if (ColorPropertyHandle.IsValid() && ColorPropertyHandle.Pin()->IsValidHandle())
+	if (ColorPropertyHandle.IsValidHandle())
 	{
-		return ColorPropertyHandle.Pin()->IsEditable();
+		return ColorPropertyHandle.GetHandle()->IsEditable();
 	}
 
 	return false;
@@ -411,41 +431,38 @@ bool SColorGradingColorWheel::ShouldUseTallLayout() const
 
 bool SColorGradingColorWheel::GetColor(FVector4& OutCurrentColor)
 {
-	if (ColorPropertyHandle.IsValid())
-	{
-		return ColorPropertyHandle.Pin()->GetValue(OutCurrentColor) == FPropertyAccess::Success;
-	}
-
-	return false;
+	return ColorPropertyHandle.GetValue(OutCurrentColor) == FPropertyAccess::Success;
 }
 
 void SColorGradingColorWheel::CommitColor(FVector4& NewValue, bool bShouldCommitValueChanges)
 {
 	FScopedTransaction Transaction(LOCTEXT("ColorWheel_TransactionName", "Color Grading Main Value"), bShouldCommitValueChanges);
 
-	if (ColorPropertyHandle.IsValid())
+	if (ColorPropertyHandle.IsValidHandle())
 	{
 		// Always perform a purely interactive change. We do this because it won't invoke reconstruction, which may cause that only the first 
 		// element gets updated due to its change causing a component reconstruction and the remaining vector element property handles updating 
 		// the trashed component.
-		ColorPropertyHandle.Pin()->SetValue(NewValue, EPropertyValueSetFlags::InteractiveChange | EPropertyValueSetFlags::NotTransactable);
+		ColorPropertyHandle.SetValue(NewValue, EPropertyValueSetFlags::InteractiveChange | EPropertyValueSetFlags::NotTransactable);
 
 		// If not purely interactive, set the value with default flags.
 		if (bShouldCommitValueChanges || !bIsUsingColorPickerSlider)
 		{
-			ColorPropertyHandle.Pin()->SetValue(NewValue, EPropertyValueSetFlags::DefaultFlags);
+			ColorPropertyHandle.SetValue(NewValue, EPropertyValueSetFlags::DefaultFlags);
 		}
 
 		TransactColorValue();
 	}
+
+	CurrentHSVColor = FLinearColor(NewValue.X, NewValue.Y, NewValue.Z).LinearRGBToHSV();
 }
 
 void SColorGradingColorWheel::TransactColorValue()
 {
-	if (ColorPropertyHandle.IsValid())
+	if (ColorPropertyHandle.IsValidHandle())
 	{
 		TArray<UObject*> OuterObjects;
-		ColorPropertyHandle.Pin()->GetOuterObjects(OuterObjects);
+		ColorPropertyHandle.GetHandle()->GetOuterObjects(OuterObjects);
 		for (UObject* Object : OuterObjects)
 		{
 			if (!Object->HasAnyFlags(RF_Transactional))
@@ -461,10 +478,10 @@ void SColorGradingColorWheel::TransactColorValue()
 
 void SColorGradingColorWheel::RecalculateHSVColor()
 {
-	if (ColorPropertyHandle.IsValid() && ColorPropertyHandle.Pin().IsValid())
+	if (ColorPropertyHandle.IsValidHandle())
 	{
 		FVector4 VectorValue;
-		if (ColorPropertyHandle.Pin()->GetValue(VectorValue) == FPropertyAccess::Success)
+		if (ColorPropertyHandle.GetValue(VectorValue) == FPropertyAccess::Success)
 		{
 			CurrentHSVColor = FLinearColor(VectorValue).LinearRGBToHSV();
 		}
@@ -475,36 +492,30 @@ void SColorGradingColorWheel::BeginUsingColorPickerSlider()
 {
 	bIsUsingColorPickerSlider = true;
 
-#if WITH_EDITOR
 	if (GEditor)
 	{
 		GEditor->BeginTransaction(LOCTEXT("ColorWheel_TransactionName", "Color Grading Main Value"));
 	}
-#endif
 }
 
 void SColorGradingColorWheel::EndUsingColorPickerSlider()
 {
 	bIsUsingColorPickerSlider = false;
 
-#if WITH_EDITOR
 	if (GEditor)
 	{
 		GEditor->EndTransaction();
 	}
-#endif
 }
 
 void SColorGradingColorWheel::BeginUsingComponentSlider(uint32 ComponentIndex)
 {
 	bIsUsingComponentSlider = true;
 
-#if WITH_EDITOR
 	if (GEditor)
 	{
 		GEditor->BeginTransaction(LOCTEXT("ColorWheel_TransactionName", "Color Grading Main Value"));
 	}
-#endif
 }
 
 void SColorGradingColorWheel::EndUsingComponentSlider(float NewValue, uint32 ComponentIndex)
@@ -512,12 +523,10 @@ void SColorGradingColorWheel::EndUsingComponentSlider(float NewValue, uint32 Com
 	bIsUsingComponentSlider = false;
 	SetComponentValue(NewValue, ComponentIndex);
 
-#if WITH_EDITOR
 	if (GEditor)
 	{
 		GEditor->EndTransaction();
 	}
-#endif
 }
 
 UE::ColorGrading::EColorGradingComponent SColorGradingColorWheel::GetComponent(uint32 ComponentIndex) const
@@ -528,10 +537,10 @@ UE::ColorGrading::EColorGradingComponent SColorGradingColorWheel::GetComponent(u
 
 TOptional<float> SColorGradingColorWheel::GetComponentValue(uint32 ComponentIndex) const
 {
-	if (ColorPropertyHandle.IsValid())
+	if (ColorPropertyHandle.IsValidHandle())
 	{
 		FVector4 ColorValue;
-		if (ColorPropertyHandle.Pin()->GetValue(ColorValue) == FPropertyAccess::Success)
+		if (ColorPropertyHandle.GetValue(ColorValue) == FPropertyAccess::Success)
 		{
 			float Value = 0.0f;
 
@@ -553,10 +562,10 @@ TOptional<float> SColorGradingColorWheel::GetComponentValue(uint32 ComponentInde
 
 void SColorGradingColorWheel::SetComponentValue(float NewValue, uint32 ComponentIndex)
 {
-	if (ColorPropertyHandle.IsValid())
+	if (ColorPropertyHandle.IsValidHandle())
 	{
 		FVector4 CurrentColorValue;
-		if (ColorPropertyHandle.Pin()->GetValue(CurrentColorValue) == FPropertyAccess::Success)
+		if (ColorPropertyHandle.GetValue(CurrentColorValue) == FPropertyAccess::Success)
 		{
 			FVector4 NewColorValue = CurrentColorValue;
 
@@ -575,7 +584,7 @@ void SColorGradingColorWheel::SetComponentValue(float NewValue, uint32 Component
 				NewColorValue = (FVector4)CurrentHSVColor.HSVToLinearRGB();
 			}
 
-			ColorPropertyHandle.Pin()->SetValue(NewColorValue, bIsUsingComponentSlider ? EPropertyValueSetFlags::InteractiveChange : EPropertyValueSetFlags::DefaultFlags);
+			ColorPropertyHandle.SetValue(NewColorValue, bIsUsingComponentSlider ? EPropertyValueSetFlags::InteractiveChange : EPropertyValueSetFlags::DefaultFlags);
 			TransactColorValue();
 		}
 	}
diff --git a/Engine/Plugins/Editor/ColorGrading/Source/ColorGradingEditor/Private/SColorGradingColorWheel.h b/Engine/Plugins/Editor/ColorGrading/Source/ColorGradingEditor/Private/SColorGradingColorWheel.h
--- a/Engine/Plugins/Editor/ColorGrading/Source/ColorGradingEditor/Private/SColorGradingColorWheel.h
+++ b/Engine/Plugins/Editor/ColorGrading/Source/ColorGradingEditor/Private/SColorGradingColorWheel.h
@@ -4,24 +4,19 @@
 
 #include "CoreMinimal.h"
 #include "Customizations/MathStructCustomizations.h"
+#include "EditorUndoClient.h"
+#include "Util/TrackedVector4PropertyHandle.h"
 #include "Widgets/DeclarativeSyntaxSupport.h"
 #include "Widgets/SCompoundWidget.h"
 #include "Widgets/ColorGrading/SColorGradingPicker.h"
 
 #include "ColorGradingPanelState.h"
 
-#if WITH_EDITOR
-#include "EditorUndoClient.h"
-#endif
-
 class IPropertyHandle;
 class SBox;
 
 /** A widget which encapsulates a color picker and numeric sliders for each color component, hooked up to a color property handle */
-class SColorGradingColorWheel : public SCompoundWidget
-#if WITH_EDITOR
-	, public FEditorUndoClient
-#endif
+class SColorGradingColorWheel : public SCompoundWidget, public FEditorUndoClient
 {
 public:
 	struct FColorPropertyMetadata
@@ -58,12 +53,13 @@ public:
 	/** Sets the widget to display as the header of the color wheel */
 	void SetHeaderContent(const TSharedRef<SWidget>& HeaderContent);
 
-#if WITH_EDITOR
 	//~ Begin FEditorUndoClient Interface
 	virtual void PostUndo(bool bSuccess) override;
 	virtual void PostRedo(bool bSuccess) override;
 	// End of FEditorUndoClient
-#endif
+
+	/** Called when any property changes */
+	void OnPropertyValueChanged(UObject* Object, FPropertyChangedEvent& PropertyChangedEvent);
 
 private:
 	TSharedRef<SWidget> CreateColorGradingPicker();
@@ -121,7 +117,7 @@ private:
 	TSharedPtr<SBox> ColorSlidersBox;
 
 	/** The property handle of the linear color property being edited */
-	TWeakPtr<IPropertyHandle> ColorPropertyHandle;
+	FTrackedVector4PropertyHandle ColorPropertyHandle;
 
 	/** The metadata of the color property */
 	TOptional<FColorPropertyMetadata> ColorPropertyMetadata;
diff --git a/Engine/Source/Editor/DetailCustomizations/Private/ColorGradingVectorCustomization.cpp b/Engine/Source/Editor/DetailCustomizations/Private/ColorGradingVectorCustomization.cpp
--- a/Engine/Source/Editor/DetailCustomizations/Private/ColorGradingVectorCustomization.cpp
+++ b/Engine/Source/Editor/DetailCustomizations/Private/ColorGradingVectorCustomization.cpp
@@ -90,16 +90,16 @@ namespace
 	}
 }
 
-FColorGradingVectorCustomizationBase::FColorGradingVectorCustomizationBase(TWeakPtr<IPropertyHandle> InColorGradingPropertyHandle, const TArray<TWeakPtr<IPropertyHandle>>& InSortedChildArray)
+FColorGradingVectorCustomizationBase::FColorGradingVectorCustomizationBase(const FTrackedVector4PropertyHandle& InColorGradingPropertyHandle, const TArray<TWeakPtr<IPropertyHandle>>& InSortedChildArray)
 	: ColorGradingPropertyHandle(InColorGradingPropertyHandle)
 	, SortedChildArray(InSortedChildArray)
 	, IsRGBMode(true)
 	, bIsUsingSlider(false)
 {
-	if (ColorGradingPropertyHandle.IsValid())
+	if (ColorGradingPropertyHandle.IsValidHandle())
 	{
 		FVector4 VectorValue;
-		ColorGradingPropertyHandle.Pin()->GetValue(VectorValue);
+		ColorGradingPropertyHandle.GetValue(VectorValue);
 		CurrentHSVColor = FLinearColor(VectorValue.X, VectorValue.Y, VectorValue.Z).LinearRGBToHSV();
 	}
 }
@@ -108,10 +108,10 @@ UE::ColorGrading::EColorGradingModes FColorGradingVectorCustomizationBase::GetCo
 {
 	UE::ColorGrading::EColorGradingModes ColorGradingMode = UE::ColorGrading::EColorGradingModes::Invalid;
 
-	if (ColorGradingPropertyHandle.IsValid())
+	if (ColorGradingPropertyHandle.IsValidHandle())
 	{
 		//Query all meta data we need
-		FProperty* Property = ColorGradingPropertyHandle.Pin()->GetProperty();
+		FProperty* Property = ColorGradingPropertyHandle.GetHandle()->GetProperty();
 		const FString& ColorGradingModeString = Property->GetMetaData(TEXT("ColorGradingMode"));
 
 		if (ColorGradingModeString.Len() > 0)
@@ -198,7 +198,7 @@ TOptional<float> FColorGradingVectorCustomizationBase::OnGetMaxValue(TOptional<f
 void FColorGradingVectorCustomizationBase::OnBeginSliderMovement()
 {
 	bIsUsingSlider = true;
-	GEditor->BeginTransaction(FText::Format(NSLOCTEXT("ColorGradingVectorCustomization", "SetPropertyValue", "Edit {0}"), ColorGradingPropertyHandle.Pin()->GetPropertyDisplayName()));
+	GEditor->BeginTransaction(FText::Format(NSLOCTEXT("ColorGradingVectorCustomization", "SetPropertyValue", "Edit {0}"), ColorGradingPropertyHandle.GetHandle()->GetPropertyDisplayName()));
 }
 
 void FColorGradingVectorCustomizationBase::OnEndSliderMovement(float NewValue, int32 ColorIndex)
@@ -220,14 +220,14 @@ UE::ColorGrading::EColorGradingComponent FColorGradingVectorCustomizationBase::O
 
 bool FColorGradingVectorCustomizationBase::GetCurrentColorGradingValue(FVector4& OutCurrentValue)
 {
-	return ColorGradingPropertyHandle.Pin()->GetValue(OutCurrentValue) == FPropertyAccess::Success;
+	return ColorGradingPropertyHandle.GetValue(OutCurrentValue) == FPropertyAccess::Success;
 }
 
 void FColorGradingVectorCustomizationBase::OnValueChanged(float NewValue, int32 ColorIndex)
 {
 	FVector4 CurrentValueVector;
-	verifySlow(ColorGradingPropertyHandle.Pin()->GetValue(CurrentValueVector) == FPropertyAccess::Success);
-	ClampValueFromMetaData(CurrentValueVector, ColorGradingPropertyHandle.Pin()->GetProperty());
+	verifySlow(ColorGradingPropertyHandle.GetValue(CurrentValueVector) == FPropertyAccess::Success);
+	ClampValueFromMetaData(CurrentValueVector, ColorGradingPropertyHandle.GetHandle()->GetProperty());
 	FVector4 NewValueVector = CurrentValueVector;
 
 	if (IsRGBMode)
@@ -255,9 +255,9 @@ void FColorGradingVectorCustomizationBase::OnValueChanged(float NewValue, int32
 		OnCurrentHSVColorChanged.Broadcast(CurrentHSVColor, true);
 	}
 
-	if (ColorGradingPropertyHandle.IsValid())
+	if (ColorGradingPropertyHandle.IsValidHandle())
 	{
-		ColorGradingPropertyHandle.Pin()->SetValue(NewValueVector, bIsUsingSlider ? EPropertyValueSetFlags::InteractiveChange : EPropertyValueSetFlags::DefaultFlags);
+		ColorGradingPropertyHandle.SetValue(NewValueVector, bIsUsingSlider ? EPropertyValueSetFlags::InteractiveChange : EPropertyValueSetFlags::DefaultFlags);
 	}
 }
 
@@ -265,7 +265,7 @@ TOptional<float> FColorGradingVectorCustomizationBase::OnSliderGetValue(int32 Co
 {
 	FVector4 ValueVector;
 	
-	if (ColorGradingPropertyHandle.Pin()->GetValue(ValueVector) == FPropertyAccess::Success)
+	if (ColorGradingPropertyHandle.GetValue(ValueVector) == FPropertyAccess::Success)
 	{
 		float Value = 0.0f;
 
@@ -296,10 +296,10 @@ void FColorGradingVectorCustomizationBase::OnCurrentHSVColorChangedDelegate(FLin
 
 void FColorGradingVectorCustomizationBase::PostUndo(bool bSuccess)
 {
-	if (ColorGradingPropertyHandle.IsValid())
+	if (ColorGradingPropertyHandle.IsValidHandle())
 	{
 		FVector4 CurrentValueVector;
-		if (ColorGradingPropertyHandle.Pin()->GetValue(CurrentValueVector) == FPropertyAccess::Success)
+		if (ColorGradingPropertyHandle.GetValue(CurrentValueVector) == FPropertyAccess::Success)
 		{
 			CurrentHSVColor = FLinearColor(CurrentValueVector.X, CurrentValueVector.Y, CurrentValueVector.Z).LinearRGBToHSV();
 			OnCurrentHSVColorChanged.Broadcast(CurrentHSVColor, true);
@@ -391,7 +391,7 @@ TSharedRef<UE::ColorGrading::SColorGradingComponentViewer> FColorGradingVectorCu
 		.OnEndSliderMovement(this, &FColorGradingVectorCustomizationBase::OnEndSliderMovement, ColorIndex)
 		.OnQueryCurrentColor(this, &FColorGradingVectorCustomizationBase::GetCurrentColorGradingValue)
 		// Only allow spin on handles with one object.  Otherwise it is not clear what value to spin
-		.AllowSpin(ColorGradingPropertyHandle.Pin()->GetNumOuterObjects() == 1)
+		.AllowSpin(ColorGradingPropertyHandle.GetHandle()->GetNumOuterObjects() == 1)
 		.ShiftMultiplier(ShiftMultiplier)
 		.CtrlMultiplier(CtrlMultiplier)
 		.SupportDynamicSliderMaxValue(this, &FColorGradingVectorCustomizationBase::GetSupportDynamicSliderMaxValue, SupportDynamicSliderMaxValue, ColorIndex)
@@ -424,7 +424,7 @@ FColorGradingVectorCustomization::~FColorGradingVectorCustomization()
 void FColorGradingVectorCustomization::MakeHeaderRow(FDetailWidgetRow& Row, TSharedRef<FVector4StructCustomization> InVector4Customization)
 {
 	TSharedPtr<SHorizontalBox> ContentHorizontalBox = SNew(SHorizontalBox)
-		.IsEnabled(InVector4Customization, &FMathStructCustomization::IsValueEnabled, ColorGradingPropertyHandle);
+		.IsEnabled(InVector4Customization, &FMathStructCustomization::IsValueEnabled, ColorGradingPropertyHandle.GetHandle().ToWeakPtr());
 
 	Row.NameContent()
 		[
@@ -433,7 +433,7 @@ void FColorGradingVectorCustomization::MakeHeaderRow(FDetailWidgetRow& Row, TSha
 			.AutoWidth()
 			.HAlign(HAlign_Left)
 			[
-				ColorGradingPropertyHandle.Pin()->CreatePropertyNameWidget()
+				ColorGradingPropertyHandle.GetHandle()->CreatePropertyNameWidget()
 			]
 		];
 
@@ -451,7 +451,7 @@ void FColorGradingVectorCustomization::MakeHeaderRow(FDetailWidgetRow& Row, TSha
 
 		// Make a widget for each property.  The vector component properties  will be displayed in the header
 
-		TSharedRef<IPropertyHandle> ColorGradingPropertyHandleRef = ColorGradingPropertyHandle.Pin().ToSharedRef();
+		TSharedRef<IPropertyHandle> ColorGradingPropertyHandleRef = ColorGradingPropertyHandle.GetHandle().ToSharedRef();
 		FMathStructCustomization::FNumericMetadata<float> Metadata;
 		FMathStructCustomization::ExtractNumericMetadata(ColorGradingPropertyHandleRef, Metadata);
 
@@ -533,14 +533,14 @@ void FColorGradingVectorCustomization::MakeHeaderRow(FDetailWidgetRow& Row, TSha
 EVisibility FColorGradingVectorCustomization::GetMultipleValuesTextVisibility() const
 {
 	FVector4 VectorValue;
-	return (ColorGradingPropertyHandle.Pin()->GetValue(VectorValue) == FPropertyAccess::MultipleValues) ? EVisibility::Visible : EVisibility::Collapsed;
+	return (ColorGradingPropertyHandle.GetValue(VectorValue) == FPropertyAccess::MultipleValues) ? EVisibility::Visible : EVisibility::Collapsed;
 }
 
 FLinearColor FColorGradingVectorCustomization::OnGetColorForHeaderColorBlock() const
 {
 	FLinearColor ColorValue(0.0f, 0.0f, 0.0f);
 	FVector4 VectorValue;
-	if (ColorGradingPropertyHandle.Pin()->GetValue(VectorValue) == FPropertyAccess::Success)
+	if (ColorGradingPropertyHandle.GetValue(VectorValue) == FPropertyAccess::Success)
 	{
 		ColorValue.R = VectorValue.X * VectorValue.W;
 		ColorValue.G = VectorValue.Y * VectorValue.W;
@@ -569,7 +569,7 @@ void FColorGradingVectorCustomization::CustomizeChildren(IDetailChildrenBuilder&
 
 	if (ParentGroup != nullptr)
 	{
-		TSharedPtr<IDetailPropertyRow> PropertyRow = ParentGroup->FindPropertyRow(ColorGradingPropertyHandle.Pin().ToSharedRef());
+		TSharedPtr<IDetailPropertyRow> PropertyRow = ParentGroup->FindPropertyRow(ColorGradingPropertyHandle.GetHandle().ToSharedRef());
 		verifySlow(PropertyRow.IsValid());
 
 		PropertyRow->OverrideResetToDefault(FResetToDefaultOverride::Create(FIsResetToDefaultVisible::CreateSP(CustomColorGradingBuilder.Get(), &FColorGradingCustomBuilder::CanResetToDefault),
@@ -581,7 +581,7 @@ void FColorGradingVectorCustomization::CustomizeChildren(IDetailChildrenBuilder&
 //////////////////////////////////////////////////////////////////////////
 // Color Gradient custom builder implementation
 
-FColorGradingCustomBuilder::FColorGradingCustomBuilder(TWeakPtr<IPropertyHandle> InColorGradingPropertyHandle, const TArray<TWeakPtr<IPropertyHandle>>& InSortedChildArray, 
+FColorGradingCustomBuilder::FColorGradingCustomBuilder(const FTrackedVector4PropertyHandle& InColorGradingPropertyHandle, const TArray<TWeakPtr<IPropertyHandle>>& InSortedChildArray,
 													   TSharedRef<FColorGradingVectorCustomization> InColorGradingCustomization, IDetailGroup* InParentGroup)
 	: FColorGradingVectorCustomizationBase(InColorGradingPropertyHandle, InSortedChildArray)
 	, ColorGradingCustomization(InColorGradingCustomization)
@@ -631,12 +631,14 @@ FColorGradingCustomBuilder::~FColorGradingCustomBuilder()
 
 	// Deregister for Undo callbacks
 	GEditor->UnregisterForUndo(this);
+
+	FCoreUObjectDelegates::OnObjectPropertyChanged.RemoveAll(this);
 }
 
 void FColorGradingCustomBuilder::OnDetailGroupReset()
 {
 	FVector4 CurrentValueVector;
-	verifySlow(ColorGradingPropertyHandle.Pin()->GetValue(CurrentValueVector) == FPropertyAccess::Success);
+	verifySlow(ColorGradingPropertyHandle.GetValue(CurrentValueVector) == FPropertyAccess::Success);
 	CurrentHSVColor = FLinearColor(CurrentValueVector.X, CurrentValueVector.Y, CurrentValueVector.Z).LinearRGBToHSV();
 
 	OnCurrentHSVColorChanged.Broadcast(CurrentHSVColor, true);
@@ -657,7 +659,7 @@ void FColorGradingCustomBuilder::ResetToDefault(TSharedPtr<IPropertyHandle> Prop
 	PropertyHandle->ResetToDefault();
 
 	FVector4 CurrentValueVector;
-	verifySlow(ColorGradingPropertyHandle.Pin()->GetValue(CurrentValueVector) == FPropertyAccess::Success);
+	verifySlow(ColorGradingPropertyHandle.GetValue(CurrentValueVector) == FPropertyAccess::Success);
 	CurrentHSVColor = FLinearColor(CurrentValueVector.X, CurrentValueVector.Y, CurrentValueVector.Z).LinearRGBToHSV();
 
 	OnCurrentHSVColorChanged.Broadcast(CurrentHSVColor, true);
@@ -686,7 +688,7 @@ void FColorGradingCustomBuilder::GenerateHeaderRowContent(FDetailWidgetRow& Node
 {
 	// Make a widget for each property.  The vector component properties  will be displayed in the header
 
-	TSharedRef<IPropertyHandle> ColorGradingPropertyHandleRef = ColorGradingPropertyHandle.Pin().ToSharedRef();
+	TSharedRef<IPropertyHandle> ColorGradingPropertyHandleRef = ColorGradingPropertyHandle.GetHandle().ToSharedRef();
 	FMathStructCustomization::FNumericMetadata<float> Metadata;
 	FMathStructCustomization::ExtractNumericMetadata(ColorGradingPropertyHandleRef, Metadata);
 	
@@ -723,7 +725,7 @@ void FColorGradingCustomBuilder::GenerateHeaderRowContent(FDetailWidgetRow& Node
 				.ColorGradingModes(ColorGradingMode)
 				.OnColorCommitted(this, &FColorGradingCustomBuilder::OnColorGradingPickerChanged)
 				.OnQueryCurrentColor(this, &FColorGradingCustomBuilder::GetCurrentColorGradingValue)
-				.AllowSpin(ColorGradingPropertyHandle.Pin()->GetNumOuterObjects() == 1)
+				.AllowSpin(ColorGradingPropertyHandle.GetHandle()->GetNumOuterObjects() == 1)
 				.OnBeginSliderMovement(this, &FColorGradingCustomBuilder::OnBeginMainValueSliderMovement)
 				.OnEndSliderMovement(this, &FColorGradingCustomBuilder::OnEndMainValueSliderMovement)
 				.OnBeginMouseCapture(this, &FColorGradingCustomBuilder::OnBeginMouseCapture)
@@ -858,11 +860,13 @@ void FColorGradingCustomBuilder::GenerateHeaderRowContent(FDetailWidgetRow& Node
 	ParentGroupName.ReplaceInline(TEXT(" "), TEXT("_"));
 	ParentGroupName.ReplaceInline(TEXT("|"), TEXT("_"));
 
-	GConfig->GetBool(TEXT("ColorGrading"), *FString::Printf(TEXT("%s_%s_IsRGB"), *ParentGroupName, *ColorGradingPropertyHandle.Pin()->GetPropertyDisplayName().ToString()), RGBMode, GEditorPerProjectIni);
+	GConfig->GetBool(TEXT("ColorGrading"), *FString::Printf(TEXT("%s_%s_IsRGB"), *ParentGroupName, *ColorGradingPropertyHandle.GetHandle()->GetPropertyDisplayName().ToString()), RGBMode, GEditorPerProjectIni);
 	OnChangeColorModeClicked(RGBMode ? UE::ColorGrading::EColorGradingColorDisplayMode::RGB : UE::ColorGrading::EColorGradingColorDisplayMode::HSV);
 
 	// Register to update when an undo/redo operation has been called to update our list of actors
 	GEditor->RegisterForUndo(this);
+
+	FCoreUObjectDelegates::OnObjectPropertyChanged.AddRaw(this, &FColorGradingCustomBuilder::OnPropertyValueChanged);
 }
 
 FText FColorGradingCustomBuilder::OnChangeColorModeText(UE::ColorGrading::EColorGradingColorDisplayMode ModeType) const
@@ -904,7 +908,7 @@ EVisibility FColorGradingCustomBuilder::OnGetGradientVisibility() const
 void FColorGradingCustomBuilder::OnChangeColorModeClicked(UE::ColorGrading::EColorGradingColorDisplayMode ModeType)
 {
 	FVector4 CurrentValueVector;
-	if (ColorGradingPropertyHandle.Pin()->GetValue(CurrentValueVector) != FPropertyAccess::Success)
+	if (ColorGradingPropertyHandle.GetValue(CurrentValueVector) != FPropertyAccess::Success)
 	{
 		return;
 	}
@@ -921,7 +925,7 @@ void FColorGradingCustomBuilder::OnChangeColorModeClicked(UE::ColorGrading::ECol
 		ParentGroupName.ReplaceInline(TEXT(" "), TEXT("_"));
 		ParentGroupName.ReplaceInline(TEXT("|"), TEXT("_"));
 
-		GConfig->SetBool(TEXT("ColorGrading"), *FString::Printf(TEXT("%s_%s_IsRGB"), *ParentGroupName, *ColorGradingPropertyHandle.Pin()->GetPropertyDisplayName().ToString()), IsRGBMode, GEditorPerProjectIni);
+		GConfig->SetBool(TEXT("ColorGrading"), *FString::Printf(TEXT("%s_%s_IsRGB"), *ParentGroupName, *ColorGradingPropertyHandle.GetHandle()->GetPropertyDisplayName().ToString()), IsRGBMode, GEditorPerProjectIni);
 
 		CurrentHSVColor = FLinearColor(CurrentValueVector.X, CurrentValueVector.Y, CurrentValueVector.Z).LinearRGBToHSV();
 
@@ -940,17 +944,17 @@ void FColorGradingCustomBuilder::OnColorGradingPickerChanged(FVector4& NewValue,
 {
 	FScopedTransaction Transaction(LOCTEXT("ColorGradingMainValue", "Color Grading Main Value"), ShouldCommitValueChanges);
 
-	if (ColorGradingPropertyHandle.IsValid())
+	if (ColorGradingPropertyHandle.IsValidHandle())
 	{
 		// Always perform a purely interactive change. We do this because it won't invoke reconstruction, which may cause that only the first 
 		// element gets updated due to its change causing a component reconstruction and the remaining vector element property handles updating 
 		// the trashed component.
-		ColorGradingPropertyHandle.Pin()->SetValue(NewValue, EPropertyValueSetFlags::InteractiveChange | EPropertyValueSetFlags::NotTransactable);
+		ColorGradingPropertyHandle.SetValue(NewValue, EPropertyValueSetFlags::InteractiveChange | EPropertyValueSetFlags::NotTransactable);
 
 		// If not purely interactive, set the value with default flags.
 		if (ShouldCommitValueChanges || !bIsUsingSlider)
 		{
-			ColorGradingPropertyHandle.Pin()->SetValue(NewValue, EPropertyValueSetFlags::DefaultFlags);
+			ColorGradingPropertyHandle.SetValue(NewValue, EPropertyValueSetFlags::DefaultFlags);
 		}
 	}
 
@@ -989,4 +993,33 @@ void FColorGradingCustomBuilder::OnEndMouseCapture()
 	GEditor->EndTransaction();
 }
 
+void FColorGradingCustomBuilder::OnPropertyValueChanged(UObject* Object, FPropertyChangedEvent& PropertyChangedEvent)
+{
+	if (ColorGradingPropertyHandle.IsSettingValue())
+	{
+		// If setting our own value, it's already handled (or will be)
+		return;
+	}
+
+	if (TSharedPtr<IPropertyHandle> PinnedPropertyHandle = ColorGradingPropertyHandle.GetHandle())
+	{
+		uint32 NumChildren;
+		if (PinnedPropertyHandle->GetNumChildren(NumChildren) == FPropertyAccess::Result::Success)
+		{
+			for (uint32 ChildIndex = 0; ChildIndex < NumChildren; ++ChildIndex)
+			{
+				if (PinnedPropertyHandle->GetChildHandle(ChildIndex)->GetProperty() == PropertyChangedEvent.Property)
+				{
+					FVector4 CurrentValueVector;
+					if (PinnedPropertyHandle->GetValue(CurrentValueVector) == FPropertyAccess::Success)
+					{
+						CurrentHSVColor = FLinearColor(CurrentValueVector.X, CurrentValueVector.Y, CurrentValueVector.Z).LinearRGBToHSV();
+					}
+					break;
+				}
+			}
+		}
+	}
+}
+
 #undef LOCTEXT_NAMESPACE
\ No newline at end of file
diff --git a/Engine/Source/Editor/DetailCustomizations/Private/ColorGradingVectorCustomization.h b/Engine/Source/Editor/DetailCustomizations/Private/ColorGradingVectorCustomization.h
--- a/Engine/Source/Editor/DetailCustomizations/Private/ColorGradingVectorCustomization.h
+++ b/Engine/Source/Editor/DetailCustomizations/Private/ColorGradingVectorCustomization.h
@@ -17,6 +17,7 @@
 #include "Templates/SharedPointer.h"
 #include "UObject/NameTypes.h"
 #include "UObject/UnrealNames.h"
+#include "Util/TrackedVector4PropertyHandle.h"
 
 class FDetailWidgetRow;
 class FVector4StructCustomization;
@@ -50,7 +51,7 @@ public:
 	/** Notification the current HSV color was changed */
 	DECLARE_MULTICAST_DELEGATE_TwoParams(FOnCurrentHSVColorChanged, FLinearColor, bool);
 
-	FColorGradingVectorCustomizationBase(TWeakPtr<IPropertyHandle> InColorGradingPropertyHandle, const TArray<TWeakPtr<IPropertyHandle>>& InSortedChildArray);
+	FColorGradingVectorCustomizationBase(const FTrackedVector4PropertyHandle& InColorGradingPropertyHandle, const TArray<TWeakPtr<IPropertyHandle>>& InSortedChildArray);
 
 	/** Return max/min slider value changed delegate (only apply if SupportDynamicSliderMaxValue or SupportDynamicSliderMinValue are true) */
 	FOnNumericEntryBoxDynamicSliderMinMaxValueChanged& GetOnNumericEntryBoxDynamicSliderMaxValueChangedDelegate() { return OnNumericEntryBoxDynamicSliderMaxValueChanged; }
@@ -117,7 +118,7 @@ protected:
 	TArray<TWeakPtr<UE::ColorGrading::SColorGradingComponentViewer>> ComponentViewers;
 	
 	/** The color grading property we're editing */
-	TWeakPtr<IPropertyHandle> ColorGradingPropertyHandle;
+	FTrackedVector4PropertyHandle ColorGradingPropertyHandle;
 	
 	/** Property for each color value (RGBY) */
 	TArray<TWeakPtr<IPropertyHandle>> SortedChildArray;
@@ -172,7 +173,7 @@ public:
 	/** Notification when we change color mode (RGB <-> HSV) */
 	DECLARE_MULTICAST_DELEGATE_OneParam(FOnColorModeChanged, bool);
 
-	FColorGradingCustomBuilder(TWeakPtr<IPropertyHandle> InColorGradingPropertyHandle, const TArray<TWeakPtr<IPropertyHandle>>& InSortedChildArray, 
+	FColorGradingCustomBuilder(const FTrackedVector4PropertyHandle& InColorGradingPropertyHandle, const TArray<TWeakPtr<IPropertyHandle>>& InSortedChildArray, 
 							   TSharedRef<FColorGradingVectorCustomization> InColorGradingCustomization, IDetailGroup* InParentGroup);
 	virtual ~FColorGradingCustomBuilder();
 
@@ -203,6 +204,9 @@ private:
 	void OnBeginMouseCapture();
 	void OnEndMouseCapture();
 
+	/** Called when any property changes */
+	void OnPropertyValueChanged(UObject* Object, FPropertyChangedEvent& PropertyChangedEvent);
+
 	/** Callback when user click the Group reset button */
 	void OnDetailGroupReset();
 
diff --git a/Engine/Source/Runtime/AdvancedWidgets/Private/Util/TrackedVector4PropertyHandle.cpp b/Engine/Source/Runtime/AdvancedWidgets/Private/Util/TrackedVector4PropertyHandle.cpp
new file mode 100644
--- /dev/null
+++ b/Engine/Source/Runtime/AdvancedWidgets/Private/Util/TrackedVector4PropertyHandle.cpp
@@ -0,0 +1,61 @@
+// Copyright Epic Games, Inc. All Rights Reserved.
+
+#if WITH_EDITOR
+
+#include "Util/TrackedVector4PropertyHandle.h"
+
+FTrackedVector4PropertyHandle::FTrackedVector4PropertyHandle()
+	: Handle(nullptr)
+{
+}
+
+FTrackedVector4PropertyHandle::FTrackedVector4PropertyHandle(TWeakPtr<IPropertyHandle> InHandle)
+	: Handle(InHandle)
+{
+}
+
+TSharedPtr<IPropertyHandle> FTrackedVector4PropertyHandle::GetHandle() const
+{
+	return Handle.Pin();
+}
+
+FPropertyAccess::Result FTrackedVector4PropertyHandle::SetValue(const FVector4& InValue, EPropertyValueSetFlags::Type Flags)
+{
+	FPropertyAccess::Result Result = FPropertyAccess::Fail;
+
+	if (TSharedPtr<IPropertyHandle> PinnedHandle = Handle.Pin())
+	{
+		bIsSettingValue = true;
+		Result = PinnedHandle->SetValue(InValue, Flags);
+		bIsSettingValue = false;
+	}
+
+	return Result;
+}
+
+FPropertyAccess::Result FTrackedVector4PropertyHandle::GetValue(FVector4& OutValue) const
+{
+	if (TSharedPtr<IPropertyHandle> PinnedHandle = Handle.Pin())
+	{
+		return PinnedHandle->GetValue(OutValue);
+	}
+
+	return FPropertyAccess::Fail;
+}
+
+bool FTrackedVector4PropertyHandle::IsSettingValue() const
+{
+	return bIsSettingValue;
+}
+
+bool FTrackedVector4PropertyHandle::IsValidHandle() const
+{
+	if (TSharedPtr<IPropertyHandle> PinnedHandle = Handle.Pin())
+	{
+		return PinnedHandle->IsValidHandle();
+	}
+
+	return false;
+}
+
+#endif
\ No newline at end of file
diff --git a/Engine/Source/Runtime/AdvancedWidgets/Public/Util/TrackedVector4PropertyHandle.h b/Engine/Source/Runtime/AdvancedWidgets/Public/Util/TrackedVector4PropertyHandle.h
new file mode 100644
--- /dev/null
+++ b/Engine/Source/Runtime/AdvancedWidgets/Public/Util/TrackedVector4PropertyHandle.h
@@ -0,0 +1,43 @@
+// Copyright Epic Games, Inc. All Rights Reserved.
+
+#pragma once
+
+#if WITH_EDITOR
+
+#include "PropertyHandle.h"
+
+/** Convenience wrapper for an IPropertyHandle that lets us track when we're using this handle to set the property's vector value. */
+class ADVANCEDWIDGETS_API FTrackedVector4PropertyHandle
+{
+public:
+	FTrackedVector4PropertyHandle();
+	FTrackedVector4PropertyHandle(TWeakPtr<IPropertyHandle> InHandle);
+
+	/** Get the underlying property handle. */
+	TSharedPtr<IPropertyHandle> GetHandle() const;
+
+	/** Set the property's vector value */
+	FPropertyAccess::Result SetValue(const FVector4& InValue, EPropertyValueSetFlags::Type Flags = EPropertyValueSetFlags::DefaultFlags);
+
+	/** Get the property's vector value */
+	FPropertyAccess::Result GetValue(FVector4& OutValue) const;
+
+	/**
+	 * @return Whether we're currently inside a call to SetValue
+	 */
+	bool IsSettingValue() const;
+
+	/**
+	 * @return Whether or not the handle points to a valid property node. This can be true but GetProperty may still return null
+	 */
+	bool IsValidHandle() const;
+
+private:
+	/** The underlying handle to the property */
+	TWeakPtr<IPropertyHandle> Handle;
+
+	/** Whether we're currently changing the property using SetValue */
+	bool bIsSettingValue = false;
+};
+
+#endif
\ No newline at end of file
