Chaos::Expose acceleration mode in angular and linear joint drive.

diff --git a/Engine/Source/Editor/DetailCustomizations/Private/PhysicsConstraintComponentDetails.cpp b/Engine/Source/Editor/DetailCustomizations/Private/PhysicsConstraintComponentDetails.cpp
--- a/Engine/Source/Editor/DetailCustomizations/Private/PhysicsConstraintComponentDetails.cpp
+++ b/Engine/Source/Editor/DetailCustomizations/Private/PhysicsConstraintComponentDetails.cpp
@@ -1611,6 +1611,9 @@ void FPhysicsConstraintComponentDetails::AddLinearDrive(IDetailLayoutBuilder& De
 
 	TSharedPtr<IPropertyHandle> LinearDriveProperty = ProfilePropertiesProperty->GetChildHandle(GET_MEMBER_NAME_CHECKED(FConstraintProfileProperties, LinearDrive));
 
+	TSharedPtr<IPropertyHandle> AccelerationModeProperty = ProfilePropertiesProperty->GetChildHandle(GET_MEMBER_NAME_CHECKED(FLinearDriveConstraint, bAccelerationMode));
+	LinearMotorCat.AddProperty(AccelerationModeProperty);
+
 	IDetailGroup& PositionGroup = LinearMotorCat.AddGroup("Linear Position Drive", LOCTEXT("LinearPositionDrive", "Linear Position Drive"), false, true);
 
 	TSharedRef<IPropertyHandle> LinearPositionTargetProperty = LinearDriveProperty->GetChildHandle(GET_MEMBER_NAME_CHECKED(FLinearDriveConstraint, PositionTarget)).ToSharedRef();
@@ -1664,7 +1667,7 @@ void FPhysicsConstraintComponentDetails::AddLinearDrive(IDetailLayoutBuilder& De
 	[
 		SNew(STextBlock)
 		.Text(LOCTEXT("Strength", "Strength"))
-		.Font(IDetailLayoutBuilder::GetDetailFont())
+		.Font(IDetailLayoutBuilder::GetDetailFont()).ToolTipText(StiffnessXProperty->GetToolTipText())
 	]
 	.ValueContent()
 	[
@@ -1721,8 +1724,8 @@ void FPhysicsConstraintComponentDetails::AddLinearDrive(IDetailLayoutBuilder& De
 	.NameContent()
 	[
 		SNew(STextBlock)
-		.Text(LOCTEXT("Strength", "Strength"))
-		.Font(IDetailLayoutBuilder::GetDetailFont())
+		.Text(LOCTEXT("Damping", "Damping"))
+		.Font(IDetailLayoutBuilder::GetDetailFont()).ToolTipText(XDampingProperty->GetToolTipText())
 	]
 	.ValueContent()
 	[
@@ -1752,6 +1755,7 @@ void FPhysicsConstraintComponentDetails::AddAngularDrive(IDetailLayoutBuilder& D
 	TSharedPtr<IPropertyHandle> AngularDriveProperty = ProfilePropertiesProperty->GetChildHandle(GET_MEMBER_NAME_CHECKED(FConstraintProfileProperties, AngularDrive));
 	TSharedPtr<IPropertyHandle> AngularDriveModeProperty = AngularDriveProperty->GetChildHandle(GET_MEMBER_NAME_CHECKED(FAngularDriveConstraint, AngularDriveMode));
 
+	TSharedPtr<IPropertyHandle> AccelerationModeProperty = AngularDriveProperty->GetChildHandle(GET_MEMBER_NAME_CHECKED(FAngularDriveConstraint, bAccelerationMode));
 	TSharedPtr<IPropertyHandle> SlerpDriveProperty = AngularDriveProperty->GetChildHandle(GET_MEMBER_NAME_CHECKED(FAngularDriveConstraint, SlerpDrive));
 	TSharedPtr<IPropertyHandle> SwingDriveProperty = AngularDriveProperty->GetChildHandle(GET_MEMBER_NAME_CHECKED(FAngularDriveConstraint, SwingDrive));
 	TSharedPtr<IPropertyHandle> TwistDriveProperty = AngularDriveProperty->GetChildHandle(GET_MEMBER_NAME_CHECKED(FAngularDriveConstraint, TwistDrive));
@@ -1816,6 +1820,7 @@ void FPhysicsConstraintComponentDetails::AddAngularDrive(IDetailLayoutBuilder& D
 	};
 
 	AngularMotorCat.AddProperty(AngularDriveModeProperty);
+	AngularMotorCat.AddProperty(AccelerationModeProperty);
 
 	IDetailGroup& OrientationGroup = AngularMotorCat.AddGroup("Orientation Drive", LOCTEXT("OrientrationDrive", "Orientation Drive"), false, true);
 	OrientationGroup.HeaderProperty(AngularDriveProperty->GetChildHandle(GET_MEMBER_NAME_CHECKED(FAngularDriveConstraint, OrientationTarget)).ToSharedRef()).DisplayName(LOCTEXT("TargetOrientation", "Target Orientation")).IsEnabled(TAttribute<bool>::Create(TAttribute<bool>::FGetter::CreateLambda(OrientationEnabled)));
@@ -1895,7 +1900,7 @@ void FPhysicsConstraintComponentDetails::AddAngularDrive(IDetailLayoutBuilder& D
 	[
 		SNew(STextBlock)
 		.Text(LOCTEXT("Strength", "Strength"))
-		.Font(IDetailLayoutBuilder::GetDetailFont())
+		.Font(IDetailLayoutBuilder::GetDetailFont()).ToolTipText(StiffnessSlerpProperty->GetToolTipText())
 	]
 	.ValueContent()
 	[
@@ -1962,8 +1967,8 @@ void FPhysicsConstraintComponentDetails::AddAngularDrive(IDetailLayoutBuilder& D
 	.NameContent()
 	[
 		SNew(STextBlock)
-		.Text(LOCTEXT("Strength", "Strength"))
-		.Font(IDetailLayoutBuilder::GetDetailFont())
+		.Text(LOCTEXT("Damping", "Damping"))
+		.Font(IDetailLayoutBuilder::GetDetailFont()).ToolTipText(DampingSlerpProperty->GetToolTipText())
 	]
 	.ValueContent()
 	[
diff --git a/Engine/Source/Runtime/Engine/Classes/PhysicsEngine/ConstraintDrives.h b/Engine/Source/Runtime/Engine/Classes/PhysicsEngine/ConstraintDrives.h
--- a/Engine/Source/Runtime/Engine/Classes/PhysicsEngine/ConstraintDrives.h
+++ b/Engine/Source/Runtime/Engine/Classes/PhysicsEngine/ConstraintDrives.h
@@ -24,11 +24,11 @@ struct FConstraintDrive
 {
 	GENERATED_BODY()
 
-	/** The spring strength of the drive. Force proportional to the position error. */
+	/** The spring strength of the drive. Force proportional to the position error. In force mode the unit is kg * cm / s^-2. In acceleration mode the unit is cm / s^-2. */
 	UPROPERTY(EditAnywhere, AdvancedDisplay, Category = Constraint, meta = (ClampMin = "0.0"))
 	float Stiffness;
 
-	/** The damping strength of the drive. Force proportional to the velocity error. */
+	/** The damping strength of the drive. Force proportional to the velocity error. In force mode the unit is kg * cm / s^-1. In acceleration mode the unit is cm / s^-1.*/
 	UPROPERTY(EditAnywhere, AdvancedDisplay, Category = Constraint, meta = (ClampMin = "0.0"))
 	float Damping;
 
@@ -83,6 +83,7 @@ struct FLinearDriveConstraint
 	UPROPERTY(EditAnywhere, Category = LinearMotor)
 	FConstraintDrive ZDrive;
 	
+	/** Whether the specified strength is acceleration or force. (default: acceleration) */
 	UPROPERTY(EditAnywhere, Category = AngularMotor)
 	bool bAccelerationMode = true;
 
@@ -150,6 +151,7 @@ struct FAngularDriveConstraint
 	UPROPERTY(EditAnywhere, Category = AngularMotor)
 	TEnumAsByte<enum EAngularDriveMode::Type> AngularDriveMode;
 
+	/** Whether to use acceleration mode for angular drive. The other option is force mode. (default:true) */
 	UPROPERTY(EditAnywhere, Category = AngularMotor)
 	bool bAccelerationMode = true;
 
