Add Audio section filter for actor and component detail panels

diff --git a/Engine/Source/Editor/DetailCustomizations/Private/DetailCustomizations.cpp b/Engine/Source/Editor/DetailCustomizations/Private/DetailCustomizations.cpp
--- a/Engine/Source/Editor/DetailCustomizations/Private/DetailCustomizations.cpp
+++ b/Engine/Source/Editor/DetailCustomizations/Private/DetailCustomizations.cpp
@@ -541,6 +541,18 @@ void FDetailCustomizationsModule::RegisterSectionMappings()
 			Section->AddCategory("Data Layers");
 			Section->AddCategory("HLOD");
 		}
+
+		{
+			TSharedRef<FPropertySection> Section = PropertyModule.FindOrCreateSection("Actor", "Audio", LOCTEXT("Audio", "Audio"));
+			Section->AddCategory("Audio");
+			Section->AddCategory("Sound");
+			Section->AddCategory("Sounds");
+			Section->AddCategory("MetaSound");
+			Section->AddCategory("MetaSounds");
+			Section->AddCategory("SFX");
+			Section->AddCategory("Music");
+			Section->AddCategory("VO");
+		}
 	}
 
 	// Pawn
@@ -560,6 +572,18 @@ void FDetailCustomizationsModule::RegisterSectionMappings()
 			Section->AddCategory("Cooking");
 			Section->AddCategory("Tags");
 		}
+
+		{
+			TSharedRef<FPropertySection> Section = PropertyModule.FindOrCreateSection("ActorComponent", "Audio", LOCTEXT("Audio", "Audio"));
+			Section->AddCategory("Audio");
+			Section->AddCategory("Sound");
+			Section->AddCategory("Sounds");
+			Section->AddCategory("MetaSound");
+			Section->AddCategory("MetaSounds");
+			Section->AddCategory("SFX");
+			Section->AddCategory("Music");
+			Section->AddCategory("VO");
+		}
 	}
 
 	// SceneComponent
