Editor: Fix the detail panel dropdown for FComponentReference sometimes showing incorrect values while null.

diff --git a/Engine/Source/Editor/DetailCustomizations/Private/ComponentReferenceCustomization.cpp b/Engine/Source/Editor/DetailCustomizations/Private/ComponentReferenceCustomization.cpp
--- a/Engine/Source/Editor/DetailCustomizations/Private/ComponentReferenceCustomization.cpp
+++ b/Engine/Source/Editor/DetailCustomizations/Private/ComponentReferenceCustomization.cpp
@@ -92,22 +92,7 @@ void FComponentReferenceCustomization::CustomizeHeader(TSharedRef<IPropertyHandl
 
 		InPropertyHandle->SetOnPropertyValueChanged(FSimpleDelegate::CreateSP(this, &FComponentReferenceCustomization::OnPropertyValueChanged));
 
-		// set cached values
-		{
-			CachedComponent.Reset();
-			CachedFirstOuterActor = GetFirstOuterActor();
-
-			FComponentReference TmpComponentReference;
-			CachedPropertyAccess = GetValue(TmpComponentReference);
-			if (CachedPropertyAccess == FPropertyAccess::Success)
-			{
-				CachedComponent = TmpComponentReference.GetComponent(CachedFirstOuterActor.Get());
-				if (!IsComponentReferenceValid(TmpComponentReference))
-				{
-					CachedComponent.Reset();
-				}
-			}
-		}
+		UpdateCachedValues(/*bResetValueIfInvalid*/ false);
 
 		HeaderRow.NameContent()
 		[
@@ -507,22 +492,38 @@ bool FComponentReferenceCustomization::IsComponentReferenceValid(const FComponen
 }
 
 void FComponentReferenceCustomization::OnPropertyValueChanged()
+{
+	UpdateCachedValues(/*bResetValueIfInvalid*/ true);
+}
+
+void FComponentReferenceCustomization::UpdateCachedValues(bool bResetValueIfInvalid)
 {
 	CachedComponent.Reset();
 	CachedFirstOuterActor = GetFirstOuterActor();
 
 	FComponentReference TmpComponentReference;
 	CachedPropertyAccess = GetValue(TmpComponentReference);
-	if (CachedPropertyAccess == FPropertyAccess::Success)
+	if (CachedPropertyAccess != FPropertyAccess::Success)
 	{
-		CachedComponent = TmpComponentReference.GetComponent(CachedFirstOuterActor.Get());
-		if (!IsComponentReferenceValid(TmpComponentReference))
+		return;
+	}
+
+	// If the reference is null, we don't want to update CachedComponent (the further code would
+	//  otherwise try to set it to the root component of our actor)
+	if (TmpComponentReference.PathToComponent.IsEmpty()
+		&& TmpComponentReference.ComponentProperty == NAME_None
+		&& TmpComponentReference.OverrideComponent.IsExplicitlyNull())
+	{
+		return;
+	}
+
+	CachedComponent = TmpComponentReference.GetComponent(CachedFirstOuterActor.Get());
+	if (!IsComponentReferenceValid(TmpComponentReference))
+	{
+		CachedComponent.Reset();
+		if (bResetValueIfInvalid && !(TmpComponentReference == FComponentReference()))
 		{
-			CachedComponent.Reset();
-			if (!(TmpComponentReference == FComponentReference()))
-			{
-				SetValue(FComponentReference());
-			}
+			SetValue(FComponentReference());
 		}
 	}
 }
diff --git a/Engine/Source/Editor/DetailCustomizations/Private/ComponentReferenceCustomization.h b/Engine/Source/Editor/DetailCustomizations/Private/ComponentReferenceCustomization.h
--- a/Engine/Source/Editor/DetailCustomizations/Private/ComponentReferenceCustomization.h
+++ b/Engine/Source/Editor/DetailCustomizations/Private/ComponentReferenceCustomization.h
@@ -106,6 +106,13 @@ private:
 	 */
 	void CloseComboButton();
 
+	/**
+	 * Updates CachedFirstOuterActor, CachedComponent, and CachedPropertyAccess
+	 * @param bResetValueIfInvalid If true and the resulting component doesn't match the filters, do
+	 *  a SetValue call to reset the reference to null.
+	 */
+	void UpdateCachedValues(bool bResetValueIfInvalid);
+
 private:
 	/** The property handle we are customizing */
 	TSharedPtr<IPropertyHandle> PropertyHandle;
