FVectorStructCustomization : When the cvar "Editor.UseForwardRightUpDisplayNames" is true, set the DisplayName and ToolTip to (Forward, Right, Up) nomenclature.

This read-only cvar is only read once and cached in the static UseForwardRightUpDisplayNames for performance.

Also add another independent cvar "Editor.DetailCustomizations.ShowVector3Children" to control whether or not the detail cutomizations for Vector3 variants expand to show children.

FRotatorStructCustomization : When the cvar "Editor.UseForwardRightUpDisplayNames" is true, set the DisplayName to (Roll, Pitch, Yaw) nomenclature and update the ToolTip to reference the (Forward, Right, Up) axes.

Another independent cvar "Editor.DetailCustomizations.ShowRotator3Children" controls whether or not the detail customizations for Rotator3 variants expand to show children.

diff --git a/Engine/Source/Editor/DetailCustomizations/Private/RotatorStructCustomization.cpp b/Engine/Source/Editor/DetailCustomizations/Private/RotatorStructCustomization.cpp
--- a/Engine/Source/Editor/DetailCustomizations/Private/RotatorStructCustomization.cpp
+++ b/Engine/Source/Editor/DetailCustomizations/Private/RotatorStructCustomization.cpp
@@ -8,15 +8,33 @@
 #include "PropertyHandle.h"
 #include "UObject/NameTypes.h"
 #include "UObject/UnrealType.h"
+#include "VectorStructCustomization.h"
 
 class IPropertyTypeCustomization;
 
+#define LOCTEXT_NAMESPACE "FRotatorStructCustomization"
+
+namespace UE::Editor::DetailCustomizations::Private
+{
+	static TAutoConsoleVariable<bool> ShowRotator3Children(
+		TEXT("Editor.DetailCustomizations.ShowRotator3Children"),
+		true,
+		TEXT("When true, the detail customizations for Rotator3 variants expand to show children")
+	);
+}
 
 TSharedRef<IPropertyTypeCustomization> FRotatorStructCustomization::MakeInstance() 
 {
 	return MakeShareable(new FRotatorStructCustomization);
 }
 
+void FRotatorStructCustomization::CustomizeChildren(TSharedRef<IPropertyHandle> StructPropertyHandle, IDetailChildrenBuilder& StructBuilder, IPropertyTypeCustomizationUtils& StructCustomizationUtils)
+{
+	if (UE::Editor::DetailCustomizations::Private::ShowRotator3Children.GetValueOnGameThread())
+	{
+		FMathStructCustomization::CustomizeChildren(StructPropertyHandle, StructBuilder, StructCustomizationUtils);
+	}
+}
 
 void FRotatorStructCustomization::GetSortedChildren(TSharedRef<IPropertyHandle> StructPropertyHandle, TArray< TSharedRef<IPropertyHandle> >& OutChildren)
 {
@@ -24,6 +42,8 @@ void FRotatorStructCustomization::GetSortedChildren(TSharedRef<IPropertyHandle>
 	static const FName Pitch("Pitch");
 	static const FName Yaw("Yaw");
 
+	const bool bUseForwardRightUpDisplayNames = FVectorStructCustomization::UseForwardRightUpDisplayNames();
+
 	TSharedPtr< IPropertyHandle > RotatorChildren[3];
 
 	uint32 NumChildren;
@@ -36,15 +56,34 @@ void FRotatorStructCustomization::GetSortedChildren(TSharedRef<IPropertyHandle>
 
 		if (PropertyName == Roll)
 		{
+			if (bUseForwardRightUpDisplayNames)
+			{
+				ChildHandle->SetPropertyDisplayName(LOCTEXT("RollDisplayName", "Roll (was X)"));
+				ChildHandle->SetToolTipText(LOCTEXT("RollToolTip", "Roll (degrees) around Forward (was X) axis"));
+			}
+
 			RotatorChildren[0] = ChildHandle;
 		}
 		else if (PropertyName == Pitch)
 		{
+			if (bUseForwardRightUpDisplayNames)
+			{
+				ChildHandle->SetPropertyDisplayName(LOCTEXT("PitchDisplayName", "Pitch (was Y)"));
+				ChildHandle->SetToolTipText(LOCTEXT("PitchToolTip", "Pitch (degrees) around Right (was Y) axis"));
+			}
+
 			RotatorChildren[1] = ChildHandle;
 		}
 		else
 		{
 			check(PropertyName == Yaw);
+
+			if (bUseForwardRightUpDisplayNames)
+			{
+				ChildHandle->SetPropertyDisplayName(LOCTEXT("YawDisplayName", "Yaw (was Z)"));
+				ChildHandle->SetToolTipText(LOCTEXT("YawToolTip", "Yaw (degrees) around Up (was Z) axis"));
+			}
+
 			RotatorChildren[2] = ChildHandle;
 		}
 	}
@@ -53,3 +92,5 @@ void FRotatorStructCustomization::GetSortedChildren(TSharedRef<IPropertyHandle>
 	OutChildren.Add(RotatorChildren[1].ToSharedRef());
 	OutChildren.Add(RotatorChildren[2].ToSharedRef());
 }
+
+#undef LOCTEXT_NAMESPACE
diff --git a/Engine/Source/Editor/DetailCustomizations/Private/RotatorStructCustomization.h b/Engine/Source/Editor/DetailCustomizations/Private/RotatorStructCustomization.h
--- a/Engine/Source/Editor/DetailCustomizations/Private/RotatorStructCustomization.h
+++ b/Engine/Source/Editor/DetailCustomizations/Private/RotatorStructCustomization.h
@@ -22,6 +22,9 @@ public:
 
 private:
 
+	/** IPropertyTypeCustomization interface */
+	virtual void CustomizeChildren(TSharedRef<IPropertyHandle> StructPropertyHandle, IDetailChildrenBuilder& StructBuilder, IPropertyTypeCustomizationUtils& StructCustomizationUtils) override;
+
 	/** FMathStructCustomization interface */
 	virtual void GetSortedChildren(TSharedRef<IPropertyHandle> StructPropertyHandle, TArray<TSharedRef<IPropertyHandle>>& OutChildren) override;
 };
diff --git a/Engine/Source/Editor/DetailCustomizations/Private/VectorStructCustomization.cpp b/Engine/Source/Editor/DetailCustomizations/Private/VectorStructCustomization.cpp
--- a/Engine/Source/Editor/DetailCustomizations/Private/VectorStructCustomization.cpp
+++ b/Engine/Source/Editor/DetailCustomizations/Private/VectorStructCustomization.cpp
@@ -11,12 +11,43 @@
 
 class IPropertyTypeCustomization;
 
+#define LOCTEXT_NAMESPACE "FVectorStructCustomization"
+
+namespace UE::Editor::DetailCustomizations::Private
+{
+	static TAutoConsoleVariable<bool> ShowVector3Children(
+		TEXT("Editor.DetailCustomizations.ShowVector3Children"),
+		true,
+		TEXT("When true, the detail customizations for Vector3 variants expand to show children")
+	);
+
+	static TOptional<bool> UseForwardRightUpDisplayNames;
+}
 
 TSharedRef<IPropertyTypeCustomization> FVectorStructCustomization::MakeInstance() 
 {
 	return MakeShareable(new FVectorStructCustomization);
 }
 
+bool FVectorStructCustomization::UseForwardRightUpDisplayNames()
+{
+	if (!UE::Editor::DetailCustomizations::Private::UseForwardRightUpDisplayNames.IsSet())
+	{
+		if (IConsoleVariable* UseForwardRightUpDisplayNamesVariable = IConsoleManager::Get().FindConsoleVariable(TEXT("Editor.UseForwardRightUpDisplayNames")))
+		{
+			UE::Editor::DetailCustomizations::Private::UseForwardRightUpDisplayNames = UseForwardRightUpDisplayNamesVariable->GetBool();
+		}
+	}
+	return UE::Editor::DetailCustomizations::Private::UseForwardRightUpDisplayNames.Get(false/*DefaultValue*/);
+}
+
+void FVectorStructCustomization::CustomizeChildren(TSharedRef<IPropertyHandle> StructPropertyHandle, IDetailChildrenBuilder& StructBuilder, IPropertyTypeCustomizationUtils& StructCustomizationUtils)
+{
+	if (UE::Editor::DetailCustomizations::Private::ShowVector3Children.GetValueOnGameThread())
+	{
+		FMathStructCustomization::CustomizeChildren(StructPropertyHandle, StructBuilder, StructCustomizationUtils);
+	}
+}
 
 void FVectorStructCustomization::GetSortedChildren(TSharedRef<IPropertyHandle> StructPropertyHandle, TArray< TSharedRef<IPropertyHandle> >& OutChildren)
 {
@@ -24,6 +55,8 @@ void FVectorStructCustomization::GetSortedChildren(TSharedRef<IPropertyHandle> S
 	static const FName Y("Y");
 	static const FName Z("Z");
 
+	const bool bUseForwardRightUpDisplayNames = UseForwardRightUpDisplayNames();
+	
 	TSharedPtr<IPropertyHandle> VectorChildren[3];
 
 	uint32 NumChildren;
@@ -36,15 +69,34 @@ void FVectorStructCustomization::GetSortedChildren(TSharedRef<IPropertyHandle> S
 
 		if (PropertyName == X)
 		{
+			if (bUseForwardRightUpDisplayNames)
+			{
+				ChildHandle->SetPropertyDisplayName(LOCTEXT("ForwardDisplayName", "Forward (was X)"));
+				ChildHandle->SetToolTipText(LOCTEXT("ForwardToolTip", "Forward (was X)"));
+			}
+
 			VectorChildren[0] = ChildHandle;
 		}
 		else if (PropertyName == Y)
 		{
+			if (bUseForwardRightUpDisplayNames)
+			{
+				ChildHandle->SetPropertyDisplayName(LOCTEXT("RightDisplayName", "Right (was Y)"));
+				ChildHandle->SetToolTipText(LOCTEXT("RightToolTip", "Right (was Y)"));
+			}
+
 			VectorChildren[1] = ChildHandle;
 		}
 		else
 		{
 			check(PropertyName == Z);
+
+			if (bUseForwardRightUpDisplayNames)
+			{
+				ChildHandle->SetPropertyDisplayName(LOCTEXT("UpDisplayName", "Up (was Z)"));
+				ChildHandle->SetToolTipText(LOCTEXT("UpToolTip", "Up (was Z)"));
+			}
+
 			VectorChildren[2] = ChildHandle;
 		}
 	}
@@ -53,3 +105,5 @@ void FVectorStructCustomization::GetSortedChildren(TSharedRef<IPropertyHandle> S
 	OutChildren.Add(VectorChildren[1].ToSharedRef());
 	OutChildren.Add(VectorChildren[2].ToSharedRef());
 }
+
+#undef LOCTEXT_NAMESPACE
diff --git a/Engine/Source/Editor/DetailCustomizations/Private/VectorStructCustomization.h b/Engine/Source/Editor/DetailCustomizations/Private/VectorStructCustomization.h
--- a/Engine/Source/Editor/DetailCustomizations/Private/VectorStructCustomization.h
+++ b/Engine/Source/Editor/DetailCustomizations/Private/VectorStructCustomization.h
@@ -20,8 +20,13 @@ public:
 	/** @return A new instance of this class */
 	static TSharedRef<IPropertyTypeCustomization> MakeInstance();
 
+	static bool UseForwardRightUpDisplayNames();
+
 private:
 
+	/** IPropertyTypeCustomization interface */
+	virtual void CustomizeChildren(TSharedRef<IPropertyHandle> StructPropertyHandle, IDetailChildrenBuilder& StructBuilder, IPropertyTypeCustomizationUtils& StructCustomizationUtils) override;
+
 	/** FMathStructCustomization interface */
 	virtual void GetSortedChildren(TSharedRef<IPropertyHandle> StructPropertyHandle, TArray<TSharedRef<IPropertyHandle>>& OutChildren) override;
 };
