Fix crash when selecting anim sequence with empty reference inside an anim composite
* Also show text for track segments with unbound sequences

diff --git a/Engine/Source/Editor/DetailCustomizations/Private/AnimMontageSegmentDetails.cpp b/Engine/Source/Editor/DetailCustomizations/Private/AnimMontageSegmentDetails.cpp
--- a/Engine/Source/Editor/DetailCustomizations/Private/AnimMontageSegmentDetails.cpp
+++ b/Engine/Source/Editor/DetailCustomizations/Private/AnimMontageSegmentDetails.cpp
@@ -695,7 +695,7 @@ void SAnimationSegmentViewport::Tick( const FGeometry& AllottedGeometry, const d
 		const float End = EndTimeAttribute.Get().Get(0.f);
 		const float PlayRate = PlayRateAttribute.Get().Get(0.f);
 
-		if (Component->PreviewInstance->GetCurrentTime() > End || Component->PreviewInstance->GetCurrentTime() < Start)
+		if (Component->PreviewInstance && (Component->PreviewInstance->GetCurrentTime() > End || Component->PreviewInstance->GetCurrentTime() < Start))
 		{
 			const float NewStart = PlayRate > 0.f ? Start : End;
 			Component->PreviewInstance->SetPosition(NewStart, false);
diff --git a/Engine/Source/Editor/Persona/Private/SAnimSegmentsPanel.cpp b/Engine/Source/Editor/Persona/Private/SAnimSegmentsPanel.cpp
--- a/Engine/Source/Editor/Persona/Private/SAnimSegmentsPanel.cpp
+++ b/Engine/Source/Editor/Persona/Private/SAnimSegmentsPanel.cpp
@@ -224,9 +224,13 @@ FString	SAnimSegmentsPanel::GetAnimSegmentName(int32 AnimSegmentIndex) const
 			{
 				TitleLabel = AssetName;
 			}
-
-			return TitleLabel;
 		}
+		else
+		{
+			TitleLabel = LOCTEXT("UnboundAnimSequence", "Unbound Animation Sequence").ToString();
+		}
+
+		return TitleLabel;
 	}
 	return FString();
 }
