[Curve Editor] Fixed Curve Editor not updated its data when resetting the runtime curve property

diff --git a/Engine/Source/Editor/DetailCustomizations/Private/CurveColorCustomization.cpp b/Engine/Source/Editor/DetailCustomizations/Private/CurveColorCustomization.cpp
--- a/Engine/Source/Editor/DetailCustomizations/Private/CurveColorCustomization.cpp
+++ b/Engine/Source/Editor/DetailCustomizations/Private/CurveColorCustomization.cpp
@@ -149,9 +149,12 @@ void FCurveColorCustomization::CustomizeHeader(TSharedRef<IPropertyHandle> InStr
 		}
 		else
 		{
-			CurveWidget->SetCurveOwner(this);
+			CurveWidget->SetCurveOwner(this, InStructPropertyHandle->IsEditable());
 			CurveWidget->SetPropertyUtils(StructCustomizationUtils.GetPropertyUtilities());
 		}
+
+		const FSimpleDelegate OnPropertyReset = FSimpleDelegate::CreateSP(this, &FCurveColorCustomization::HandleCurveColorReset);
+		StructPropertyHandle->SetOnPropertyResetToDefault(OnPropertyReset);
 	}
 	else
 	{
@@ -521,4 +524,9 @@ void FCurveColorCustomization::DestroyPopOutWindow()
 	}
 }
 
+void FCurveColorCustomization::HandleCurveColorReset()
+{
+	CurveWidget->SetCurveOwner(this, StructPropertyHandle->IsEditable());
+}
+
 #undef LOCTEXT_NAMESPACE
diff --git a/Engine/Source/Editor/DetailCustomizations/Private/CurveColorCustomization.h b/Engine/Source/Editor/DetailCustomizations/Private/CurveColorCustomization.h
--- a/Engine/Source/Editor/DetailCustomizations/Private/CurveColorCustomization.h
+++ b/Engine/Source/Editor/DetailCustomizations/Private/CurveColorCustomization.h
@@ -116,6 +116,11 @@ private:
 	 */
 	void DestroyPopOutWindow();
 
+	/**
+	 * Reset the curve view models of the curve editor to be this, after resetting the property
+	 */
+	void HandleCurveColorReset();
+
 private:
 	/** Cached RuntimeCurveLinearColor struct handle */
 	TSharedPtr<IPropertyHandle> StructPropertyHandle;
diff --git a/Engine/Source/Editor/DetailCustomizations/Private/CurveStructCustomization.cpp b/Engine/Source/Editor/DetailCustomizations/Private/CurveStructCustomization.cpp
--- a/Engine/Source/Editor/DetailCustomizations/Private/CurveStructCustomization.cpp
+++ b/Engine/Source/Editor/DetailCustomizations/Private/CurveStructCustomization.cpp
@@ -154,6 +154,9 @@ void FCurveStructCustomization::CustomizeHeader( TSharedRef<IPropertyHandle> InS
 			CurveWidget->SetCurveOwner(this, InStructPropertyHandle->IsEditable());
 			CurveWidget->SetPropertyUtils(StructCustomizationUtils.GetPropertyUtilities());
 		}
+
+		const FSimpleDelegate OnPropertyReset = FSimpleDelegate::CreateSP(this, &FCurveStructCustomization::HandleCurveStructReset);
+		StructPropertyHandle->SetOnPropertyResetToDefault(OnPropertyReset);
 	}
 	else
 	{
@@ -530,4 +533,9 @@ void FCurveStructCustomization::DestroyPopOutWindow()
 	}
 }
 
+void FCurveStructCustomization::HandleCurveStructReset()
+{
+	CurveWidget->SetCurveOwner(this, StructPropertyHandle->IsEditable());
+}
+
 #undef LOCTEXT_NAMESPACE
diff --git a/Engine/Source/Editor/DetailCustomizations/Private/CurveStructCustomization.h b/Engine/Source/Editor/DetailCustomizations/Private/CurveStructCustomization.h
--- a/Engine/Source/Editor/DetailCustomizations/Private/CurveStructCustomization.h
+++ b/Engine/Source/Editor/DetailCustomizations/Private/CurveStructCustomization.h
@@ -117,6 +117,11 @@ private:
 	 */
 	void DestroyPopOutWindow();
 
+	/**
+	 * Reset the curve view models of the curve editor to be this, after resetting the property
+	 */
+	void HandleCurveStructReset();
+
 private:
 	/** Cached RuntimeFloatCurve struct handle */
 	TSharedPtr<IPropertyHandle> StructPropertyHandle;
diff --git a/Engine/Source/Editor/DetailCustomizations/Private/CurveVectorCustomization.cpp b/Engine/Source/Editor/DetailCustomizations/Private/CurveVectorCustomization.cpp
--- a/Engine/Source/Editor/DetailCustomizations/Private/CurveVectorCustomization.cpp
+++ b/Engine/Source/Editor/DetailCustomizations/Private/CurveVectorCustomization.cpp
@@ -152,6 +152,9 @@ void FCurveVectorCustomization::CustomizeHeader(TSharedRef<IPropertyHandle> InSt
 			CurveWidget->SetCurveOwner(this, InStructPropertyHandle->IsEditable());
 			CurveWidget->SetPropertyUtils(StructCustomizationUtils.GetPropertyUtilities());
 		}
+
+		const FSimpleDelegate OnPropertyReset = FSimpleDelegate::CreateSP(this, &FCurveVectorCustomization::HandleCurveVectorReset);
+		StructPropertyHandle->SetOnPropertyResetToDefault(OnPropertyReset);
 	}
 	else
 	{
@@ -505,4 +508,9 @@ void FCurveVectorCustomization::DestroyPopOutWindow()
 	}
 }
 
+void FCurveVectorCustomization::HandleCurveVectorReset()
+{
+	CurveWidget->SetCurveOwner(this, StructPropertyHandle->IsEditable());
+}
+
 #undef LOCTEXT_NAMESPACE
diff --git a/Engine/Source/Editor/DetailCustomizations/Private/CurveVectorCustomization.h b/Engine/Source/Editor/DetailCustomizations/Private/CurveVectorCustomization.h
--- a/Engine/Source/Editor/DetailCustomizations/Private/CurveVectorCustomization.h
+++ b/Engine/Source/Editor/DetailCustomizations/Private/CurveVectorCustomization.h
@@ -109,6 +109,11 @@ private:
 	 */
 	void DestroyPopOutWindow();
 
+	/**
+	 * Reset the curve view models of the curve editor to be this, after resetting the property
+	 */
+	void HandleCurveVectorReset();
+
 private:
 	/** Cached RuntimeCurveLinearColor struct handle */
 	TSharedPtr<IPropertyHandle> StructPropertyHandle;
