Make FVectorStructCustomization public so that it can be extended in other modules

diff --git a/Engine/Source/Editor/DetailCustomizations/Private/DetailCustomizations.cpp b/Engine/Source/Editor/DetailCustomizations/Private/DetailCustomizations.cpp
--- a/Engine/Source/Editor/DetailCustomizations/Private/DetailCustomizations.cpp
+++ b/Engine/Source/Editor/DetailCustomizations/Private/DetailCustomizations.cpp
@@ -46,6 +46,7 @@
 #include "Customizations/MathStructProxyCustomizations.h"
 #include "Customizations/SlateBrushCustomization.h"
 #include "Customizations/SlateFontInfoCustomization.h"
+#include "Customizations/VectorStructCustomization.h"
 #include "DataTableCategoryCustomization.h"
 #include "DataTableCustomization.h"
 #include "DateTimeStructCustomization.h"
@@ -158,7 +159,6 @@
 #include "TimespanStructCustomization.h"
 #include "UObject/UnrealNames.h"
 #include "Vector4StructCustomization.h"
-#include "VectorStructCustomization.h"
 #include "WindowsTargetSettingsDetails.h"
 #include "WorldSettingsDetails.h"
 #include "LandscapeGrassTypeDetails.h"
diff --git a/Engine/Source/Editor/DetailCustomizations/Private/RotatorStructCustomization.cpp b/Engine/Source/Editor/DetailCustomizations/Private/RotatorStructCustomization.cpp
--- a/Engine/Source/Editor/DetailCustomizations/Private/RotatorStructCustomization.cpp
+++ b/Engine/Source/Editor/DetailCustomizations/Private/RotatorStructCustomization.cpp
@@ -2,13 +2,13 @@
 
 #include "RotatorStructCustomization.h"
 
+#include "Customizations/VectorStructCustomization.h"
 #include "HAL/Platform.h"
 #include "HAL/PlatformCrt.h"
 #include "Misc/AssertionMacros.h"
 #include "PropertyHandle.h"
 #include "UObject/NameTypes.h"
 #include "UObject/UnrealType.h"
-#include "VectorStructCustomization.h"
 
 class IPropertyTypeCustomization;
 
diff --git a/Engine/Source/Editor/DetailCustomizations/Private/VectorStructCustomization.cpp b/Engine/Source/Editor/DetailCustomizations/Private/VectorStructCustomization.cpp
--- a/Engine/Source/Editor/DetailCustomizations/Private/VectorStructCustomization.cpp
+++ b/Engine/Source/Editor/DetailCustomizations/Private/VectorStructCustomization.cpp
@@ -1,6 +1,6 @@
 // Copyright Epic Games, Inc. All Rights Reserved.
 
-#include "VectorStructCustomization.h"
+#include "Customizations/VectorStructCustomization.h"
 
 #include "HAL/Platform.h"
 #include "HAL/PlatformCrt.h"
diff --git a/Engine/Source/Editor/DetailCustomizations/Private/VectorStructCustomization.h b/Engine/Source/Editor/DetailCustomizations/Private/VectorStructCustomization.h
deleted file mode 100644
diff --git a/Engine/Source/Editor/DetailCustomizations/Public/Customizations/VectorStructCustomization.h b/Engine/Source/Editor/DetailCustomizations/Public/Customizations/VectorStructCustomization.h
new file mode 100644
--- /dev/null
+++ b/Engine/Source/Editor/DetailCustomizations/Public/Customizations/VectorStructCustomization.h
@@ -0,0 +1,32 @@
+// Copyright Epic Games, Inc. All Rights Reserved.
+
+#pragma once
+
+#include "Containers/Array.h"
+#include "Customizations/MathStructCustomizations.h"
+#include "Templates/SharedPointer.h"
+
+class IPropertyHandle;
+class IPropertyTypeCustomization;
+
+/**
+ * Customizes FVector structs.
+ */
+class DETAILCUSTOMIZATIONS_API FVectorStructCustomization
+	: public FMathStructCustomization
+{
+public:
+
+	/** @return A new instance of this class */
+	static TSharedRef<IPropertyTypeCustomization> MakeInstance();
+
+	static bool UseForwardRightUpDisplayNames();
+
+private:
+
+	/** IPropertyTypeCustomization interface */
+	virtual void CustomizeChildren(TSharedRef<IPropertyHandle> StructPropertyHandle, IDetailChildrenBuilder& StructBuilder, IPropertyTypeCustomizationUtils& StructCustomizationUtils) override;
+
+	/** FMathStructCustomization interface */
+	virtual void GetSortedChildren(TSharedRef<IPropertyHandle> StructPropertyHandle, TArray<TSharedRef<IPropertyHandle>>& OutChildren) override;
+};
