Fix rotation overridability, turns out the propagation was enabled for the rotation inside the transform...

#rnx

diff --git a/Engine/Source/Editor/DetailCustomizations/Private/MathStructProxyCustomizations.cpp b/Engine/Source/Editor/DetailCustomizations/Private/MathStructProxyCustomizations.cpp
--- a/Engine/Source/Editor/DetailCustomizations/Private/MathStructProxyCustomizations.cpp
+++ b/Engine/Source/Editor/DetailCustomizations/Private/MathStructProxyCustomizations.cpp
@@ -29,6 +29,7 @@
 #include "SlotBase.h"
 #include "Templates/UnrealTemplate.h"
 #include "UObject/Object.h"
+#include "UObject/OverridableManager.h"
 #include "UObject/UnrealType.h"
 #include "Widgets/DeclarativeSyntaxSupport.h"
 #include "Widgets/Input/SNumericEntryBox.h"
@@ -538,10 +539,13 @@ bool FMatrixStructCustomization<T>::FlushValues( TWeakPtr<IPropertyHandle> Prope
 				OuterObjects[ValueIndex]->GetArchetypeInstances(ArchetypeInstances);
 				for (UObject* ArchetypeInstance : ArchetypeInstances)
 				{
-					UE::Math::TMatrix<T>* CurrentValue = reinterpret_cast<UE::Math::TMatrix<T>*>(PropertyHandle->GetValueBaseAddress(reinterpret_cast<uint8*>(ArchetypeInstance)));
-					if (CurrentValue && CurrentValue->Equals(PreviousValue))
+					if (!FOverridableManager::Get().IsEnabled(*ArchetypeInstance))
 					{
-						*CurrentValue = NewValue;
+						UE::Math::TMatrix<T>* CurrentValue = reinterpret_cast<UE::Math::TMatrix<T>*>(PropertyHandle->GetValueBaseAddress(reinterpret_cast<uint8*>(ArchetypeInstance)));
+						if (CurrentValue && CurrentValue->Equals(PreviousValue))
+						{
+							*CurrentValue = NewValue;
+						}
 					}
 				}
 			}
@@ -696,10 +700,13 @@ bool FTransformStructCustomization<T>::FlushValues( TWeakPtr<IPropertyHandle> Pr
 				OuterObjects[ValueIndex]->GetArchetypeInstances(ArchetypeInstances);
 				for (UObject* ArchetypeInstance : ArchetypeInstances)
 				{
-					UE::Math::TTransform<T>* CurrentValue = reinterpret_cast<UE::Math::TTransform<T>*>(PropertyHandle->GetValueBaseAddress(reinterpret_cast<uint8*>(ArchetypeInstance)));
-					if (CurrentValue && CurrentValue->Equals(PreviousValue))
+					if (!FOverridableManager::Get().IsEnabled(*ArchetypeInstance))
 					{
-						*CurrentValue = NewValue;
+						UE::Math::TTransform<T>* CurrentValue = reinterpret_cast<UE::Math::TTransform<T>*>(PropertyHandle->GetValueBaseAddress(reinterpret_cast<uint8*>(ArchetypeInstance)));
+						if (CurrentValue && CurrentValue->Equals(PreviousValue))
+						{
+							*CurrentValue = NewValue;
+						}
 					}
 				}
 			}
@@ -839,10 +846,13 @@ bool FQuatStructCustomization<T>::FlushValues(TWeakPtr<IPropertyHandle> Property
 				OuterObjects[ValueIndex]->GetArchetypeInstances(ArchetypeInstances);
 				for (UObject* ArchetypeInstance : ArchetypeInstances)
 				{
-					UE::Math::TQuat<T>* CurrentValue = reinterpret_cast<UE::Math::TQuat<T>*>(PropertyHandle->GetValueBaseAddress(reinterpret_cast<uint8*>(ArchetypeInstance)));
-					if (CurrentValue && CurrentValue->Equals(PreviousValue))
+					if (!FOverridableManager::Get().IsEnabled(*ArchetypeInstance))
 					{
-						*CurrentValue = NewValue;
+						UE::Math::TQuat<T>* CurrentValue = reinterpret_cast<UE::Math::TQuat<T>*>(PropertyHandle->GetValueBaseAddress(reinterpret_cast<uint8*>(ArchetypeInstance)));
+						if (CurrentValue && CurrentValue->Equals(PreviousValue))
+						{
+							*CurrentValue = NewValue;
+						}
 					}
 				}
 			}
