Fix ToolTipText formatting for vector, rotation, transform (and other math struct types) to include the component display name in the tooltip of the numeric entry boxes.

There was existing code for this but it was not wired up properly, so only the value was being rendered in the ToolTipText.

SNumericEntryBox : Add an optional ToolTipTextFormat attribute

diff --git a/Engine/Source/Editor/DetailCustomizations/Private/MathStructCustomizations.cpp b/Engine/Source/Editor/DetailCustomizations/Private/MathStructCustomizations.cpp
--- a/Engine/Source/Editor/DetailCustomizations/Private/MathStructCustomizations.cpp
+++ b/Engine/Source/Editor/DetailCustomizations/Private/MathStructCustomizations.cpp
@@ -475,7 +475,7 @@ TSharedRef<SWidget> FMathStructCustomization::MakeNumericWidget(
 			.Delta(Metadata.Delta)
 			// LinearDeltaSensitivity must be left unset if not provided, rather than being set to some default
 			.LinearDeltaSensitivity(Metadata.LinearDeltaSensitivity != 0 ? Metadata.LinearDeltaSensitivity : TAttribute<int32>())
-			.ToolTipText(this, &FMathStructCustomization::OnGetValueToolTip<NumericType>, WeakHandlePtr)
+			.ToolTipTextFormat(this, &FMathStructCustomization::OnGetValueToolTipTextFormat, WeakHandlePtr)
 			.TypeInterface(Metadata.TypeInterface);
 }
 
@@ -694,16 +694,21 @@ void FMathStructCustomization::SetValue(NumericType NewValue, EPropertyValueSetF
 template <typename NumericType>
 FText FMathStructCustomization::OnGetValueToolTip(TWeakPtr<IPropertyHandle> WeakHandlePtr) const
 {
-	if(TSharedPtr<IPropertyHandle> PropertyHandle = WeakHandlePtr.Pin())
+	return FText::GetEmpty();
+}
+
+TOptional<FTextFormat> FMathStructCustomization::OnGetValueToolTipTextFormat(TWeakPtr<IPropertyHandle> WeakHandlePtr) const
+{
+	if (TSharedPtr<IPropertyHandle> PropertyHandle = WeakHandlePtr.Pin())
 	{
-		TOptional<NumericType> Value = OnGetValue<NumericType>(WeakHandlePtr);
-		if (Value.IsSet())
-		{
-			return FText::Format(LOCTEXT("ValueToolTip", "{0}: {1}"),  PropertyHandle->GetPropertyDisplayName(), FText::AsNumber(Value.GetValue()));
-		}
+		TStringBuilder<32> ToolTipTextFormatString;
+		ToolTipTextFormatString.Append(PropertyHandle->GetPropertyDisplayName().ToString());
+		ToolTipTextFormatString.Append(TEXT(": {0}"));
+
+		return TOptional<FTextFormat>(FText::FromString(ToolTipTextFormatString.ToString()));
 	}
-	
-	return FText::GetEmpty();
+
+	return TOptional<FTextFormat>();
 }
 
 bool FMathStructCustomization::IsValueEnabled(TWeakPtr<IPropertyHandle> WeakHandlePtr) const
diff --git a/Engine/Source/Editor/DetailCustomizations/Private/MathStructProxyCustomizations.cpp b/Engine/Source/Editor/DetailCustomizations/Private/MathStructProxyCustomizations.cpp
--- a/Engine/Source/Editor/DetailCustomizations/Private/MathStructProxyCustomizations.cpp
+++ b/Engine/Source/Editor/DetailCustomizations/Private/MathStructProxyCustomizations.cpp
@@ -72,7 +72,7 @@ TSharedRef<SWidget> FMathStructProxyCustomization::MakeNumericProxyWidget(TShare
 		.MaxSliderValue(bRotationInDegrees ? 360.0f : TOptional<NumericType>())
 		.MinSliderValue(bRotationInDegrees ? 0.0f : TOptional<NumericType>())
 		.LabelPadding(FMargin(3))
-		.ToolTipText(this, &FMathStructProxyCustomization::OnGetValueToolTip<ProxyType, NumericType>, WeakHandlePtr, ProxyValue, Label)
+		.ToolTipTextFormat(this, &FMathStructProxyCustomization::OnGetValueToolTipTextFormat, Label)
 		.LabelLocation(SNumericEntryBox<NumericType>::ELabelLocation::Inside)
 		.Label()
 		[
@@ -129,13 +129,19 @@ void FMathStructProxyCustomization::OnEndSliderMovement( NumericType NewValue, T
 template <typename ProxyType, typename NumericType>
 FText FMathStructProxyCustomization::OnGetValueToolTip(TWeakPtr<IPropertyHandle> WeakHandlePtr, TSharedRef<TProxyProperty<ProxyType, NumericType>> ProxyValue, FText Label) const
 {
-	TOptional<NumericType> Value = OnGetValue<ProxyType, NumericType>(WeakHandlePtr, ProxyValue);
-	if (Value.IsSet())
+	return FText::GetEmpty();
+}
+
+TOptional<FTextFormat> FMathStructProxyCustomization::OnGetValueToolTipTextFormat(FText Label) const
+{
+	if (!Label.IsEmptyOrWhitespace())
 	{
-		return FText::Format(LOCTEXT("ValueToolTip", "{0}: {1}"), Label, FText::AsNumber(Value.GetValue()));
+		TStringBuilder<32> ToolTipTextFormatString;
+		ToolTipTextFormatString.Append(Label.ToString());
+		ToolTipTextFormatString.Append(TEXT(": {0}"));
+		return TOptional<FTextFormat>(FText::FromString(ToolTipTextFormatString.ToString()));
 	}
-
-	return FText::GetEmpty();
+	return TOptional<FTextFormat>();
 }
 
 template<typename T>
@@ -185,17 +191,17 @@ void FMatrixStructCustomization<T>::CustomizeLocation(TSharedRef<class IProperty
 		+ SHorizontalBox::Slot()
 		.Padding(FMargin(0.0f, 2.0f, 3.0f, 2.0f))
 		[
-			MakeNumericProxyWidget<UE::Math::TVector<T>, T>(StructPropertyHandle, CachedTranslationX, LOCTEXT("TranslationX", "X"), false, AxisDisplayInfo::GetAxisColor(EAxisList::X))
+			MakeNumericProxyWidget<UE::Math::TVector<T>, T>(StructPropertyHandle, CachedTranslationX, AxisDisplayInfo::GetAxisDisplayName(EAxisList::X), false, AxisDisplayInfo::GetAxisColor(EAxisList::X))
 		]
 		+ SHorizontalBox::Slot()
 		.Padding(FMargin(0.0f, 2.0f, 3.0f, 2.0f))
 		[
-			MakeNumericProxyWidget<UE::Math::TVector<T>, T>(StructPropertyHandle, CachedTranslationY, LOCTEXT("TranslationY", "Y"), false, AxisDisplayInfo::GetAxisColor(EAxisList::Y))
+			MakeNumericProxyWidget<UE::Math::TVector<T>, T>(StructPropertyHandle, CachedTranslationY, AxisDisplayInfo::GetAxisDisplayName(EAxisList::Y), false, AxisDisplayInfo::GetAxisColor(EAxisList::Y))
 		]
 		+ SHorizontalBox::Slot()
 		.Padding(FMargin(0.0f, 2.0f, 0.0f, 2.0f))
 		[
-			MakeNumericProxyWidget<UE::Math::TVector<T>, T>(StructPropertyHandle, CachedTranslationZ, LOCTEXT("TranslationZ", "Z"), false, AxisDisplayInfo::GetAxisColor(EAxisList::Z))
+			MakeNumericProxyWidget<UE::Math::TVector<T>, T>(StructPropertyHandle, CachedTranslationZ, AxisDisplayInfo::GetAxisDisplayName(EAxisList::Z), false, AxisDisplayInfo::GetAxisColor(EAxisList::Z))
 		]
 	];
 }
@@ -225,17 +231,17 @@ void FMatrixStructCustomization<T>::CustomizeRotation(TSharedRef<class IProperty
 		+ SHorizontalBox::Slot()
 		.Padding(FMargin(0.0f, 2.0f, 3.0f, 2.0f))
 		[
-			MakeNumericProxyWidget<UE::Math::TRotator<T>, T>(StructPropertyHandle, CachedRotationRoll, LOCTEXT("RotationRoll", "X"), true, AxisDisplayInfo::GetAxisColor(EAxisList::X))
+			MakeNumericProxyWidget<UE::Math::TRotator<T>, T>(StructPropertyHandle, CachedRotationRoll, LOCTEXT("RotationRoll", "Roll"), true, AxisDisplayInfo::GetAxisColor(EAxisList::X))
 		]
 		+ SHorizontalBox::Slot()
 		.Padding(FMargin(0.0f, 2.0f, 3.0f, 2.0f))
 		[
-			MakeNumericProxyWidget<UE::Math::TRotator<T>, T>(StructPropertyHandle, CachedRotationPitch, LOCTEXT("RotationPitch", "Y"), true, AxisDisplayInfo::GetAxisColor(EAxisList::Y))
+			MakeNumericProxyWidget<UE::Math::TRotator<T>, T>(StructPropertyHandle, CachedRotationPitch, LOCTEXT("RotationPitch", "Pitch"), true, AxisDisplayInfo::GetAxisColor(EAxisList::Y))
 		]
 		+ SHorizontalBox::Slot()
 		.Padding(FMargin(0.0f, 2.0f, 0.0f, 2.0f))
 		[
-			MakeNumericProxyWidget<UE::Math::TRotator<T>, T>(StructPropertyHandle, CachedRotationYaw, LOCTEXT("RotationYaw", "Z"), true, AxisDisplayInfo::GetAxisColor(EAxisList::Z))
+			MakeNumericProxyWidget<UE::Math::TRotator<T>, T>(StructPropertyHandle, CachedRotationYaw, LOCTEXT("RotationYaw", "Yaw"), true, AxisDisplayInfo::GetAxisColor(EAxisList::Z))
 		]
 	];
 }
@@ -265,17 +271,17 @@ void FMatrixStructCustomization<T>::CustomizeScale(TSharedRef<class IPropertyHan
 		+ SHorizontalBox::Slot()
 		.Padding(FMargin(0.0f, 2.0f, 3.0f, 2.0f))
 		[
-			MakeNumericProxyWidget<UE::Math::TVector<T>, T>(StructPropertyHandle, CachedScaleX, LOCTEXT("ScaleX", "X"), false, AxisDisplayInfo::GetAxisColor(EAxisList::X))
+			MakeNumericProxyWidget<UE::Math::TVector<T>, T>(StructPropertyHandle, CachedScaleX, AxisDisplayInfo::GetAxisDisplayName(EAxisList::X), false, AxisDisplayInfo::GetAxisColor(EAxisList::X))
 		]
 		+ SHorizontalBox::Slot()
 		.Padding(FMargin(0.0f, 2.0f, 3.0f, 2.0f))
 		[
-			MakeNumericProxyWidget<UE::Math::TVector<T>, T>(StructPropertyHandle, CachedScaleY, LOCTEXT("ScaleY", "Y"), false, AxisDisplayInfo::GetAxisColor(EAxisList::Y))
+			MakeNumericProxyWidget<UE::Math::TVector<T>, T>(StructPropertyHandle, CachedScaleY, AxisDisplayInfo::GetAxisDisplayName(EAxisList::Y), false, AxisDisplayInfo::GetAxisColor(EAxisList::Y))
 		]
 		+ SHorizontalBox::Slot()
 		.Padding(FMargin(0.0f, 2.0f, 0.0f, 2.0f))
 		[
-			MakeNumericProxyWidget<UE::Math::TVector<T>, T>(StructPropertyHandle, CachedScaleZ, LOCTEXT("ScaleZ", "Z"), false, AxisDisplayInfo::GetAxisColor(EAxisList::Z))
+			MakeNumericProxyWidget<UE::Math::TVector<T>, T>(StructPropertyHandle, CachedScaleZ, AxisDisplayInfo::GetAxisDisplayName(EAxisList::Z), false, AxisDisplayInfo::GetAxisColor(EAxisList::Z))
 		]
 	];
 }
diff --git a/Engine/Source/Editor/DetailCustomizations/Public/Customizations/MathStructCustomizations.h b/Engine/Source/Editor/DetailCustomizations/Public/Customizations/MathStructCustomizations.h
--- a/Engine/Source/Editor/DetailCustomizations/Public/Customizations/MathStructCustomizations.h
+++ b/Engine/Source/Editor/DetailCustomizations/Public/Customizations/MathStructCustomizations.h
@@ -166,9 +166,13 @@ protected:
 	 * @praram WeakHandlePtr Handle to the property to get the value from
 	 */
 	template <typename NumericType>
+	UE_DEPRECATED(5.6, "OnGetValueToolTip is deprecated and no longer called.")
 	FText OnGetValueToolTip(TWeakPtr<IPropertyHandle> WeakHandlePtr) const;
+
 private:
 
+	TOptional<FTextFormat> OnGetValueToolTipTextFormat(TWeakPtr<IPropertyHandle> WeakHandlePtr) const;
+
 	/** Gets the brush to use for the lock icon. */
 	const FSlateBrush* GetPreserveScaleRatioImage() const;
 
diff --git a/Engine/Source/Editor/DetailCustomizations/Public/Customizations/MathStructProxyCustomizations.h b/Engine/Source/Editor/DetailCustomizations/Public/Customizations/MathStructProxyCustomizations.h
--- a/Engine/Source/Editor/DetailCustomizations/Public/Customizations/MathStructProxyCustomizations.h
+++ b/Engine/Source/Editor/DetailCustomizations/Public/Customizations/MathStructProxyCustomizations.h
@@ -193,8 +193,13 @@ protected:
 	TSharedRef<SWidget> MakeNumericProxyWidget(TSharedRef<IPropertyHandle>& StructPropertyHandle, TSharedRef< TProxyProperty<ProxyType, NumericType> >& ProxyValue, const FText& Label, bool bRotationInDegrees = false, const FLinearColor& LabelColor = FCoreStyle::Get().GetColor("DefaultForeground"));
 
 	template <typename ProxyType, typename NumericType>
+	UE_DEPRECATED(5.6, "OnGetValueToolTip is deprecated and no longer called. Please use the Label in MakeNumericProxyWidget to customize the display name portion of the ToolTip.")
 	FText OnGetValueToolTip(TWeakPtr<IPropertyHandle> WeakHandlePtr, TSharedRef<TProxyProperty<ProxyType, NumericType>> ProxyValue, FText Label) const;
+
 private:
+
+	TOptional<FTextFormat> OnGetValueToolTipTextFormat(FText Label) const;
+
 	/**
 	 * Gets the value as a float for the provided property handle
 	 *
diff --git a/Engine/Source/Runtime/Slate/Public/Widgets/Input/SNumericEntryBox.h b/Engine/Source/Runtime/Slate/Public/Widgets/Input/SNumericEntryBox.h
--- a/Engine/Source/Runtime/Slate/Public/Widgets/Input/SNumericEntryBox.h
+++ b/Engine/Source/Runtime/Slate/Public/Widgets/Input/SNumericEntryBox.h
@@ -99,6 +99,7 @@ public:
 		, _DisplayToggle(false)
 		, _ToggleChecked(ECheckBoxState::Checked)
 		, _TogglePadding(FMargin(1.f,0.f,1.f,0.f) )
+		, _ToolTipTextFormat(TOptional<FTextFormat>())
 	{}
 
 		/** Style to use for the editable text box within this widget */
@@ -198,6 +199,8 @@ public:
 		SLATE_EVENT( FOnCheckStateChanged, OnToggleChanged )
 		/** Padding around the toggle checkbox */
 		SLATE_ARGUMENT( FMargin, TogglePadding )
+		/** An optional format pattern that is used to format the ToolTipText. If set, this is used with a single argument ({0} = the display value). */
+		SLATE_ATTRIBUTE(TOptional<FTextFormat>, ToolTipTextFormat)
 
 	SLATE_END_ARGS()
 	SNumericEntryBox()
@@ -230,6 +233,8 @@ public:
 		SetMinFractionalDigits(MinFractionalDigits);
 		SetMaxFractionalDigits(MaxFractionalDigits);
 
+		ToolTipTextFormat = InArgs._ToolTipTextFormat;
+
 		CachedExternalValue = ValueAttribute.Get();
 		if (CachedExternalValue.IsSet())
 		{
@@ -281,13 +286,13 @@ public:
 				.OnEndSliderMovement(InArgs._OnEndSliderMovement)
 				.MinDesiredWidth(InArgs._MinDesiredValueWidth)
 				.TypeInterface(Interface)
-				.ToolTipText(this, &SNumericEntryBox<NumericType>::GetValueAsText);
+				.ToolTipText(this, &SNumericEntryBox<NumericType>::GetToolTipText);
 		}
 
 		// Always create an editable text box.  In the case of an undetermined value being passed in, we cant use the spinbox.
 		SAssignNew(EditableText, SEditableText)
 			.Text(this, &SNumericEntryBox<NumericType>::OnGetValueForTextBox)
-			.ToolTipText(this, &SNumericEntryBox<NumericType>::GetValueAsText)
+			.ToolTipText(this, &SNumericEntryBox<NumericType>::GetToolTipText)
 			.ColorAndOpacity(InArgs._EditableTextBoxStyle->ForegroundColor)
 			.Visibility(bAllowSpin ? EVisibility::Collapsed : EVisibility::Visible)
 			.Font(InArgs._Font.IsSet() ? InArgs._Font : InArgs._EditableTextBoxStyle->TextStyle.Font)
@@ -568,8 +573,7 @@ private:
 		return bUseCachedString ? CachedValueString : Interface->ToString(CurrentValue);  
 	}
 
-	/** @return the value being observed by the Numeric Entry Box as a FText */
-	FText GetValueAsText() const
+	FText GetToolTipText() const
 	{
 		const TOptional<NumericType>& Value = ValueAttribute.Get();
 
@@ -595,8 +599,17 @@ private:
 		if (Value.IsSet() == true)
 		{
 			NumericType CurrentValue = Value.GetValue();
-			return FText::FromString(GetCachedString(CurrentValue));
+
+			FText CurrentValueText = FText::FromString(GetCachedString(CurrentValue));
+
+			if (ToolTipTextFormat.Get().IsSet())
+			{
+				return FText::Format(ToolTipTextFormat.Get().GetValue(), CurrentValueText);
+			}
+
+			return CurrentValueText;
 		}
+
 		return FText::GetEmpty();
 	}
 
@@ -846,6 +859,7 @@ private:
 	bool bCachedValueStringDirty;
 	TAttribute< TOptional<int32> > MinFractionalDigits;
 	TAttribute< TOptional<int32> > MaxFractionalDigits;
+	TAttribute< TOptional<FTextFormat> > ToolTipTextFormat;
 };
 
 
diff --git a/Engine/Source/Runtime/Slate/Public/Widgets/Input/SRotatorInputBox.h b/Engine/Source/Runtime/Slate/Public/Widgets/Input/SRotatorInputBox.h
--- a/Engine/Source/Runtime/Slate/Public/Widgets/Input/SRotatorInputBox.h
+++ b/Engine/Source/Runtime/Slate/Public/Widgets/Input/SRotatorInputBox.h
@@ -182,13 +182,7 @@ public:
 				.OnBeginSliderMovement( CreatePerComponentSliderMovementEvent(InArgs._OnBeginSliderMovement, InArgs._OnRollBeginSliderMovement) )
 				.OnEndSliderMovement( CreatePerComponentSliderMovementEvent< FOnNumericValueChanged, NumericType >(InArgs._OnEndSliderMovement, InArgs._OnRollEndSliderMovement) )
 				.UndeterminedString( NSLOCTEXT("SRotatorInputBox", "MultipleValues", "Multiple Values") )
-				.ToolTipText_Lambda([RollAttr = InArgs._Roll]
-				{
-					const TOptional<NumericType>& RollValue = RollAttr.Get();
-					return RollValue.IsSet() 
-						? FText::Format(NSLOCTEXT("SRotatorInputBox", "Roll_ToolTip", "X(Roll): {0}"), RollValue.GetValue())
-						: NSLOCTEXT("SRotatorInputBox", "MultipleValues", "Multiple Values");
-				})
+				.ToolTipTextFormat(TAttribute<TOptional<FTextFormat>>(NSLOCTEXT("SRotatorInputBox", "Roll_ToolTipTextFormat", "Roll: {0}")))
 				.TypeInterface(InArgs._TypeInterface)
 				.DisplayToggle(InArgs._DisplayToggle)
 				.ToggleChecked(InArgs._ToggleRollChecked)
@@ -219,13 +213,7 @@ public:
 				.OnBeginSliderMovement( CreatePerComponentSliderMovementEvent( InArgs._OnBeginSliderMovement, InArgs._OnPitchBeginSliderMovement ) )
 				.OnEndSliderMovement( CreatePerComponentSliderMovementEvent< FOnNumericValueChanged, NumericType >( InArgs._OnEndSliderMovement, InArgs._OnPitchEndSliderMovement ) )
 				.UndeterminedString( NSLOCTEXT("SRotatorInputBox", "MultipleValues", "Multiple Values") )
-				.ToolTipText_Lambda([PitchAttr = InArgs._Pitch]
-				{
-					const TOptional<NumericType>& PitchValue = PitchAttr.Get();
-					return PitchValue.IsSet()
-						? FText::Format(NSLOCTEXT("SRotatorInputBox", "Pitch_ToolTip", "Y(Pitch): {0}"), PitchValue.GetValue())
-						: NSLOCTEXT("SRotatorInputBox", "MultipleValues", "Multiple Values");
-				})
+				.ToolTipTextFormat(TAttribute<TOptional<FTextFormat>>(NSLOCTEXT("SRotatorInputBox", "Pitch_ToolTipTextFormat", "Pitch: {0}")))
 				.TypeInterface(InArgs._TypeInterface)
 				.DisplayToggle(InArgs._DisplayToggle)
 				.ToggleChecked(InArgs._TogglePitchChecked)
@@ -256,13 +244,7 @@ public:
 				.OnBeginSliderMovement( CreatePerComponentSliderMovementEvent(InArgs._OnBeginSliderMovement, InArgs._OnYawBeginSliderMovement) )
 				.OnEndSliderMovement( CreatePerComponentSliderMovementEvent< FOnNumericValueChanged, NumericType >(InArgs._OnEndSliderMovement, InArgs._OnYawEndSliderMovement) )
 				.UndeterminedString( NSLOCTEXT("SRotatorInputBox", "MultipleValues", "Multiple Values") )
-				.ToolTipText_Lambda([YawAttr = InArgs._Yaw]
-				{
-					const TOptional<NumericType>& YawValue = YawAttr.Get();
-					return YawValue.IsSet()
-						? FText::Format(NSLOCTEXT("SRotatorInputBox", "Yaw_ToolTip", "Z(Yaw): {0}"), YawValue.GetValue())
-						: NSLOCTEXT("SRotatorInputBox", "MultipleValues", "Multiple Values");
-				})
+				.ToolTipTextFormat(TAttribute<TOptional<FTextFormat>>(NSLOCTEXT("SRotatorInputBox", "Yaw_ToolTipTextFormat", "Yaw: {0}")))
 				.TypeInterface(InArgs._TypeInterface)
 				.DisplayToggle(InArgs._DisplayToggle)
 				.ToggleChecked(InArgs._ToggleYawChecked)
diff --git a/Engine/Source/Runtime/Slate/Public/Widgets/Input/SVectorInputBox.h b/Engine/Source/Runtime/Slate/Public/Widgets/Input/SVectorInputBox.h
--- a/Engine/Source/Runtime/Slate/Public/Widgets/Input/SVectorInputBox.h
+++ b/Engine/Source/Runtime/Slate/Public/Widgets/Input/SVectorInputBox.h
@@ -405,7 +405,7 @@ private:
 	void ConstructComponent(int32 ComponentIndex,
 		const FArguments& InArgs,
 		const FLinearColor& LabelColor,
-		const FText& TooltipText,
+		const FText& ComponentDisplayName,
 		TSharedRef<SHorizontalBox>& HorizontalBox,
 		const TAttribute<TOptional<NumericType>>& Component,
 		const FOnNumericValueChanged& OnComponentChanged,
@@ -422,6 +422,10 @@ private:
 			LabelWidget = SNumericEntryBox<NumericType>::BuildNarrowColorLabel(LabelColor);
 		}
 
+		TStringBuilder<32> ToolTipTextFormatString;
+		ToolTipTextFormatString.Append(ComponentDisplayName.ToString());
+		ToolTipTextFormatString.Append(TEXT(": {0}"));
+
 		TAttribute<TOptional<NumericType>> Value = CreatePerComponentGetter(ComponentIndex, Component, VectorAttribute);
 
 		// any other getter below can use the vector
@@ -437,14 +441,7 @@ private:
 			.Value(Value)
 			.OnValueChanged(CreatePerComponentChanged(ComponentIndex, OnComponentChanged, InArgs._ConstrainVector))
 			.OnValueCommitted(CreatePerComponentCommitted(ComponentIndex, OnComponentCommitted, InArgs._ConstrainVector))
-			.ToolTipText(MakeAttributeLambda([Value, TooltipText]
-			{
-				if (Value.Get().IsSet())
-				{
-					return FText::Format(TooltipText, Value.Get().GetValue());
-				}
-				return NSLOCTEXT("SVectorInputBox", "MultipleValues", "Multiple Values");
-			}))
+			.ToolTipTextFormat(TAttribute<TOptional<FTextFormat>>(FText::FromString(ToolTipTextFormatString.ToString())))
 			.UndeterminedString(NSLOCTEXT("SVectorInputBox", "MultipleValues", "Multiple Values"))
 			.ContextMenuExtender(OnContextMenuExtenderComponent)
 			.TypeInterface(InArgs._TypeInterface)
@@ -478,7 +475,7 @@ private:
 		ConstructComponent(0,
 			InArgs,
 			AxisDisplayInfo::GetAxisColor(EAxisList::X),
-			NSLOCTEXT("SVectorInputBox", "X_ToolTip", "X: {0}"),
+			AxisDisplayInfo::GetAxisDisplayName(EAxisList::X),
 			HorizontalBox,
 			InArgs._X,
 			InArgs._OnXChanged,
@@ -499,7 +496,7 @@ private:
 		ConstructComponent(1,
 			InArgs,
 			AxisDisplayInfo::GetAxisColor(EAxisList::Y),
-			NSLOCTEXT("SVectorInputBox", "Y_ToolTip", "Y: {0}"),
+			AxisDisplayInfo::GetAxisDisplayName(EAxisList::Y),
 			HorizontalBox,
 			InArgs._Y,
 			InArgs._OnYChanged,
@@ -520,7 +517,7 @@ private:
 		ConstructComponent(2,
 			InArgs,
 			AxisDisplayInfo::GetAxisColor(EAxisList::Z),
-			NSLOCTEXT("SVectorInputBox", "Z_ToolTip", "Z: {0}"),
+			AxisDisplayInfo::GetAxisDisplayName(EAxisList::Z),
 			HorizontalBox,
 			InArgs._Z,
 			InArgs._OnZChanged,
@@ -541,7 +538,7 @@ private:
 		ConstructComponent(3,
 			InArgs,
 			SNumericEntryBox<NumericType>::LilacLabelBackgroundColor,
-			NSLOCTEXT("SVectorInputBox", "W_ToolTip", "W: {0}"),
+			NSLOCTEXT("SVectorInputBox", "W_DisplayName", "W"),
 			HorizontalBox,
 			InArgs._W,
 			InArgs._OnWChanged,
