Also hide (per instance) Custom Data when there are too many instances in a static mesh component

diff --git a/Engine/Source/Editor/DetailCustomizations/Private/InstancedStaticMeshComponentDetails.cpp b/Engine/Source/Editor/DetailCustomizations/Private/InstancedStaticMeshComponentDetails.cpp
--- a/Engine/Source/Editor/DetailCustomizations/Private/InstancedStaticMeshComponentDetails.cpp
+++ b/Engine/Source/Editor/DetailCustomizations/Private/InstancedStaticMeshComponentDetails.cpp
@@ -49,11 +49,14 @@ void FInstancedStaticMeshComponentDetails::CustomizeDetails(IDetailLayoutBuilder
 	DetailBuilder.GetObjectsBeingCustomized(ObjectsBeingCustomized);
 
 	TSharedPtr<IPropertyHandle> PerInstanceSMDataProperty = DetailBuilder.GetProperty(GET_MEMBER_NAME_CHECKED(UInstancedStaticMeshComponent, PerInstanceSMData));
+	TSharedPtr<IPropertyHandle> PerInstanceSMCustomDataProperty = DetailBuilder.GetProperty(GET_MEMBER_NAME_CHECKED(UInstancedStaticMeshComponent, PerInstanceSMCustomData));
+
 
 	// Doesn't support multi-selection : 
 	if (ObjectsBeingCustomized.Num() != 1)
 	{
 		PerInstanceSMDataProperty->MarkHiddenByCustomization();
+		PerInstanceSMCustomDataProperty->MarkHiddenByCustomization();
 	}
 	else
 	{
@@ -72,6 +75,7 @@ void FInstancedStaticMeshComponentDetails::CustomizeDetails(IDetailLayoutBuilder
 			if (bTooManyInstances && !bForceShowAllInstances)
 			{
 				PerInstanceSMDataProperty->MarkHiddenByCustomization();
+				PerInstanceSMCustomDataProperty->MarkHiddenByCustomization();
 			}
 
 			// If there are too many instances to display, display an additional button to show/hide instances : 
@@ -84,39 +88,39 @@ void FInstancedStaticMeshComponentDetails::CustomizeDetails(IDetailLayoutBuilder
 					.NameContent()
 					[
 						SNew(SHorizontalBox)
-						+ SHorizontalBox::Slot()
-						.VAlign(VAlign_Center)
-						.HAlign(HAlign_Left)
-						.AutoWidth()
-						.Padding(0.0f, 0.0f, 4.0f, 0.0f)
-						[
-							SNew(SImage)
-							.Image(FCoreStyle::Get().GetBrush("Icons.Warning"))
-						.ToolTipText(TooltipText)
-						]
-						+ SHorizontalBox::Slot()
-						.VAlign(VAlign_Center)
-						.HAlign(HAlign_Left)
-						.AutoWidth()
-						[
-							SNew(STextBlock)
-							.Font(IDetailLayoutBuilder::GetDetailFontItalic())
-						.Text(SlatePerformanceWarningText)
-						.ToolTipText(TooltipText)
-						]
+							+ SHorizontalBox::Slot()
+							.VAlign(VAlign_Center)
+							.HAlign(HAlign_Left)
+							.AutoWidth()
+							.Padding(0.0f, 0.0f, 4.0f, 0.0f)
+							[
+								SNew(SImage)
+									.Image(FCoreStyle::Get().GetBrush("Icons.Warning"))
+									.ToolTipText(TooltipText)
+							]
+							+ SHorizontalBox::Slot()
+							.VAlign(VAlign_Center)
+							.HAlign(HAlign_Left)
+							.AutoWidth()
+							[
+								SNew(STextBlock)
+									.Font(IDetailLayoutBuilder::GetDetailFontItalic())
+									.Text(SlatePerformanceWarningText)
+									.ToolTipText(TooltipText)
+							]
 					]
-				.ValueContent()
+					.ValueContent()
 					.MaxDesiredWidth(120.f)
 					[
 						SNew(SButton)
-						.OnClicked(this, &FInstancedStaticMeshComponentDetails::OnShowHideAllInstancesClicked)
-						.ToolTipText(TooltipText)
-						.IsEnabled(this, &FInstancedStaticMeshComponentDetails::IsShowHideAllInstancesEnabled)
-						[
-							SNew(STextBlock)
-							.Font(IDetailLayoutBuilder::GetDetailFontItalic())
-							.Text(this, &FInstancedStaticMeshComponentDetails::GetShowHideAllInstancesText)
-						]
+							.OnClicked(this, &FInstancedStaticMeshComponentDetails::OnShowHideAllInstancesClicked)
+							.ToolTipText(TooltipText)
+							.IsEnabled(this, &FInstancedStaticMeshComponentDetails::IsShowHideAllInstancesEnabled)
+							[
+								SNew(STextBlock)
+									.Font(IDetailLayoutBuilder::GetDetailFontItalic())
+									.Text(this, &FInstancedStaticMeshComponentDetails::GetShowHideAllInstancesText)
+							]
 					];
 			}
 		}
@@ -159,4 +163,4 @@ FText FInstancedStaticMeshComponentDetails::GetShowHideAllInstancesText() const
 		: FText::Format(LOCTEXT("ShowAllInstances", "Show All {0} Instances"), NumInstances);
 }
 
-#undef LOCTEXT_NAMESPACE
+#undef LOCTEXT_NAMESPACE
\ No newline at end of file
