RotatorStructCustomization titles shouldn't have Roll (was X), Pitch (was Y), Yaw (was Z)

The tooltips indicate which axis to rotate around, but the titles should just be Roll, Pitch, Yaw

diff --git a/Engine/Source/Editor/DetailCustomizations/Private/RotatorStructCustomization.cpp b/Engine/Source/Editor/DetailCustomizations/Private/RotatorStructCustomization.cpp
--- a/Engine/Source/Editor/DetailCustomizations/Private/RotatorStructCustomization.cpp
+++ b/Engine/Source/Editor/DetailCustomizations/Private/RotatorStructCustomization.cpp
@@ -58,7 +58,7 @@ void FRotatorStructCustomization::GetSortedChildren(TSharedRef<IPropertyHandle>
 		{
 			if (bUseForwardRightUpDisplayNames)
 			{
-				ChildHandle->SetPropertyDisplayName(LOCTEXT("RollDisplayName", "Roll (was X)"));
+				ChildHandle->SetPropertyDisplayName(LOCTEXT("RollDisplayName", "Roll"));
 				ChildHandle->SetToolTipText(LOCTEXT("RollToolTip", "Roll (degrees) around Forward (was X) axis"));
 			}
 
@@ -68,7 +68,7 @@ void FRotatorStructCustomization::GetSortedChildren(TSharedRef<IPropertyHandle>
 		{
 			if (bUseForwardRightUpDisplayNames)
 			{
-				ChildHandle->SetPropertyDisplayName(LOCTEXT("PitchDisplayName", "Pitch (was Y)"));
+				ChildHandle->SetPropertyDisplayName(LOCTEXT("PitchDisplayName", "Pitch"));
 				ChildHandle->SetToolTipText(LOCTEXT("PitchToolTip", "Pitch (degrees) around Right (was Y) axis"));
 			}
 
@@ -80,7 +80,7 @@ void FRotatorStructCustomization::GetSortedChildren(TSharedRef<IPropertyHandle>
 
 			if (bUseForwardRightUpDisplayNames)
 			{
-				ChildHandle->SetPropertyDisplayName(LOCTEXT("YawDisplayName", "Yaw (was Z)"));
+				ChildHandle->SetPropertyDisplayName(LOCTEXT("YawDisplayName", "Yaw"));
 				ChildHandle->SetToolTipText(LOCTEXT("YawToolTip", "Yaw (degrees) around Up (was Z) axis"));
 			}
 
