Update a few more transform related call sites to use AxisDisplayInfo

diff --git a/Engine/Plugins/Experimental/EditorDataStorageFeatures/Source/TedsOutliner/Private/Widgets/TransformHeadsUpWidget.cpp b/Engine/Plugins/Experimental/EditorDataStorageFeatures/Source/TedsOutliner/Private/Widgets/TransformHeadsUpWidget.cpp
--- a/Engine/Plugins/Experimental/EditorDataStorageFeatures/Source/TedsOutliner/Private/Widgets/TransformHeadsUpWidget.cpp
+++ b/Engine/Plugins/Experimental/EditorDataStorageFeatures/Source/TedsOutliner/Private/Widgets/TransformHeadsUpWidget.cpp
@@ -7,6 +7,7 @@
 #include "Elements/Columns/TypedElementTransformColumns.h"
 #include "Elements/Framework/TypedElementQueryBuilder.h"
 #include "Elements/Interfaces/TypedElementDataStorageUiInterface.h"
+#include "Misc/AxisDisplayInfo.h"
 #include "Widgets/SCompoundWidget.h"
 #include "Widgets/SBoxPanel.h"
 #include "Widgets/Images/SImage.h"
@@ -81,24 +82,24 @@ public:
 		[
 			SNew(SImage)
 				.Image(FAppStyle::GetBrush("EditorViewport.ScaleMode"))
-				.ColorAndOpacity(SNumericEntryBox<double>::RedLabelBackgroundColor)
-				.ToolTipText(LOCTEXT("NegativeXScaleTooltip", "Has Negative X Scale"))
+				.ColorAndOpacity(AxisDisplayInfo::GetAxisColor(EAxisList::X))
+				.ToolTipText(FText::Format(LOCTEXT("NegativeXScaleTooltip", "Has Negative {0} Scale"), AxisDisplayInfo::GetAxisDisplayName(EAxisList::X)))
 				.Visibility_Lambda([this]() { return EnumHasAllFlags(AbnormalTransformFlags, EAbnormalTransformTypes::NegativeXScale) ? EVisibility::Visible : EVisibility::Collapsed; })
 		];
 		AddSlot().AutoWidth().Padding(IconPadding)
 		[
 			SNew(SImage)
 				.Image(FAppStyle::GetBrush("EditorViewport.ScaleMode"))
-				.ColorAndOpacity(SNumericEntryBox<double>::GreenLabelBackgroundColor)
-				.ToolTipText(LOCTEXT("NegativeYScaleTooltip", "Has Negative Y Scale"))
+				.ColorAndOpacity(AxisDisplayInfo::GetAxisColor(EAxisList::Y))
+				.ToolTipText(FText::Format(LOCTEXT("NegativeYScaleTooltip", "Has Negative {0} Scale"), AxisDisplayInfo::GetAxisDisplayName(EAxisList::Y)))
 				.Visibility_Lambda([this]() { return EnumHasAllFlags(AbnormalTransformFlags, EAbnormalTransformTypes::NegativeYScale) ? EVisibility::Visible : EVisibility::Collapsed; })
 		];
 		AddSlot().AutoWidth().Padding(IconPadding)
 		[
 			SNew(SImage)
 				.Image(FAppStyle::GetBrush("EditorViewport.ScaleMode"))
-				.ColorAndOpacity(SNumericEntryBox<double>::BlueLabelBackgroundColor)
-				.ToolTipText(LOCTEXT("NegativeZScaleTooltip", "Has Negative Z Scale"))
+				.ColorAndOpacity(AxisDisplayInfo::GetAxisColor(EAxisList::Z))
+				.ToolTipText(FText::Format(LOCTEXT("NegativeZScaleTooltip", "Has Negative {0} Scale"), AxisDisplayInfo::GetAxisDisplayName(EAxisList::Z)))
 				.Visibility_Lambda([this]() { return EnumHasAllFlags(AbnormalTransformFlags, EAbnormalTransformTypes::NegativeZScale) ? EVisibility::Visible : EVisibility::Collapsed; })
 		];
 		AddSlot().AutoWidth().Padding(IconPadding)
diff --git a/Engine/Plugins/Experimental/MeshModelingToolsetExp/Source/ModelingEditorUI/Private/STransformGizmoNumericalUIOverlay.cpp b/Engine/Plugins/Experimental/MeshModelingToolsetExp/Source/ModelingEditorUI/Private/STransformGizmoNumericalUIOverlay.cpp
--- a/Engine/Plugins/Experimental/MeshModelingToolsetExp/Source/ModelingEditorUI/Private/STransformGizmoNumericalUIOverlay.cpp
+++ b/Engine/Plugins/Experimental/MeshModelingToolsetExp/Source/ModelingEditorUI/Private/STransformGizmoNumericalUIOverlay.cpp
@@ -9,6 +9,7 @@
 #include "Framework/Commands/UICommandInfo.h" // EUserInterfaceActionType
 #include "Framework/MultiBox/MultiBoxBuilder.h" // FMenuBuilder
 #include "InteractiveToolsContext.h"
+#include "Misc/AxisDisplayInfo.h"
 #include "ModelingWidgets/SDraggableBox.h" // SDraggableBoxOverlay
 #include "ModelingWidgets/STickDelegateWidget.h"
 #include "Transforms/TransformGizmoDataBinder.h"
@@ -115,15 +116,10 @@ void STransformGizmoNumericalUIOverlay::OnTickWidget(const FGeometry& AllottedGe
 SNumericEntryBox<double>::FArguments STransformGizmoNumericalUIOverlay::CreateComponentBaseArgs(FVector3d& TargetVector, int ComponentIndex, TAttribute<bool> IsEnabled)
 {
 	// Can be indexed via ComponentIndex:
-	const FText Tooltips[3] {
-		LOCTEXT("X_ToolTip", "X: {0}"),
-		LOCTEXT("Y_ToolTip", "Y: {0}"),
-		LOCTEXT("Z_ToolTip", "Z: {0}")
-	};
-	const FLinearColor Colors[3] {
-		SNumericEntryBox<double>::RedLabelBackgroundColor,
-		SNumericEntryBox<double>::GreenLabelBackgroundColor,
-		SNumericEntryBox<double>::BlueLabelBackgroundColor
+	const EAxisList::Type Axes[3] {
+		EAxisList::X,
+		EAxisList::Y,
+		EAxisList::Z
 	};
 
 	SNumericEntryBox<double>::FArguments Args;
@@ -155,9 +151,9 @@ SNumericEntryBox<double>::FArguments STransformGizmoNumericalUIOverlay::CreateCo
 				DataBinderWeakPtr.Pin()->EndDataEditSequence();
 			}
 		}))
-		.ToolTipText(MakeAttributeLambda([this, &TargetVector, Tooltip = Tooltips[ComponentIndex], ComponentIndex]
+		.ToolTipText(MakeAttributeLambda([this, &TargetVector, Axis = Axes[ComponentIndex], ComponentIndex]
 		{
-			return FText::Format(Tooltip, TargetVector[ComponentIndex]);
+			return FText::Format(LOCTEXT("Component_ToolTip", "{0}: {1}"), AxisDisplayInfo::GetAxisDisplayName(Axis), TargetVector[ComponentIndex]);
 		}))
 		.IsEnabled(IsEnabled)
 
@@ -166,7 +162,7 @@ SNumericEntryBox<double>::FArguments STransformGizmoNumericalUIOverlay::CreateCo
 
 		.Label()
 		[
-			SNumericEntryBox<double>::BuildNarrowColorLabel(Colors[ComponentIndex])
+			SNumericEntryBox<double>::BuildNarrowColorLabel(AxisDisplayInfo::GetAxisColor(Axes[ComponentIndex]))
 		]
 		;
 
diff --git a/Engine/Source/Editor/DetailCustomizations/Private/PhysicsConstraintComponentDetails.cpp b/Engine/Source/Editor/DetailCustomizations/Private/PhysicsConstraintComponentDetails.cpp
--- a/Engine/Source/Editor/DetailCustomizations/Private/PhysicsConstraintComponentDetails.cpp
+++ b/Engine/Source/Editor/DetailCustomizations/Private/PhysicsConstraintComponentDetails.cpp
@@ -28,6 +28,7 @@
 #include "Math/Rotator.h"
 #include "Misc/AssertionMacros.h"
 #include "Misc/Attribute.h"
+#include "Misc/AxisDisplayInfo.h"
 #include "Misc/EnumClassFlags.h"
 #include "PhysicsEngine/ConstraintDrives.h"
 #include "PhysicsEngine/ConstraintInstance.h"
@@ -156,17 +157,17 @@ void FConstraintTransformCustomization::MakeRotationRow(TSharedRef<class IProper
 		+ SHorizontalBox::Slot()
 		.Padding(FMargin(0.0f, 2.0f, 3.0f, 2.0f))
 		[
-			MakeNumericProxyWidget<FRotator, FRotator::FReal>(InPriAxisPropertyHandle, InSecAxisPropertyHandle, CachedRotationRoll, FText::Format(ComponentToolTipFormatText, LOCTEXT("RotationRoll", "Roll (X)"), FrameLabelText), true, SNumericEntryBox<FRotator::FReal>::RedLabelBackgroundColor)
+			MakeNumericProxyWidget<FRotator, FRotator::FReal>(InPriAxisPropertyHandle, InSecAxisPropertyHandle, CachedRotationRoll, FText::Format(ComponentToolTipFormatText, LOCTEXT("RotationRoll", "Roll"), FrameLabelText), true, AxisDisplayInfo::GetAxisColor(EAxisList::X))
 		]
 		+ SHorizontalBox::Slot()
 		.Padding(FMargin(0.0f, 2.0f, 3.0f, 2.0f))
 		[
-			MakeNumericProxyWidget<FRotator, FRotator::FReal>(InPriAxisPropertyHandle, InSecAxisPropertyHandle, CachedRotationPitch, FText::Format(ComponentToolTipFormatText, LOCTEXT("RotationPitch", "Pitch (Y)"), FrameLabelText), true, SNumericEntryBox<FRotator::FReal>::GreenLabelBackgroundColor)
+			MakeNumericProxyWidget<FRotator, FRotator::FReal>(InPriAxisPropertyHandle, InSecAxisPropertyHandle, CachedRotationPitch, FText::Format(ComponentToolTipFormatText, LOCTEXT("RotationPitch", "Pitch"), FrameLabelText), true, AxisDisplayInfo::GetAxisColor(EAxisList::Y))
 		]
 		+ SHorizontalBox::Slot()
 		.Padding(FMargin(0.0f, 2.0f, 0.0f, 2.0f))
 		[
-			MakeNumericProxyWidget<FRotator, FRotator::FReal>(InPriAxisPropertyHandle, InSecAxisPropertyHandle, CachedRotationYaw, FText::Format(ComponentToolTipFormatText, LOCTEXT("RotationYaw", "Yaw (Z)"), FrameLabelText), true, SNumericEntryBox<FRotator::FReal>::BlueLabelBackgroundColor)
+			MakeNumericProxyWidget<FRotator, FRotator::FReal>(InPriAxisPropertyHandle, InSecAxisPropertyHandle, CachedRotationYaw, FText::Format(ComponentToolTipFormatText, LOCTEXT("RotationYaw", "Yaw"), FrameLabelText), true, AxisDisplayInfo::GetAxisColor(EAxisList::Z))
 		]
 	];
 }
@@ -212,17 +213,17 @@ void FConstraintTransformCustomization::MakePositionRow(TSharedRef<class IProper
 		+ SHorizontalBox::Slot()
 		.Padding(FMargin(0.0f, 2.0f, 3.0f, 2.0f))
 		[
-			MakeNumericProxyWidget<FVector, FVector::FReal>(InPositionPropertyHandle, CachedPositionX, FText::Format(ComponentToolTipFormatText, LOCTEXT("PositionX", "X"), FrameLabelText), SNumericEntryBox<FRotator::FReal>::RedLabelBackgroundColor)
+			MakeNumericProxyWidget<FVector, FVector::FReal>(InPositionPropertyHandle, CachedPositionX, FText::Format(ComponentToolTipFormatText, AxisDisplayInfo::GetAxisDisplayName(EAxisList::X), FrameLabelText), AxisDisplayInfo::GetAxisColor(EAxisList::X))
 		]
 		+ SHorizontalBox::Slot()
 		.Padding(FMargin(0.0f, 2.0f, 3.0f, 2.0f))
 		[
-			MakeNumericProxyWidget<FVector, FVector::FReal>(InPositionPropertyHandle, CachedPositionY, FText::Format(ComponentToolTipFormatText, LOCTEXT("PositionY", "Y"), FrameLabelText), SNumericEntryBox<FRotator::FReal>::GreenLabelBackgroundColor)
+			MakeNumericProxyWidget<FVector, FVector::FReal>(InPositionPropertyHandle, CachedPositionY, FText::Format(ComponentToolTipFormatText, AxisDisplayInfo::GetAxisDisplayName(EAxisList::Y), FrameLabelText), AxisDisplayInfo::GetAxisColor(EAxisList::Y))
 		]
 		+ SHorizontalBox::Slot()
 		.Padding(FMargin(0.0f, 2.0f, 0.0f, 2.0f))
 		[
-			MakeNumericProxyWidget<FVector, FVector::FReal>(InPositionPropertyHandle, CachedPositionZ, FText::Format(ComponentToolTipFormatText, LOCTEXT("PositionZ", "Z"), FrameLabelText), SNumericEntryBox<FRotator::FReal>::BlueLabelBackgroundColor)
+			MakeNumericProxyWidget<FVector, FVector::FReal>(InPositionPropertyHandle, CachedPositionZ, FText::Format(ComponentToolTipFormatText, AxisDisplayInfo::GetAxisDisplayName(EAxisList::Z), FrameLabelText), AxisDisplayInfo::GetAxisColor(EAxisList::Z))
 		]
 	];
 }
diff --git a/Engine/Source/Runtime/Core/Public/Misc/AxisDisplayInfo.h b/Engine/Source/Runtime/Core/Public/Misc/AxisDisplayInfo.h
--- a/Engine/Source/Runtime/Core/Public/Misc/AxisDisplayInfo.h
+++ b/Engine/Source/Runtime/Core/Public/Misc/AxisDisplayInfo.h
@@ -11,6 +11,9 @@ namespace AxisDisplayInfo
 {
 	/**
 	 * Gets the display name to use for the given axis.
+	 * 
+	 * The return value may be driven by the config system and therefore
+	 * it should not be cached until the editor has finished loading.
 	 *
 	 * @param Axis The axis whose display name you are requesting
 	 * @return The display name corresponding to the given axis
@@ -19,6 +22,9 @@ namespace AxisDisplayInfo
 
 	/**
 	 * Gets the short display name to use for the given axis.
+	 * 
+	 * The return value may be driven by the config system and therefore
+	 * it should not be cached until the editor has finished loading.
 	 *
 	 * @param Axis The axis whose short display name you are requesting
 	 * @return The short display name corresponding to the given axis
@@ -28,6 +34,9 @@ namespace AxisDisplayInfo
 	/**
 	 * Gets the color to use for the given axis.
 	 * 
+	 * The return value may be driven by the config system and therefore
+	 * it should not be cached until the editor has finished loading.
+	 * 
 	 * @param Axis The axis whose color you are requesting
 	 * @return The color corresponding to the given axis
 	 */
@@ -35,6 +44,9 @@ namespace AxisDisplayInfo
 
 	/**
 	 * Gets whether or not the engine uses forward/right/up nomenclature for display names.
+	 * 
+	 * The return value may be driven by the config system and therefore
+	 * it should not be cached until the editor has finished loading.
 	 */
 	CORE_API bool UseForwardRightUpDisplayNames();
 };
