[Truncation Warnings] Fix and enable DetailsCustomization module truncation warnings

#rnx

diff --git a/Engine/Plugins/Experimental/Avalanche/Source/AvalancheSequencer/Private/EaseCurveTool/Widgets/SAvaEaseCurvePreview.cpp b/Engine/Plugins/Experimental/Avalanche/Source/AvalancheSequencer/Private/EaseCurveTool/Widgets/SAvaEaseCurvePreview.cpp
--- a/Engine/Plugins/Experimental/Avalanche/Source/AvalancheSequencer/Private/EaseCurveTool/Widgets/SAvaEaseCurvePreview.cpp
+++ b/Engine/Plugins/Experimental/Avalanche/Source/AvalancheSequencer/Private/EaseCurveTool/Widgets/SAvaEaseCurvePreview.cpp
@@ -167,7 +167,7 @@ void SAvaEaseCurvePreview::Tick(const FGeometry& InAllottedGeometry, const doubl
 			CurrentAnimateTime += InDeltaTime;
 			CurrentAnimateTime = FMath::Min(1.f, CurrentAnimateTime);
 
-			const FTrackScaleInfo ScaleInfo(0.f, 1.f, 0.f, 1.f, FVector2D(PreviewSize)/*InAllottedGeometry.GetLocalSize()*/);
+			const FTrackScaleInfo ScaleInfo(0.f, 1.f, 0.f, 1.f, FVector2f(PreviewSize)/*InAllottedGeometry.GetLocalSize()*/);
 			const FAvaEaseCurveTangents DrawTangents = Tangents.Get(FAvaEaseCurveTangents());
 			const float CurrentValue = WeightedEvalForTwoKeys(DrawTangents, CurrentAnimateTime);
 
@@ -207,7 +207,7 @@ int32 SAvaEaseCurvePreview::OnPaint(const FPaintArgs& InArgs, const FGeometry& I
 	const bool bEnabled = ShouldBeEnabled(bInParentEnabled);
 	const ESlateDrawEffect DrawEffects = bEnabled ? ESlateDrawEffect::None : ESlateDrawEffect::DisabledEffect;
 	const FVector2D LocalSize = InAllottedGeometry.GetLocalSize();
-	const FTrackScaleInfo ScaleInfo(0.f, 1.f, 0.f, 1.f, FVector2D(PreviewSize));
+	const FTrackScaleInfo ScaleInfo(0.f, 1.f, 0.f, 1.f, FVector2f(PreviewSize));
 
 	const float ZeroOutputY = ScaleInfo.OutputToLocalY(0.f);
 	const float OneOutputY = ScaleInfo.OutputToLocalY(1.f);
@@ -318,7 +318,7 @@ int32 SAvaEaseCurvePreview::OnPaint(const FPaintArgs& InArgs, const FGeometry& I
 void SAvaEaseCurvePreview::MakeCurveData(const FGeometry& InGeometry, TArray<FVector2D>& OutCurvePoints, TArray<FLinearColor>& OutCurveColors) const
 {
 	const FAvaEaseCurveTangents DrawTangents = Tangents.Get(FAvaEaseCurveTangents());
-	const FTrackScaleInfo ScaleInfo(0.f, 1.f, 0.f, 1.f, FVector2D(PreviewSize)/*InAllottedGeometry.GetLocalSize()InGeometry.GetLocalSize()*/);
+	const FTrackScaleInfo ScaleInfo(0.f, 1.f, 0.f, 1.f, FVector2f(PreviewSize)/*InAllottedGeometry.GetLocalSize()InGeometry.GetLocalSize()*/);
 
 	// Clamp to screen to avoid massive slowdown when zoomed in
 	float StartX = FMath::Max(ScaleInfo.InputToLocalX(0.f), 0.0f);
diff --git a/Engine/Source/Editor/DetailCustomizations/DetailCustomizations.Build.cs b/Engine/Source/Editor/DetailCustomizations/DetailCustomizations.Build.cs
--- a/Engine/Source/Editor/DetailCustomizations/DetailCustomizations.Build.cs
+++ b/Engine/Source/Editor/DetailCustomizations/DetailCustomizations.Build.cs
@@ -6,6 +6,9 @@ public class DetailCustomizations : ModuleRules
 {
 	public DetailCustomizations(ReadOnlyTargetRules Target) : base(Target)
 	{
+		// Enable truncation warnings in this module
+		CppCompileWarningSettings.UnsafeTypeCastWarningLevel = WarningLevel.Warning;
+
 		PrivateDependencyModuleNames.AddRange(
 			new string[] {
 				"AdvancedWidgets",
@@ -28,9 +31,9 @@ public class DetailCustomizations : ModuleRules
 				"MovieSceneCapture",
 				"MovieSceneTools",
 				"MovieSceneTracks",
-                "Sequencer",
-                "MovieScene",
-                "TimeManagement",
+				"Sequencer",
+				"MovieScene",
+				"TimeManagement",
 				"SharedSettingsWidgets",
 				"ContentBrowser",
 				"BlueprintGraph",
@@ -58,7 +61,7 @@ public class DetailCustomizations : ModuleRules
 				"AdvancedPreviewScene",
 				"AudioSettingsEditor",
 				"HeadMountedDisplay",
-                "DataTableEditor",
+				"DataTableEditor",
 				"ToolMenus",
 				"PhysicsCore",
 				"RenderCore",
@@ -70,7 +73,7 @@ public class DetailCustomizations : ModuleRules
 			}
 		);
 
-        PrivateIncludePathModuleNames.AddRange(
+		PrivateIncludePathModuleNames.AddRange(
 			new string[] {
 				"Media",
 				"LandscapeEditor",
diff --git a/Engine/Source/Editor/DetailCustomizations/Private/ActorDetails.cpp b/Engine/Source/Editor/DetailCustomizations/Private/ActorDetails.cpp
--- a/Engine/Source/Editor/DetailCustomizations/Private/ActorDetails.cpp
+++ b/Engine/Source/Editor/DetailCustomizations/Private/ActorDetails.cpp
@@ -352,8 +352,8 @@ TSharedRef<SWidget> FActorDetails::OnGetConvertContent()
 
 	return
 		SNew(SBox)
-		.WidthOverride(280)
-		.MaxDesiredHeight(500)
+		.WidthOverride(280.f)
+		.MaxDesiredHeight(500.f)
 		[
 			ClassPicker
 		];
@@ -369,7 +369,7 @@ TSharedRef<SWidget> FActorDetails::MakeConvertMenu( const FSelectedActorInfo& Se
 	UClass* RootConversionClass = GetConversionRoot(SelectedActorInfo.SelectionClass);
 	return
 		SNew(SComboButton)
-		.ContentPadding(2)
+		.ContentPadding(2.f)
 		.IsEnabled(RootConversionClass != NULL)
 		.Visibility(this, &FActorDetails::GetConvertMenuVisibility)
 		.OnGetMenuContent(this, &FActorDetails::OnGetConvertContent)
@@ -542,12 +542,12 @@ void FActorDetails::AddEventsCategory(IDetailLayoutBuilder& DetailBuilder)
 				+ SHorizontalBox::Slot()
 				.HAlign(HAlign_Left)
 				.VAlign(VAlign_Center)
-				.Padding(0)
+				.Padding(0.f)
 				[
 					// A "Plus" button to add a binding. For dynamic delegates on the CDO, you can always
 					// make a new binding, so always display the "Plus"
 					SNew(SButton)
-					.ContentPadding(FMargin(3.0, 2.0))
+					.ContentPadding(FMargin(3.0f, 2.0f))
 					.HAlign(HAlign_Center)
 					.OnClicked(this, &FActorDetails::HandleAddOrViewEventForVariable, Blueprint, Property)
 					[
@@ -696,7 +696,7 @@ void FActorDetails::AddActorCategory( IDetailLayoutBuilder& DetailBuilder, const
 				.ToolTipText(Tooltip)
 		]
 		.ValueContent()
-		.MaxDesiredWidth(0)
+		.MaxDesiredWidth(0.f)
 		[
 			SNew(STextBlock)
 				.Text(LevelDescription)
diff --git a/Engine/Source/Editor/DetailCustomizations/Private/AmbientSoundDetails.cpp b/Engine/Source/Editor/DetailCustomizations/Private/AmbientSoundDetails.cpp
--- a/Engine/Source/Editor/DetailCustomizations/Private/AmbientSoundDetails.cpp
+++ b/Engine/Source/Editor/DetailCustomizations/Private/AmbientSoundDetails.cpp
@@ -73,7 +73,7 @@ void FAmbientSoundDetails::CustomizeDetails( IDetailLayoutBuilder& DetailBuilder
 		[
 			SNew(SVerticalBox)
 			+ SVerticalBox::Slot()
-			.Padding( 0, 2.0f, 0, 0 )
+			.Padding( 0.f, 2.0f, 0.f, 0.f)
 			.FillHeight(1.0f)
 			.VAlign( VAlign_Center )
 			[
@@ -101,7 +101,7 @@ void FAmbientSoundDetails::CustomizeDetails( IDetailLayoutBuilder& DetailBuilder
 						SNew( SComboButton )
 						.OnGetMenuContent( this, &FAmbientSoundDetails::OnGetSoundCueTemplates )
 						.VAlign( VAlign_Center )
-						.ContentPadding(2)
+						.ContentPadding(2.f)
 						.ButtonContent()
 						[
 							SNew( STextBlock )
diff --git a/Engine/Source/Editor/DetailCustomizations/Private/AnimMontageSegmentDetails.cpp b/Engine/Source/Editor/DetailCustomizations/Private/AnimMontageSegmentDetails.cpp
--- a/Engine/Source/Editor/DetailCustomizations/Private/AnimMontageSegmentDetails.cpp
+++ b/Engine/Source/Editor/DetailCustomizations/Private/AnimMontageSegmentDetails.cpp
@@ -373,7 +373,7 @@ void FAnimMontageSegmentDetails::CustomizeDetails( IDetailLayoutBuilder& DetailB
 							{									
 								if(const IAnimationDataModel* DataModel = AnimSegment->GetAnimReference()->GetDataModel())
 								{
-									const float NewPlayLength = DataModel->GetPlayLength();
+									const double NewPlayLength = DataModel->GetPlayLength();
 									PlayLengthHandle->SetValue(NewPlayLength);
 								}
 							}
@@ -397,7 +397,7 @@ void FAnimMontageSegmentDetails::CustomizeDetails( IDetailLayoutBuilder& DetailB
 							{
 								if(const IAnimationDataModel* DataModel = AnimSegment->GetAnimReference()->GetDataModel())
 								{
-									const float NewPlayLength = DataModel->GetPlayLength();
+									const double NewPlayLength = DataModel->GetPlayLength();
 									AnimEndTimeProperty->SetValue(NewPlayLength);
 								}
 							}
@@ -745,7 +745,7 @@ float SAnimationSegmentViewport::GetViewMinInput() const
 		}
 		else if (PreviewComponent->GetAnimInstance() != NULL)
 		{
-			return FMath::Max<float>((float)(PreviewComponent->GetAnimInstance()->LifeTimer - 30.0), 0.0f);
+			return FMath::Max(static_cast<float>(PreviewComponent->GetAnimInstance()->LifeTimer - 30.0), 0.0f);
 		}
 	}
 
@@ -762,7 +762,7 @@ float SAnimationSegmentViewport::GetViewMaxInput() const
 		}
 		else if (PreviewComponent->GetAnimInstance() != NULL)
 		{
-			return PreviewComponent->GetAnimInstance()->LifeTimer;
+			return static_cast<float>(PreviewComponent->GetAnimInstance()->LifeTimer);
 		}
 	}
 
diff --git a/Engine/Source/Editor/DetailCustomizations/Private/AnimSequenceDetails.cpp b/Engine/Source/Editor/DetailCustomizations/Private/AnimSequenceDetails.cpp
--- a/Engine/Source/Editor/DetailCustomizations/Private/AnimSequenceDetails.cpp
+++ b/Engine/Source/Editor/DetailCustomizations/Private/AnimSequenceDetails.cpp
@@ -164,7 +164,7 @@ void FAnimSequenceDetails::CustomizeDetails(IDetailLayoutBuilder& DetailBuilder)
 		.OnComboBoxOpening(this, &FAnimSequenceDetails::OnRetargetSourceComboOpening)
 		.InitiallySelectedItem(InitialSelected)
 		.IsEnabled( FSlateApplication::Get().GetNormalExecutionAttribute() )
-		.ContentPadding(0)
+		.ContentPadding(0.f)
 		.Content()
 		[
 			SNew( STextBlock )
@@ -752,7 +752,7 @@ float SAnimationRefPoseViewport::GetViewMaxInput() const
 		}
 		else if (PreviewComponent->GetAnimInstance() != NULL)
 		{
-			return PreviewComponent->GetAnimInstance()->LifeTimer;
+			return static_cast<float>(PreviewComponent->GetAnimInstance()->LifeTimer);
 		}
 	}
 
@@ -766,8 +766,8 @@ TArray<float> SAnimationRefPoseViewport::GetBars() const
 	{
 		int32 RefFrameIndex;
 		RefFrameIndexPropertyHandle->GetValue(RefFrameIndex);
-		float Fraction = (PreviewAnimationSequence->GetNumberOfSampledKeys() > 0)? FMath::Clamp<float>((float)RefFrameIndex/(float)PreviewAnimationSequence->GetNumberOfSampledKeys(), 0.f, 1.f) : 0.f;
-		Bars.Add(PreviewAnimationSequence->GetPlayLength() * Fraction);
+		float Fraction = (PreviewAnimationSequence->GetNumberOfSampledKeys() > 0)? FMath::Clamp((float)RefFrameIndex/(float)PreviewAnimationSequence->GetNumberOfSampledKeys(), 0.f, 1.f) : 0.f;
+		Bars.Add(static_cast<float>(PreviewAnimationSequence->GetPlayLength() * Fraction));
 	}
 	else
 	{
@@ -780,7 +780,7 @@ void SAnimationRefPoseViewport::OnBarDrag(int32 Index, float Position)
 {
 	if (PreviewAnimationSequence)
 	{
-		int RefFrameIndex = FMath::Clamp<int>(PreviewAnimationSequence->GetPlayLength() > 0.0f? (int)(Position * (float)PreviewAnimationSequence->GetNumberOfSampledKeys() / PreviewAnimationSequence->GetPlayLength() + 0.5f) : 0.0f, 0, PreviewAnimationSequence->GetNumberOfSampledKeys() - 1);
+		int RefFrameIndex = FMath::Clamp(PreviewAnimationSequence->GetPlayLength() > 0.0 ? (int)(Position * (float)PreviewAnimationSequence->GetNumberOfSampledKeys() / static_cast<float>(PreviewAnimationSequence->GetPlayLength()) + 0.5f) : static_cast<int>(0), static_cast<int>(0), static_cast<int>(PreviewAnimationSequence->GetNumberOfSampledKeys() - 1));
 		RefFrameIndexPropertyHandle->SetValue(RefFrameIndex);
 	}
 }
diff --git a/Engine/Source/Editor/DetailCustomizations/Private/AutoReimportDirectoryCustomization.cpp b/Engine/Source/Editor/DetailCustomizations/Private/AutoReimportDirectoryCustomization.cpp
--- a/Engine/Source/Editor/DetailCustomizations/Private/AutoReimportDirectoryCustomization.cpp
+++ b/Engine/Source/Editor/DetailCustomizations/Private/AutoReimportDirectoryCustomization.cpp
@@ -62,7 +62,7 @@ void FAutoReimportWildcardCustomization::CustomizeHeader(TSharedRef<IPropertyHan
 
 		+SHorizontalBox::Slot()
 		.VAlign(VAlign_Center)
-		.Padding(FMargin(0,0,4.f,0))
+		.Padding(FMargin(0.f,0.f,4.f,0.f))
 		[
 			SNew(SEditableTextBox)
 			.Text(this, &FAutoReimportWildcardCustomization::GetWildcardText)
@@ -146,7 +146,7 @@ void FAutoReimportDirectoryCustomization::CustomizeHeader(TSharedRef<IPropertyHa
 
 		+SHorizontalBox::Slot()
 		.VAlign(VAlign_Center)
-		.Padding(FMargin(0,0,4.f,0))
+		.Padding(FMargin(0.f,0.f,4.f,0.f))
 		[
 			SNew(SEditableTextBox)
 			.Text(this, &FAutoReimportDirectoryCustomization::GetDirectoryText)
@@ -233,9 +233,9 @@ TSharedRef<SWidget> FAutoReimportDirectoryCustomization::GetPathPickerContent()
 	FContentBrowserModule& ContentBrowserModule = FModuleManager::LoadModuleChecked<FContentBrowserModule>("ContentBrowser");
 
 	return SNew(SBox)
-		.WidthOverride(300)
-		.HeightOverride(500)
-		.Padding(4)
+		.WidthOverride(300.f)
+		.HeightOverride(500.f)
+		.Padding(4.f)
 		[
 			SNew(SVerticalBox)
 
diff --git a/Engine/Source/Editor/DetailCustomizations/Private/BodyInstanceCustomization.cpp b/Engine/Source/Editor/DetailCustomizations/Private/BodyInstanceCustomization.cpp
--- a/Engine/Source/Editor/DetailCustomizations/Private/BodyInstanceCustomization.cpp
+++ b/Engine/Source/Editor/DetailCustomizations/Private/BodyInstanceCustomization.cpp
@@ -55,7 +55,7 @@ class USceneComponent;
 
 #define LOCTEXT_NAMESPACE "BodyInstanceCustomization"
 
-#define RowWidth_Customization 50
+#define RowWidth_Customization 50.f
 
 ////////////////////////////////////////////////////////////////
 FBodyInstanceCustomization::FBodyInstanceCustomization()
@@ -434,7 +434,7 @@ void FBodyInstanceCustomization::CreateCustomCollisionSetup(TSharedRef<IProperty
 			.OnGenerateWidget(this, &FBodyInstanceCustomization::MakeObjectTypeComboWidget)
 			.OnSelectionChanged(this, &FBodyInstanceCustomization::OnObjectTypeChanged)
 			.InitiallySelectedItem(ObjectTypeComboList[IndexSelected])
-			.ContentPadding(2)
+			.ContentPadding(2.f)
 			.Content()
 			[
 				SNew(STextBlock)
@@ -449,8 +449,8 @@ void FBodyInstanceCustomization::CreateCustomCollisionSetup(TSharedRef<IProperty
 	.IsEnabled(CustomCollisionEnabled)
 	.Visibility(CustomCollisionVisibility)
 	.ValueContent()
-	.MaxDesiredWidth(0)
-	.MinDesiredWidth(0)
+	.MaxDesiredWidth(0.f)
+	.MinDesiredWidth(0.f)
 	[
 		SNew(SHorizontalBox)
 		+SHorizontalBox::Slot()
@@ -738,7 +738,7 @@ void FBodyInstanceCustomization::OnObjectTypeChanged( TSharedPtr<FString> NewSel
 	if (SelectInfo != ESelectInfo::Direct)
 	{
 		FString NewValue = *NewSelection.Get();
-		uint8 NewEnumVal = ECC_WorldStatic;
+		ECollisionChannel NewEnumVal = ECC_WorldStatic;
 		// find index of NewValue
 		for (auto Iter = ObjectTypeComboList.CreateConstIterator(); Iter; ++Iter)
 		{
@@ -1134,7 +1134,7 @@ void FBodyInstanceCustomization::SetResponse(int32 ValidIndex, ECollisionRespons
 
 ECheckBoxState FBodyInstanceCustomization::IsCollisionChannelChecked( int32 ValidIndex, ECollisionResponse InCollisionResponse) const
 {
-	TArray<uint8> CollisionResponses;
+	TArray<ECollisionResponse> CollisionResponses;
 
 	if ( ValidCollisionChannels.IsValidIndex(ValidIndex) )
 	{
diff --git a/Engine/Source/Editor/DetailCustomizations/Private/BodySetupDetails.cpp b/Engine/Source/Editor/DetailCustomizations/Private/BodySetupDetails.cpp
--- a/Engine/Source/Editor/DetailCustomizations/Private/BodySetupDetails.cpp
+++ b/Engine/Source/Editor/DetailCustomizations/Private/BodySetupDetails.cpp
@@ -93,7 +93,7 @@ void FBodySetupDetails::CustomizeDetails( IDetailLayoutBuilder& DetailBuilder )
 			SNew(SCheckBox)
 			.IsChecked_Lambda([WeakCollisionReponseHandle]() 
 			{
-				uint8 Value = 0;
+				int Value = 0;
 				if(WeakCollisionReponseHandle.IsValid() && WeakCollisionReponseHandle.Pin()->GetValue(Value) == FPropertyAccess::Success)
 				{
 					return (Value == EBodyCollisionResponse::BodyCollision_Enabled) ? ECheckBoxState::Checked : ECheckBoxState::Unchecked;
@@ -105,7 +105,7 @@ void FBodySetupDetails::CustomizeDetails( IDetailLayoutBuilder& DetailBuilder )
 			{
 				if(WeakCollisionReponseHandle.IsValid())
 				{
-					uint8 Value = InCheckBoxState == ECheckBoxState::Checked ? EBodyCollisionResponse::BodyCollision_Enabled : EBodyCollisionResponse::BodyCollision_Disabled;
+					int Value = InCheckBoxState == ECheckBoxState::Checked ? EBodyCollisionResponse::BodyCollision_Enabled : EBodyCollisionResponse::BodyCollision_Disabled;
 					WeakCollisionReponseHandle.Pin()->SetValue(Value);
 				}
 			})
diff --git a/Engine/Source/Editor/DetailCustomizations/Private/BoundsCopyComponentDetails.cpp b/Engine/Source/Editor/DetailCustomizations/Private/BoundsCopyComponentDetails.cpp
--- a/Engine/Source/Editor/DetailCustomizations/Private/BoundsCopyComponentDetails.cpp
+++ b/Engine/Source/Editor/DetailCustomizations/Private/BoundsCopyComponentDetails.cpp
@@ -88,7 +88,7 @@ void FBoundsCopyComponentDetailsCustomization::CustomizeDetails(IDetailLayoutBui
 		SNew(SButton)
 		.VAlign(VAlign_Center)
 		.HAlign(HAlign_Center)
-		.ContentPadding(2)
+		.ContentPadding(2.f)
 		.Text(LOCTEXT("Button_Copy", "Copy"))
 		.ToolTipText(LOCTEXT("Button_CopyRotation_Tooltip", "Set the actor rotation to match the source actor"))
 		.OnClicked(this, &FBoundsCopyComponentDetailsCustomization::SetRotation)
@@ -119,7 +119,7 @@ void FBoundsCopyComponentDetailsCustomization::CustomizeDetails(IDetailLayoutBui
 				SNew(SButton)
 				.VAlign(VAlign_Center)
 				.HAlign(HAlign_Center)
-				.ContentPadding(2)
+				.ContentPadding(2.f)
 				.Text(LOCTEXT("Button_Copy", "Copy"))
 				.ToolTipText(LOCTEXT("Button_CopyBounds_Tooltip", "Set the actor transform so that it includes the full bounds of the source actor"))
 				.OnClicked(this, &FBoundsCopyComponentDetailsCustomization::SetTransformToBounds)
diff --git a/Engine/Source/Editor/DetailCustomizations/Private/BrushDetails.cpp b/Engine/Source/Editor/DetailCustomizations/Private/BrushDetails.cpp
--- a/Engine/Source/Editor/DetailCustomizations/Private/BrushDetails.cpp
+++ b/Engine/Source/Editor/DetailCustomizations/Private/BrushDetails.cpp
@@ -148,13 +148,13 @@ void FBrushDetails::CustomizeDetails( IDetailLayoutBuilder& InDetailLayout )
 		.Font( IDetailLayoutBuilder::GetDetailFont() )
 	]
 	.ValueContent()
-	.MinDesiredWidth(105)
-	.MaxDesiredWidth(105)
+	.MinDesiredWidth(105.f)
+	.MaxDesiredWidth(105.f)
 	[
 		SNew(SComboButton)
 		.ToolTipText(LOCTEXT("BspModeBuildTooltip", "Rebuild this brush from a parametric builder."))
 		.OnGetMenuContent(this, &FBrushDetails::GenerateBuildMenuContent)
-		.ContentPadding(2)
+		.ContentPadding(2.f)
 		.ButtonContent()
 		[
 			SNew(STextBlock)
@@ -167,11 +167,11 @@ void FBrushDetails::CustomizeDetails( IDetailLayoutBuilder& InDetailLayout )
 	[
 		SNew(SHorizontalBox)
 		+ SHorizontalBox::Slot()
-		.FillWidth(1)
+		.FillWidth(1.f)
 		.Padding(1.0f)
 		[
 			SNew(SComboButton)
-			.ContentPadding(2)
+			.ContentPadding(2.f)
 			.ButtonContent()
 			[
 				SNew(STextBlock)
@@ -185,11 +185,11 @@ void FBrushDetails::CustomizeDetails( IDetailLayoutBuilder& InDetailLayout )
 			]
 		]
 		+ SHorizontalBox::Slot()
-		.FillWidth(1)
+		.FillWidth(1.f)
 		.Padding(1.0f)
 		[
 			SNew(SComboButton)
-			.ContentPadding(2)
+			.ContentPadding(2.f)
 			.Visibility(bHaveAVolume ? EVisibility::Collapsed : EVisibility::Visible)
 			.ButtonContent()
 			[
@@ -204,11 +204,11 @@ void FBrushDetails::CustomizeDetails( IDetailLayoutBuilder& InDetailLayout )
 			]
 		]
 		+ SHorizontalBox::Slot()
-		.FillWidth(1)
+		.FillWidth(1.f)
 		.Padding(1.0f)
 		[
 			SNew(SComboButton)
-			.ContentPadding(2)
+			.ContentPadding(2.f)
 			.Visibility(bHaveAVolume ? EVisibility::Collapsed : EVisibility::Visible)
 			.ButtonContent()
 			[
diff --git a/Engine/Source/Editor/DetailCustomizations/Private/CameraCropSettingsCustomization.cpp b/Engine/Source/Editor/DetailCustomizations/Private/CameraCropSettingsCustomization.cpp
--- a/Engine/Source/Editor/DetailCustomizations/Private/CameraCropSettingsCustomization.cpp
+++ b/Engine/Source/Editor/DetailCustomizations/Private/CameraCropSettingsCustomization.cpp
@@ -53,7 +53,7 @@ void FCameraCropSettingsCustomization::CustomizeHeader(TSharedRef<IPropertyHandl
 				.OnGenerateWidget(this, &FCameraCropSettingsCustomization::MakePresetComboWidget)
 				.OnSelectionChanged(this, &FCameraCropSettingsCustomization::OnPresetChanged)
 				.IsEnabled(FSlateApplication::Get().GetNormalExecutionAttribute())
-				.ContentPadding(2)
+				.ContentPadding(2.f)
 				.Content()
 				[
 					SNew(STextBlock)
diff --git a/Engine/Source/Editor/DetailCustomizations/Private/CameraFilmbackSettingsCustomization.cpp b/Engine/Source/Editor/DetailCustomizations/Private/CameraFilmbackSettingsCustomization.cpp
--- a/Engine/Source/Editor/DetailCustomizations/Private/CameraFilmbackSettingsCustomization.cpp
+++ b/Engine/Source/Editor/DetailCustomizations/Private/CameraFilmbackSettingsCustomization.cpp
@@ -79,7 +79,7 @@ void FCameraFilmbackSettingsCustomization::CustomizeHeader(TSharedRef<IPropertyH
 				.OnSelectionChanged(this, &FCameraFilmbackSettingsCustomization::OnPresetChanged)
 				.OnComboBoxOpening(this, &FCameraFilmbackSettingsCustomization::BuildPresetComboList)
 				.IsEnabled(this, &FCameraFilmbackSettingsCustomization::IsPresetEnabled)
-				.ContentPadding(2)
+				.ContentPadding(2.f)
 				.Content()
 				[
 					SNew(STextBlock)
diff --git a/Engine/Source/Editor/DetailCustomizations/Private/CameraFocusSettingsCustomization.cpp b/Engine/Source/Editor/DetailCustomizations/Private/CameraFocusSettingsCustomization.cpp
--- a/Engine/Source/Editor/DetailCustomizations/Private/CameraFocusSettingsCustomization.cpp
+++ b/Engine/Source/Editor/DetailCustomizations/Private/CameraFocusSettingsCustomization.cpp
@@ -158,7 +158,7 @@ void FCameraFocusSettingsCustomization::OnSceneDepthLocationSelected(FVector Pic
 			// if picked behind camera, don't set it
 			if ((CamToPickedLoc | CamForward) > 0.f)
 			{
-				float const FinalSceneDepth = CamToPickedLoc.ProjectOnToNormal(CamForward).Size();
+				float const FinalSceneDepth = static_cast<float>(CamToPickedLoc.ProjectOnToNormal(CamForward).Size());
 
 				const FScopedTransaction Transaction(LOCTEXT("PickedSceneDepth", "Pick Scene Depth"));
 				ensure(ManualFocusDistanceHandle->SetValue(FinalSceneDepth, EPropertyValueSetFlags::NotTransactable) == FPropertyAccess::Result::Success);
diff --git a/Engine/Source/Editor/DetailCustomizations/Private/CameraLensSettingsCustomization.cpp b/Engine/Source/Editor/DetailCustomizations/Private/CameraLensSettingsCustomization.cpp
--- a/Engine/Source/Editor/DetailCustomizations/Private/CameraLensSettingsCustomization.cpp
+++ b/Engine/Source/Editor/DetailCustomizations/Private/CameraLensSettingsCustomization.cpp
@@ -79,7 +79,7 @@ void FCameraLensSettingsCustomization::CustomizeHeader(TSharedRef<IPropertyHandl
 				.OnSelectionChanged(this, &FCameraLensSettingsCustomization::OnPresetChanged)
 				.OnComboBoxOpening(this, &FCameraLensSettingsCustomization::BuildPresetComboList)
 				.IsEnabled(FSlateApplication::Get().GetNormalExecutionAttribute())
-				.ContentPadding(2)
+				.ContentPadding(2.f)
 				.Content()
 				[
 					SNew(STextBlock)
diff --git a/Engine/Source/Editor/DetailCustomizations/Private/CollisionProfileDetails.cpp b/Engine/Source/Editor/DetailCustomizations/Private/CollisionProfileDetails.cpp
--- a/Engine/Source/Editor/DetailCustomizations/Private/CollisionProfileDetails.cpp
+++ b/Engine/Source/Editor/DetailCustomizations/Private/CollisionProfileDetails.cpp
@@ -62,15 +62,15 @@ DECLARE_DELEGATE_RetVal_TwoParams(bool, FOnValidateProfile, const FCollisionResp
 #define MAX_CUSTOMCOLLISION_CHANNEL (ECC_GameTraceChannel18-ECC_GameTraceChannel1+1)
 #define MAX_COLLISION_CHANNEL		32
 
-#define COLLIISION_COLUMN_WIDTH		50
+#define COLLIISION_COLUMN_WIDTH		50.f
 
-#define PROFILE_WINDOW_WIDTH		300
-#define PROFILE_WINDOW_HEIGHT		540
+#define PROFILE_WINDOW_WIDTH		300.f
+#define PROFILE_WINDOW_HEIGHT		540.f
 
-#define CHANNEL_WINDOW_WIDTH		200
-#define CHANNEL_WINDOW_HEIGHT		93
+#define CHANNEL_WINDOW_WIDTH		200.f
+#define CHANNEL_WINDOW_HEIGHT		93.f
 
-#define RowWidth_Customization 50
+#define RowWidth_Customization 50.f
 //====================================================================================
 // SChannelEditDialog 
 //=====================================================================================
@@ -158,7 +158,7 @@ void SChannelEditDialog::Construct(const FArguments& InArgs)
 
 		// channel name
 		+ SVerticalBox::Slot()
-		.FillHeight(1)
+		.FillHeight(1.f)
 		.VAlign(VAlign_Center)
 		.Padding(3.f, 1.f)
 		[
@@ -168,7 +168,7 @@ void SChannelEditDialog::Construct(const FArguments& InArgs)
 			.HAlign(HAlign_Left)
 			[
 				SNew(SBox)
-				.WidthOverride(100)
+				.WidthOverride(100.f)
 				[
 					SNew(STextBlock)
 					.Text(LOCTEXT("SChannelEditDialog_Name", "Name"))
@@ -176,7 +176,7 @@ void SChannelEditDialog::Construct(const FArguments& InArgs)
 				]
 			]
 			+SHorizontalBox::Slot()
-			.FillWidth(1)
+			.FillWidth(1.f)
 			.HAlign(HAlign_Left)
 			[
 				SAssignNew(NameBox, SEditableTextBox)
@@ -190,7 +190,7 @@ void SChannelEditDialog::Construct(const FArguments& InArgs)
 
 		// default response
 		+ SVerticalBox::Slot()
-		.FillHeight(1)
+		.FillHeight(1.f)
 		.Padding(3.f, 1.f)
 		.VAlign(VAlign_Center)
 		[
@@ -200,7 +200,7 @@ void SChannelEditDialog::Construct(const FArguments& InArgs)
 			.HAlign(HAlign_Left)
 			[
 				SNew(SBox)
-				.WidthOverride(100)
+				.WidthOverride(100.f)
 				[
 					SNew(STextBlock)
 					.Text(LOCTEXT("SChannelEditDialog_DefaultResponse", "Default Response"))
@@ -208,7 +208,7 @@ void SChannelEditDialog::Construct(const FArguments& InArgs)
 				]
 			]
 			+SHorizontalBox::Slot()
-			.FillWidth(1)
+			.FillWidth(1.f)
 			.HAlign(HAlign_Left)
 			[
 				SAssignNew(ResponseComboBox, SComboBox< TSharedPtr<FString> >)
@@ -493,7 +493,7 @@ void SProfileEditDialog::Construct(const FArguments& InArgs)
 			.HAlign(HAlign_Left)
 			[
 				SNew(SBox)
-				.WidthOverride(100)
+				.WidthOverride(100.f)
 				[
 					SNew(STextBlock)
 					.Text(LOCTEXT("SProfileEditDialog_Name", "Name"))
@@ -501,7 +501,7 @@ void SProfileEditDialog::Construct(const FArguments& InArgs)
 				]
 			]
 			+SHorizontalBox::Slot()
-			.FillWidth(1)
+			.FillWidth(1.f)
 			.HAlign(HAlign_Left)
 			[
 				SAssignNew(NameBox, SEditableTextBox)
@@ -526,7 +526,7 @@ void SProfileEditDialog::Construct(const FArguments& InArgs)
 			.HAlign(HAlign_Left)
 			[
 				SNew(SBox)
-				.WidthOverride(100)
+				.WidthOverride(100.f)
 				[
 					SNew(STextBlock)
 					.Text(LOCTEXT("SProfileEditDialog_CollisionEnabled", "CollisionEnabled"))
@@ -534,7 +534,7 @@ void SProfileEditDialog::Construct(const FArguments& InArgs)
 				]
 			]
 			+SHorizontalBox::Slot()
-			.FillWidth(1)
+			.FillWidth(1.f)
 			.HAlign(HAlign_Left)
 			[
 				SAssignNew(CollisionEnabledComboBox, SComboBox< TSharedPtr<FString> >)
@@ -562,7 +562,7 @@ void SProfileEditDialog::Construct(const FArguments& InArgs)
 			.HAlign(HAlign_Left)
 			[
 				SNew(SBox)
-				.WidthOverride(100)
+				.WidthOverride(100.f)
 				[
 					SNew(STextBlock)
 					.Text(LOCTEXT("SProfileEditDialog_ObjectType", "ObjectType"))
@@ -570,7 +570,7 @@ void SProfileEditDialog::Construct(const FArguments& InArgs)
 				]
 			]
 			+SHorizontalBox::Slot()
-			.FillWidth(1)
+			.FillWidth(1.f)
 			.HAlign(HAlign_Left)
 			[
 				SAssignNew(ObjectTypeComboBox, SComboBox< TSharedPtr<FString> >)
@@ -598,7 +598,7 @@ void SProfileEditDialog::Construct(const FArguments& InArgs)
 			.HAlign(HAlign_Left)
 			[
 				SNew(SBox)
-				.WidthOverride(100)
+				.WidthOverride(100.f)
 				[
 					SNew(STextBlock)
 					.Text(LOCTEXT("SProfileEditDialog_Description", "Description"))
@@ -606,7 +606,7 @@ void SProfileEditDialog::Construct(const FArguments& InArgs)
 				]
 			]
 			+SHorizontalBox::Slot()
-			.FillWidth(1)
+			.FillWidth(1.f)
 			.HAlign(HAlign_Left)
 			[
 				SNew(SEditableTextBox)
@@ -618,7 +618,7 @@ void SProfileEditDialog::Construct(const FArguments& InArgs)
 		]
 
 		+SVerticalBox::Slot()
-		.FillHeight(1)
+		.FillHeight(1.f)
 		[
 			SNew(SScrollBox)
 			+ SScrollBox::Slot()
@@ -630,7 +630,7 @@ void SProfileEditDialog::Construct(const FArguments& InArgs)
 		// accept or cancel button
 		+ SVerticalBox::Slot()
 		.AutoHeight()
-		.Padding(1,3)
+		.Padding(1.f,3.f)
 		[
 			SNew(SHorizontalBox)
 			+SHorizontalBox::Slot()
@@ -1142,7 +1142,7 @@ void SProfileEditDialog::OnCollisionChannelChanged(ECheckBoxState InNewValue, in
 {
 	if(ValidIndex >= 0 && ValidIndex < MAX_COLLISION_CHANNEL)
 	{
-		ProfileTemplate.ResponseToChannels.EnumArray[ValidIndex] = InCollisionResponse;
+		ProfileTemplate.ResponseToChannels.EnumArray[ValidIndex] = static_cast<uint8>(InCollisionResponse);
 	}
 }
 
@@ -1152,7 +1152,7 @@ ECheckBoxState SProfileEditDialog::IsCollisionChannelChecked(int32 ValidIndex, E
 
 	if( ValidIndex >= 0 && ValidIndex < MAX_COLLISION_CHANNEL )
 	{
-		if(ProfileTemplate.ResponseToChannels.EnumArray[ValidIndex] == InCollisionResponse)
+		if(ProfileTemplate.ResponseToChannels.EnumArray[ValidIndex] == static_cast<uint8>(InCollisionResponse))
 		{
 			return ECheckBoxState::Checked;
 		}
@@ -1165,7 +1165,7 @@ void SProfileEditDialog::OnAllCollisionChannelChanged(ECheckBoxState InNewValue,
 {
 	for(int32 Index=0; Index<MAX_COLLISION_CHANNEL; ++Index)
 	{
-		ProfileTemplate.ResponseToChannels.EnumArray[Index] = InCollisionResponse;
+		ProfileTemplate.ResponseToChannels.EnumArray[Index] = static_cast<uint8>(InCollisionResponse);
 	}
 }
 
@@ -1173,7 +1173,7 @@ ECheckBoxState SProfileEditDialog::IsAllCollisionChannelChecked(ECollisionRespon
 {
 	for(int32 Index=0; Index<MAX_COLLISION_CHANNEL; ++Index)
 	{
-		if(ProfileTemplate.ResponseToChannels.EnumArray[Index] != InCollisionResponse)
+		if(ProfileTemplate.ResponseToChannels.EnumArray[Index] != static_cast<uint8>(InCollisionResponse))
 		{
 			return ECheckBoxState::Unchecked;
 		}
@@ -1212,8 +1212,8 @@ TSharedRef<SWidget> SChannelListItem::GenerateWidgetForColumn(const FName& Colum
 	if (ColumnName == TEXT("Name"))
 	{
 		return	SNew(SBox)
-			.HeightOverride(20)
-			.Padding(FMargin(3, 0))
+			.HeightOverride(20.f)
+			.Padding(FMargin(3.f, 0.f))
 			.VAlign(VAlign_Center)
 			[
 				SNew(STextBlock)
@@ -1224,8 +1224,8 @@ TSharedRef<SWidget> SChannelListItem::GenerateWidgetForColumn(const FName& Colum
 	else if (ColumnName == TEXT("DefaultResponse"))
 	{
 		return	SNew(SBox)
-			.HeightOverride(20)
-			.Padding(FMargin(3, 0))
+			.HeightOverride(20.f)
+			.Padding(FMargin(3.f, 0.f))
 			.VAlign(VAlign_Center)
 			[
 				SNew(STextBlock)
@@ -1282,7 +1282,7 @@ TSharedRef<SWidget> SProfileListItem::GenerateWidgetForColumn(const FName& Colum
 		if(ProfileTemplate->bCanModify == false)
 		{
 			return	SNew(SBox)
-				.HeightOverride(20)
+				.HeightOverride(20.f)
 				.VAlign(VAlign_Center)
 				[
 					SNew(SImage)
@@ -1293,7 +1293,7 @@ TSharedRef<SWidget> SProfileListItem::GenerateWidgetForColumn(const FName& Colum
 		else
 		{
 			return	SNew(SBox)
-				.HeightOverride(20)
+				.HeightOverride(20.f)
 				.VAlign(VAlign_Center)
 				[
 					SNew(SImage)
@@ -1305,7 +1305,7 @@ TSharedRef<SWidget> SProfileListItem::GenerateWidgetForColumn(const FName& Colum
 	else if(ColumnName == TEXT("Name"))
 	{
 		return	SNew(SBox)
-			.HeightOverride(20)
+			.HeightOverride(20.f)
 			.VAlign(VAlign_Center)
 			[
 				SNew(STextBlock)
@@ -1316,7 +1316,7 @@ TSharedRef<SWidget> SProfileListItem::GenerateWidgetForColumn(const FName& Colum
 	else if(ColumnName == TEXT("Collision"))
 	{
 		return	SNew(SBox)
-			.HeightOverride(20)
+			.HeightOverride(20.f)
 			.VAlign(VAlign_Center)
 			[
 				SNew(STextBlock)
@@ -1327,7 +1327,7 @@ TSharedRef<SWidget> SProfileListItem::GenerateWidgetForColumn(const FName& Colum
 	else if(ColumnName == TEXT("ObjectType"))
 	{
 		return	SNew(SBox)
-			.HeightOverride(20)
+			.HeightOverride(20.f)
 			.VAlign(VAlign_Center)
 			[
 				SNew(STextBlock)
@@ -1338,7 +1338,7 @@ TSharedRef<SWidget> SProfileListItem::GenerateWidgetForColumn(const FName& Colum
 	else if(ColumnName == TEXT("Description"))
 	{
 		return	SNew(SBox)
-			.HeightOverride(20)
+			.HeightOverride(20.f)
 			.VAlign(VAlign_Center)
 			[
 				SNew(STextBlock)
@@ -1389,13 +1389,13 @@ void FCollisionProfileDetails::CustomizeDetails( IDetailLayoutBuilder& DetailBui
 	[
 		SNew(SVerticalBox)
 		+SVerticalBox::Slot()
-		.Padding(5)
+		.Padding(5.f)
 		.AutoHeight()
 		[
 			SNew(SHorizontalBox)
 			+SHorizontalBox::Slot()
-			.Padding(2, 10)
-			.FillWidth(1)
+			.Padding(2.f, 10.f)
+			.FillWidth(1.f)
 			[
 				SNew(STextBlock)
 				.AutoWrapText(true)
@@ -1405,7 +1405,7 @@ void FCollisionProfileDetails::CustomizeDetails( IDetailLayoutBuilder& DetailBui
 			]
 
 			+SHorizontalBox::Slot()
-			.Padding(2, 10)
+			.Padding(2.f, 10.f)
 			.AutoWidth()
 			.VAlign(VAlign_Center)
 			[
@@ -1417,7 +1417,7 @@ void FCollisionProfileDetails::CustomizeDetails( IDetailLayoutBuilder& DetailBui
 			]
 
 			+SHorizontalBox::Slot()
-			.Padding(2, 10)
+			.Padding(2.f, 10.f)
 			.AutoWidth()
 			.VAlign(VAlign_Center)
 			[
@@ -1429,7 +1429,7 @@ void FCollisionProfileDetails::CustomizeDetails( IDetailLayoutBuilder& DetailBui
 			]
 
 			+SHorizontalBox::Slot()
-			.Padding(2, 10)
+			.Padding(2.f, 10.f)
 			.AutoWidth()
 			.VAlign(VAlign_Center)
 			[
@@ -1442,8 +1442,8 @@ void FCollisionProfileDetails::CustomizeDetails( IDetailLayoutBuilder& DetailBui
 		]
 
 		+SVerticalBox::Slot()
-		.Padding(5)
-		.FillHeight(1)
+		.Padding(5.f)
+		.FillHeight(1.f)
 		[
 			SNew(SBox)
 			.MaxDesiredHeight(200.f)
@@ -1459,8 +1459,8 @@ void FCollisionProfileDetails::CustomizeDetails( IDetailLayoutBuilder& DetailBui
 					// Name
 					+ SHeaderRow::Column("Name")
 					.HAlignCell(HAlign_Left)
-					.FillWidth(1)
-					.HeaderContentPadding(FMargin(0, 3))
+					.FillWidth(1.f)
+					.HeaderContentPadding(FMargin(0.f, 3.f))
 					[
 						SNew(STextBlock)
 						.Text(LOCTEXT("ChannelListHeader_Name", "Name"))
@@ -1469,8 +1469,8 @@ void FCollisionProfileDetails::CustomizeDetails( IDetailLayoutBuilder& DetailBui
 					// Default Response
 					+ SHeaderRow::Column("DefaultResponse")
 					.HAlignCell(HAlign_Left)
-					.FillWidth(1)
-					.HeaderContentPadding(FMargin(0, 3))
+					.FillWidth(1.f)
+					.HeaderContentPadding(FMargin(0.f, 3.f))
 					[
 						SNew(STextBlock)
 						.Text(LOCTEXT("ChannelListHeader_DefaultResponse", "Default Response"))
@@ -1491,8 +1491,8 @@ void FCollisionProfileDetails::CustomizeDetails( IDetailLayoutBuilder& DetailBui
 		[
 			SNew(SHorizontalBox)
 			+SHorizontalBox::Slot()
-			.Padding(2, 10)
-			.FillWidth(1)
+			.Padding(2.f, 10.f)
+			.FillWidth(1.f)
 			[
 				SNew(STextBlock)
 				.Font(IDetailLayoutBuilder::GetDetailFont())
@@ -1502,7 +1502,7 @@ void FCollisionProfileDetails::CustomizeDetails( IDetailLayoutBuilder& DetailBui
 			]
 
 			+SHorizontalBox::Slot()
-			.Padding(2, 10)
+			.Padding(2.f, 10.f)
 			.AutoWidth()
 			.VAlign(VAlign_Center)
 			[
@@ -1514,7 +1514,7 @@ void FCollisionProfileDetails::CustomizeDetails( IDetailLayoutBuilder& DetailBui
 			]
 
 			+SHorizontalBox::Slot()
-			.Padding(2, 10)
+			.Padding(2.f, 10.f)
 			.AutoWidth()
 			.VAlign(VAlign_Center)
 			[
@@ -1526,7 +1526,7 @@ void FCollisionProfileDetails::CustomizeDetails( IDetailLayoutBuilder& DetailBui
 			]
 
 			+SHorizontalBox::Slot()
-			.Padding(2, 10)
+			.Padding(2.f, 10.f)
 			.AutoWidth()
 			.VAlign(VAlign_Center)
 			[
@@ -1539,8 +1539,8 @@ void FCollisionProfileDetails::CustomizeDetails( IDetailLayoutBuilder& DetailBui
 		]
 
 		+SVerticalBox::Slot()
-		.Padding(5)
-		.FillHeight(1)
+		.Padding(5.f)
+		.FillHeight(1.f)
 		[
 			SNew(SBox)
 			.MaxDesiredHeight(400.f)
@@ -1556,8 +1556,8 @@ void FCollisionProfileDetails::CustomizeDetails( IDetailLayoutBuilder& DetailBui
 					// Name
 					+ SHeaderRow::Column("Name")
 					.HAlignCell(HAlign_Left)
-					.FillWidth(1)
-					.HeaderContentPadding(FMargin(0, 3))
+					.FillWidth(1.f)
+					.HeaderContentPadding(FMargin(0.f, 3.f))
 					[
 						SNew(STextBlock)
 						.Text(LOCTEXT("ChannelListHeader_Name", "Name"))
@@ -1567,7 +1567,7 @@ void FCollisionProfileDetails::CustomizeDetails( IDetailLayoutBuilder& DetailBui
 					+ SHeaderRow::Column("DefaultResponse")
 					.HAlignCell(HAlign_Left)
 					.FillWidth(1)
-					.HeaderContentPadding(FMargin(0, 3))
+					.HeaderContentPadding(FMargin(0.f, 3.f))
 					[
 						SNew(STextBlock)
 						.Text(LOCTEXT("ChannelListHeader_DefaultResponse", "Default Response"))
@@ -1588,8 +1588,8 @@ void FCollisionProfileDetails::CustomizeDetails( IDetailLayoutBuilder& DetailBui
 		[
 			SNew(SHorizontalBox)
 			+SHorizontalBox::Slot()
-			.Padding(2, 2)
-			.FillWidth(1)
+			.Padding(2.f, 2.f)
+			.FillWidth(1.f)
 			[
 				SNew(STextBlock)
 				.AutoWrapText(true)
@@ -1599,7 +1599,7 @@ void FCollisionProfileDetails::CustomizeDetails( IDetailLayoutBuilder& DetailBui
 			]
 
 			+SHorizontalBox::Slot()
-			.Padding(2, 2)
+			.Padding(2.f, 2.f)
 			.AutoWidth()
 			.VAlign(VAlign_Center)
 			[
@@ -1610,7 +1610,7 @@ void FCollisionProfileDetails::CustomizeDetails( IDetailLayoutBuilder& DetailBui
 			]
 
 			+SHorizontalBox::Slot()
-			.Padding(2, 2)
+			.Padding(2.f, 2.f)
 			.AutoWidth()
 			.VAlign(VAlign_Center)
 			[
@@ -1622,7 +1622,7 @@ void FCollisionProfileDetails::CustomizeDetails( IDetailLayoutBuilder& DetailBui
 			]
 
 			+SHorizontalBox::Slot()
-			.Padding(2, 2)
+			.Padding(2.f, 2.f)
 			.AutoWidth()
 			.VAlign(VAlign_Center)
 			[
@@ -1635,8 +1635,8 @@ void FCollisionProfileDetails::CustomizeDetails( IDetailLayoutBuilder& DetailBui
 		]
 
 		+SVerticalBox::Slot()
-		.Padding(5)
-		.FillHeight(1)
+		.Padding(5.f)
+		.FillHeight(1.f)
 		[
 			SNew(SBox)
 			.MaxDesiredHeight(600.f)
@@ -1652,7 +1652,7 @@ void FCollisionProfileDetails::CustomizeDetails( IDetailLayoutBuilder& DetailBui
 					// Name
 					+ SHeaderRow::Column("Engine")
 					.HAlignCell(HAlign_Left)
-					.FixedWidth(30)
+					.FixedWidth(30.f)
 					[
 						SNew(SHorizontalBox)
 						+SHorizontalBox::Slot()
@@ -1666,8 +1666,8 @@ void FCollisionProfileDetails::CustomizeDetails( IDetailLayoutBuilder& DetailBui
 					// Name
 					+ SHeaderRow::Column("Name")
 					.HAlignCell(HAlign_Left)
-					.FillWidth(1)
-					.HeaderContentPadding(FMargin(0, 3))
+					.FillWidth(1.f)
+					.HeaderContentPadding(FMargin(0.f, 3.f))
 					[
 						SNew(STextBlock)
 						.Text(LOCTEXT("ProfileListHeader_Name", "Name"))
@@ -1677,7 +1677,7 @@ void FCollisionProfileDetails::CustomizeDetails( IDetailLayoutBuilder& DetailBui
 					+ SHeaderRow::Column("Collision")
 					.HAlignCell(HAlign_Left)
 					.FillWidth(1)
-					.HeaderContentPadding(FMargin(0, 3))
+					.HeaderContentPadding(FMargin(0.f, 3.f))
 					[
 						SNew(STextBlock)
 						.Text(LOCTEXT("ProfileListHeader_Collision", "Collision"))
@@ -1686,8 +1686,8 @@ void FCollisionProfileDetails::CustomizeDetails( IDetailLayoutBuilder& DetailBui
 					// Trace Type
 					+ SHeaderRow::Column("ObjectType")
 					.HAlignCell(HAlign_Left)
-					.FillWidth(1)
-					.HeaderContentPadding(FMargin(0, 3))
+					.FillWidth(1.f)
+					.HeaderContentPadding(FMargin(0.f, 3.f))
 					[
 						SNew(STextBlock)
 						.Text(LOCTEXT("ProfileListHeader_ObjectType", "Object Type"))
@@ -1696,8 +1696,8 @@ void FCollisionProfileDetails::CustomizeDetails( IDetailLayoutBuilder& DetailBui
 					// Static Object
 					+ SHeaderRow::Column("Description")
 					.HAlignCell(HAlign_Left)
-					.FillWidth(2)
-					.HeaderContentPadding(FMargin(0, 3))
+					.FillWidth(2.f)
+					.HeaderContentPadding(FMargin(0.f, 3.f))
 					[
 						SNew(STextBlock)
 						.Text(LOCTEXT("ProfileListHeader_Description", "Description"))
@@ -1985,7 +1985,7 @@ FReply FCollisionProfileDetails::OnNewChannel(bool bTraceType)
 		// Create modal window for modification
 		TSharedRef<SWindow> WidgetWindow = SNew(SWindow)
 			.Title(LOCTEXT("CollisionProfileDetail_NewChannelTitle", "New Channel"))
-			.ClientSize(FVector2D(CHANNEL_WINDOW_WIDTH, CHANNEL_WINDOW_HEIGHT))
+			.ClientSize(FVector2f(CHANNEL_WINDOW_WIDTH, CHANNEL_WINDOW_HEIGHT))
 			.SupportsMaximize(false)
 			.SupportsMinimize(false)
 			.SizingRule(ESizingRule::UserSized);
@@ -2027,7 +2027,7 @@ FReply	FCollisionProfileDetails::OnEditChannel(bool bTraceType)
 	{
 		TSharedRef<SWindow> WidgetWindow = SNew(SWindow)
 			.Title(LOCTEXT("CollisionProfileDetail_EditChannelTitle", "Edit Channel"))
-			.ClientSize(FVector2D(CHANNEL_WINDOW_WIDTH, CHANNEL_WINDOW_HEIGHT))
+			.ClientSize(FVector2f(CHANNEL_WINDOW_WIDTH, CHANNEL_WINDOW_HEIGHT))
 			.SupportsMaximize(false)
 			.SupportsMinimize(false)
 			.SizingRule(ESizingRule::UserSized);
@@ -2098,7 +2098,7 @@ FReply	FCollisionProfileDetails::OnNewProfile()
 	// Create modal window for modification
 	TSharedRef<SWindow> WidgetWindow = SNew(SWindow)
 		.Title(LOCTEXT("CollisionProfileDetail_NewProfileTitle", "New Profile"))
-		.ClientSize(FVector2D(PROFILE_WINDOW_WIDTH, PROFILE_WINDOW_HEIGHT))
+		.ClientSize(FVector2f(PROFILE_WINDOW_WIDTH, PROFILE_WINDOW_HEIGHT))
 		.SupportsMaximize(false)
 		.SupportsMinimize(false)
 		.SizingRule(ESizingRule::UserSized);
@@ -2148,7 +2148,7 @@ FReply	FCollisionProfileDetails::OnEditProfile()
 			// Create modal window for modification
 			TSharedRef<SWindow> WidgetWindow = SNew(SWindow)
 				.Title(LOCTEXT("CollisionProfileDetail_EditProfileTitle", "Edit Profile"))
-				.ClientSize(FVector2D(PROFILE_WINDOW_WIDTH, PROFILE_WINDOW_HEIGHT))
+				.ClientSize(FVector2f(PROFILE_WINDOW_WIDTH, PROFILE_WINDOW_HEIGHT))
 				.SupportsMaximize(false)
 				.SupportsMinimize(false)
 				.SizingRule(ESizingRule::UserSized);
diff --git a/Engine/Source/Editor/DetailCustomizations/Private/ColorGradingVectorCustomization.cpp b/Engine/Source/Editor/DetailCustomizations/Private/ColorGradingVectorCustomization.cpp
--- a/Engine/Source/Editor/DetailCustomizations/Private/ColorGradingVectorCustomization.cpp
+++ b/Engine/Source/Editor/DetailCustomizations/Private/ColorGradingVectorCustomization.cpp
@@ -100,7 +100,7 @@ FColorGradingVectorCustomizationBase::FColorGradingVectorCustomizationBase(const
 	{
 		FVector4 VectorValue;
 		ColorGradingPropertyHandle.GetValue(VectorValue);
-		CurrentHSVColor = FLinearColor(VectorValue.X, VectorValue.Y, VectorValue.Z).LinearRGBToHSV();
+		CurrentHSVColor = FLinearColor(static_cast<float>(VectorValue.X), static_cast<float>(VectorValue.Y), static_cast<float>(VectorValue.Z)).LinearRGBToHSV();
 	}
 }
 
@@ -236,7 +236,7 @@ void FColorGradingVectorCustomizationBase::OnValueChanged(float NewValue, int32
 
 		if (ColorIndex < 3)
 		{
-			CurrentHSVColor = FLinearColor(NewValueVector.X, NewValueVector.Y, NewValueVector.Z).LinearRGBToHSV();
+			CurrentHSVColor = FLinearColor(static_cast<float>(NewValueVector.X), static_cast<float>(NewValueVector.Y), static_cast<float>(NewValueVector.Z)).LinearRGBToHSV();
 		}
 	}
 	else
@@ -271,11 +271,11 @@ TOptional<float> FColorGradingVectorCustomizationBase::OnSliderGetValue(int32 Co
 
 		if (IsRGBMode)
 		{
-			Value = ValueVector[ColorIndex];
+			Value = static_cast<float>(ValueVector[ColorIndex]);
 		}
 		else
 		{
-			Value = (ColorIndex < 3) ? CurrentHSVColor.Component(ColorIndex) : ValueVector.W;
+			Value = (ColorIndex < 3) ? CurrentHSVColor.Component(ColorIndex) : static_cast<float>(ValueVector.W);
 		}
 
 		return Value;
@@ -301,7 +301,7 @@ void FColorGradingVectorCustomizationBase::PostUndo(bool bSuccess)
 		FVector4 CurrentValueVector;
 		if (ColorGradingPropertyHandle.GetValue(CurrentValueVector) == FPropertyAccess::Success)
 		{
-			CurrentHSVColor = FLinearColor(CurrentValueVector.X, CurrentValueVector.Y, CurrentValueVector.Z).LinearRGBToHSV();
+			CurrentHSVColor = FLinearColor(static_cast<float>(CurrentValueVector.X), static_cast<float>(CurrentValueVector.Y), static_cast<float>(CurrentValueVector.Z)).LinearRGBToHSV();
 			OnCurrentHSVColorChanged.Broadcast(CurrentHSVColor, true);
 		}
 	}
@@ -443,8 +443,8 @@ void FColorGradingVectorCustomization::MakeHeaderRow(FDetailWidgetRow& Row, TSha
 	{
 		Row.ValueContent()
 			// Make enough space for each child handle
-			.MinDesiredWidth(125.0f * SortedChildArray.Num())
-			.MaxDesiredWidth(125.0f * SortedChildArray.Num())
+			.MinDesiredWidth(125.0f * static_cast<float>(SortedChildArray.Num()))
+			.MaxDesiredWidth(125.0f * static_cast<float>(SortedChildArray.Num()))
 			[
 				ContentHorizontalBox.ToSharedRef()
 			];
@@ -542,9 +542,9 @@ FLinearColor FColorGradingVectorCustomization::OnGetColorForHeaderColorBlock() c
 	FVector4 VectorValue;
 	if (ColorGradingPropertyHandle.GetValue(VectorValue) == FPropertyAccess::Success)
 	{
-		ColorValue.R = VectorValue.X * VectorValue.W;
-		ColorValue.G = VectorValue.Y * VectorValue.W;
-		ColorValue.B = VectorValue.Z * VectorValue.W;
+		ColorValue.R = static_cast<float>(VectorValue.X * VectorValue.W);
+		ColorValue.G = static_cast<float>(VectorValue.Y * VectorValue.W);
+		ColorValue.B = static_cast<float>(VectorValue.Z * VectorValue.W);
 	}
 	else
 	{
@@ -639,7 +639,7 @@ void FColorGradingCustomBuilder::OnDetailGroupReset()
 {
 	FVector4 CurrentValueVector;
 	verifySlow(ColorGradingPropertyHandle.GetValue(CurrentValueVector) == FPropertyAccess::Success);
-	CurrentHSVColor = FLinearColor(CurrentValueVector.X, CurrentValueVector.Y, CurrentValueVector.Z).LinearRGBToHSV();
+	CurrentHSVColor = FLinearColor(static_cast<float>(CurrentValueVector.X), static_cast<float>(CurrentValueVector.Y), static_cast<float>(CurrentValueVector.Z)).LinearRGBToHSV();
 
 	OnCurrentHSVColorChanged.Broadcast(CurrentHSVColor, true);
 
@@ -660,7 +660,7 @@ void FColorGradingCustomBuilder::ResetToDefault(TSharedPtr<IPropertyHandle> Prop
 
 	FVector4 CurrentValueVector;
 	verifySlow(ColorGradingPropertyHandle.GetValue(CurrentValueVector) == FPropertyAccess::Success);
-	CurrentHSVColor = FLinearColor(CurrentValueVector.X, CurrentValueVector.Y, CurrentValueVector.Z).LinearRGBToHSV();
+	CurrentHSVColor = FLinearColor(static_cast<float>(CurrentValueVector.X), static_cast<float>(CurrentValueVector.Y), static_cast<float>(CurrentValueVector.Z)).LinearRGBToHSV();
 
 	OnCurrentHSVColorChanged.Broadcast(CurrentHSVColor, true);
 
@@ -927,7 +927,7 @@ void FColorGradingCustomBuilder::OnChangeColorModeClicked(UE::ColorGrading::ECol
 
 		GConfig->SetBool(TEXT("ColorGrading"), *FString::Printf(TEXT("%s_%s_IsRGB"), *ParentGroupName, *ColorGradingPropertyHandle.GetHandle()->GetPropertyDisplayName().ToString()), IsRGBMode, GEditorPerProjectIni);
 
-		CurrentHSVColor = FLinearColor(CurrentValueVector.X, CurrentValueVector.Y, CurrentValueVector.Z).LinearRGBToHSV();
+		CurrentHSVColor = FLinearColor(static_cast<float>(CurrentValueVector.X), static_cast<float>(CurrentValueVector.Y), static_cast<float>(CurrentValueVector.Z)).LinearRGBToHSV();
 
 		OnCurrentHSVColorChanged.Broadcast(CurrentHSVColor, true);
 
@@ -958,7 +958,7 @@ void FColorGradingCustomBuilder::OnColorGradingPickerChanged(FVector4& NewValue,
 		}
 	}
 
-	FLinearColor NewHSVColor(NewValue.X, NewValue.Y, NewValue.Z);
+	FLinearColor NewHSVColor(static_cast<float>(NewValue.X), static_cast<float>(NewValue.Y), static_cast<float>(NewValue.Z));
 	NewHSVColor = NewHSVColor.LinearRGBToHSV();
 
 	OnCurrentHSVColorChangedDelegate(NewHSVColor, true);
@@ -1013,7 +1013,7 @@ void FColorGradingCustomBuilder::OnPropertyValueChanged(UObject* Object, FProper
 					FVector4 CurrentValueVector;
 					if (PinnedPropertyHandle->GetValue(CurrentValueVector) == FPropertyAccess::Success)
 					{
-						CurrentHSVColor = FLinearColor(CurrentValueVector.X, CurrentValueVector.Y, CurrentValueVector.Z).LinearRGBToHSV();
+						CurrentHSVColor = FLinearColor(static_cast<float>(CurrentValueVector.X), static_cast<float>(CurrentValueVector.Y), static_cast<float>(CurrentValueVector.Z)).LinearRGBToHSV();
 					}
 					break;
 				}
diff --git a/Engine/Source/Editor/DetailCustomizations/Private/ComponentTransformDetails.cpp b/Engine/Source/Editor/DetailCustomizations/Private/ComponentTransformDetails.cpp
--- a/Engine/Source/Editor/DetailCustomizations/Private/ComponentTransformDetails.cpp
+++ b/Engine/Source/Editor/DetailCustomizations/Private/ComponentTransformDetails.cpp
@@ -210,7 +210,7 @@ TSharedRef<SWidget> FComponentTransformDetails::BuildTransformFieldLabel( ETrans
 		.VAlign(VAlign_Center)
 		[
 			SNew(SComboButton)
-			.ContentPadding(0)
+			.ContentPadding(0.f)
 			.IsEnabled(this, &FComponentTransformDetails::CanChangeAbsoluteFlag, TransformField)
 			.MenuContent()
 			[
@@ -618,7 +618,7 @@ void FComponentTransformDetails::Tick( float DeltaTime )
 
 void FComponentTransformDetails::CacheCommonLocationUnits()
 {
-	float LargestValue = 0.f;
+	FVector::FReal LargestValue = 0.0;
 	if (CachedLocation.X.IsSet() && CachedLocation.X.GetValue() > LargestValue)
 	{
 		LargestValue = CachedLocation.X.GetValue();
@@ -1377,7 +1377,7 @@ void FComponentTransformDetails::OnSetTransform(ETransformField::Type TransformF
 							if (bPreserveScaleRatio)
 							{
 								// If we set a single axis, scale the others
-								float Ratio = 0.0f;
+								FVector::FReal Ratio = 0.0f;
 
 								switch (Axis)
 								{
diff --git a/Engine/Source/Editor/DetailCustomizations/Private/CurveColorCustomization.cpp b/Engine/Source/Editor/DetailCustomizations/Private/CurveColorCustomization.cpp
--- a/Engine/Source/Editor/DetailCustomizations/Private/CurveColorCustomization.cpp
+++ b/Engine/Source/Editor/DetailCustomizations/Private/CurveColorCustomization.cpp
@@ -467,10 +467,10 @@ FReply FCurveColorCustomization::OnCurvePreviewDoubleClick(const FGeometry& InMy
 			DestroyPopOutWindow();
 
 			// Determine the position of the window so that it will spawn near the mouse, but not go off the screen.
-			const FVector2D CursorPos = FSlateApplication::Get().GetCursorPos();
+			const FVector2f CursorPos = FSlateApplication::Get().GetCursorPos();
 			FSlateRect Anchor(CursorPos.X, CursorPos.Y, CursorPos.X, CursorPos.Y);
 
-			FVector2D AdjustedSummonLocation = FSlateApplication::Get().CalculatePopupWindowPosition( Anchor, FCurveColorCustomization::DEFAULT_WINDOW_SIZE, true, FVector2D::ZeroVector, Orient_Horizontal );
+			FVector2f AdjustedSummonLocation = FSlateApplication::Get().CalculatePopupWindowPosition( Anchor, FCurveColorCustomization::DEFAULT_WINDOW_SIZE, true, FVector2f::ZeroVector, Orient_Horizontal );
 
 			TSharedPtr<SWindow> Window = SNew(SWindow)
 				.Title( FText::Format( LOCTEXT("WindowHeader", "{0} - Internal Color Curve Editor"), StructPropertyHandle->GetPropertyDisplayName()) )
diff --git a/Engine/Source/Editor/DetailCustomizations/Private/CurveStructCustomization.cpp b/Engine/Source/Editor/DetailCustomizations/Private/CurveStructCustomization.cpp
--- a/Engine/Source/Editor/DetailCustomizations/Private/CurveStructCustomization.cpp
+++ b/Engine/Source/Editor/DetailCustomizations/Private/CurveStructCustomization.cpp
@@ -476,10 +476,10 @@ FReply FCurveStructCustomization::OnCurvePreviewDoubleClick( const FGeometry& In
 			DestroyPopOutWindow();
 
 			// Determine the position of the window so that it will spawn near the mouse, but not go off the screen.
-			const FVector2D CursorPos = FSlateApplication::Get().GetCursorPos();
+			const FVector2f CursorPos = FSlateApplication::Get().GetCursorPos();
 			FSlateRect Anchor(CursorPos.X, CursorPos.Y, CursorPos.X, CursorPos.Y);
 
-			FVector2D AdjustedSummonLocation = FSlateApplication::Get().CalculatePopupWindowPosition( Anchor, FCurveStructCustomization::DEFAULT_WINDOW_SIZE, true, FVector2D::ZeroVector, Orient_Horizontal );
+			FVector2f AdjustedSummonLocation = FSlateApplication::Get().CalculatePopupWindowPosition( Anchor, FCurveStructCustomization::DEFAULT_WINDOW_SIZE, true, FVector2f::ZeroVector, Orient_Horizontal );
 
 			TSharedPtr<SWindow> Window = SNew(SWindow)
 				.Title( FText::Format( LOCTEXT("WindowHeader", "{0} - Internal Curve Editor"), StructPropertyHandle->GetPropertyDisplayName()) )
diff --git a/Engine/Source/Editor/DetailCustomizations/Private/CurveVectorCustomization.cpp b/Engine/Source/Editor/DetailCustomizations/Private/CurveVectorCustomization.cpp
--- a/Engine/Source/Editor/DetailCustomizations/Private/CurveVectorCustomization.cpp
+++ b/Engine/Source/Editor/DetailCustomizations/Private/CurveVectorCustomization.cpp
@@ -451,10 +451,10 @@ FReply FCurveVectorCustomization::OnCurvePreviewDoubleClick(const FGeometry& InM
 			DestroyPopOutWindow();
 
 			// Determine the position of the window so that it will spawn near the mouse, but not go off the screen.
-			const FVector2D CursorPos = FSlateApplication::Get().GetCursorPos();
+			const FVector2f CursorPos = FSlateApplication::Get().GetCursorPos();
 			FSlateRect Anchor(CursorPos.X, CursorPos.Y, CursorPos.X, CursorPos.Y);
 
-			FVector2D AdjustedSummonLocation = FSlateApplication::Get().CalculatePopupWindowPosition(Anchor, FCurveVectorCustomization::DEFAULT_WINDOW_SIZE, true, FVector2D::ZeroVector, Orient_Horizontal);
+			FVector2f AdjustedSummonLocation = FSlateApplication::Get().CalculatePopupWindowPosition(Anchor, FCurveVectorCustomization::DEFAULT_WINDOW_SIZE, true, FVector2f::ZeroVector, Orient_Horizontal);
 
 			TSharedPtr<SWindow> Window = SNew(SWindow)
 				.Title(FText::Format(LOCTEXT("WindowHeader", "{0} - Internal Vector Curve Editor"), StructPropertyHandle->GetPropertyDisplayName()))
diff --git a/Engine/Source/Editor/DetailCustomizations/Private/CustomPrimitiveDataCustomization.cpp b/Engine/Source/Editor/DetailCustomizations/Private/CustomPrimitiveDataCustomization.cpp
--- a/Engine/Source/Editor/DetailCustomizations/Private/CustomPrimitiveDataCustomization.cpp
+++ b/Engine/Source/Editor/DetailCustomizations/Private/CustomPrimitiveDataCustomization.cpp
@@ -563,8 +563,9 @@ void FCustomPrimitiveDataCustomization::PopulateParameterData(UPrimitiveComponen
 			
 			if (ParameterMetadata.PrimitiveDataIndex > INDEX_NONE)
 			{
+				const uint8 PrimIdx = static_cast<uint8>(ParameterMetadata.PrimitiveDataIndex);
 				// Add each element individually, so that we can overlap vector parameter names
-				VectorParameterData.FindOrAdd(ParameterMetadata.PrimitiveDataIndex + 0).Add(
+				VectorParameterData.FindOrAdd(PrimIdx + 0).Add(
 				{
 					PrimitiveComponent,
 					MaterialInterface,
@@ -572,7 +573,7 @@ void FCustomPrimitiveDataCustomization::PopulateParameterData(UPrimitiveComponen
 					ParameterMetadata.ExpressionGuid,
 					0
 				});
-				VectorParameterData.FindOrAdd(ParameterMetadata.PrimitiveDataIndex + 1).Add(
+				VectorParameterData.FindOrAdd(PrimIdx + 1).Add(
 				{
 					PrimitiveComponent,
 					MaterialInterface,
@@ -580,7 +581,7 @@ void FCustomPrimitiveDataCustomization::PopulateParameterData(UPrimitiveComponen
 					ParameterMetadata.ExpressionGuid,
 					1
 				});
-				VectorParameterData.FindOrAdd(ParameterMetadata.PrimitiveDataIndex + 2).Add(
+				VectorParameterData.FindOrAdd(PrimIdx + 2).Add(
 				{
 					PrimitiveComponent,
 					MaterialInterface,
@@ -588,7 +589,7 @@ void FCustomPrimitiveDataCustomization::PopulateParameterData(UPrimitiveComponen
 					ParameterMetadata.ExpressionGuid,
 					2
 				});
-				VectorParameterData.FindOrAdd(ParameterMetadata.PrimitiveDataIndex + 3).Add(
+				VectorParameterData.FindOrAdd(PrimIdx + 3).Add(
 				{
 					PrimitiveComponent,
 					MaterialInterface,
@@ -596,7 +597,7 @@ void FCustomPrimitiveDataCustomization::PopulateParameterData(UPrimitiveComponen
 					ParameterMetadata.ExpressionGuid,
 					3
 				});
-				MaxPrimitiveDataIndex = FMath::Max(MaxPrimitiveDataIndex, (int32)(ParameterMetadata.PrimitiveDataIndex + 3));
+				MaxPrimitiveDataIndex = FMath::Max(MaxPrimitiveDataIndex, PrimIdx + 3);
 			}
 		}
 
@@ -611,7 +612,7 @@ void FCustomPrimitiveDataCustomization::PopulateParameterData(UPrimitiveComponen
 
 			if (ParameterMetadata.PrimitiveDataIndex > INDEX_NONE)
 			{
-				ScalarParameterData.FindOrAdd(ParameterMetadata.PrimitiveDataIndex).Add(
+				ScalarParameterData.FindOrAdd(static_cast<uint8>(ParameterMetadata.PrimitiveDataIndex)).Add(
 				{
 					PrimitiveComponent,
 					MaterialInterface,
@@ -751,11 +752,9 @@ void FCustomPrimitiveDataCustomization::OnAddedDesiredPrimitiveData(uint8 PrimId
 		
 		if (RawArray.Num() > 0)
 		{
-			TArray<float>* Data = reinterpret_cast<TArray<float>*>(RawArray[0]);
-
-			for (int32 i = NumElements; i <= PrimIdx; ++i)
+			for (uint32 i = NumElements; i <= static_cast<uint32>(PrimIdx); ++i)
 			{
-				SetDefaultValue(i);
+				SetDefaultValue(static_cast<uint8>(i));
 			}
 		}
 		
@@ -772,7 +771,7 @@ void FCustomPrimitiveDataCustomization::OnRemovedPrimitiveData(uint8 PrimIdx)
 	{
 		GEditor->BeginTransaction(LOCTEXT("OnRemovedPrimitiveData", "Removed Items"));
 
-		for (int32 i = NumElements - 1; i >= PrimIdx; --i)
+		for (uint32 i = NumElements - 1; i >= static_cast<uint32>(PrimIdx); --i)
 		{
 			DataArrayHandle->DeleteItem(i);
 		}
@@ -788,9 +787,9 @@ FLinearColor FCustomPrimitiveDataCustomization::GetVectorColor(uint8 PrimIdx) co
 	uint32 NumElems;
 	if (GetNumElements(NumElems) == FPropertyAccess::Success)
 	{
-		const int32 MaxElems = FMath::Min((int32)NumElems, PrimIdx + 4);
+		const uint32 MaxElems = FMath::Min(NumElems, static_cast<uint32>(PrimIdx + 4));
 
-		for (int32 i = PrimIdx; i < MaxElems; ++i)
+		for (uint32 i = PrimIdx; i < MaxElems; ++i)
 		{
 			DataArrayHandle->GetElement(i)->GetValue(Color[i - PrimIdx]);
 		}
@@ -875,8 +874,8 @@ void FCustomPrimitiveDataCustomization::SetDefaultVectorValue(uint8 PrimIdx)
 	{
 		GEditor->BeginTransaction(LOCTEXT("SetDefaultVectorValue", "Reset Vector To Default"));
 
-		const int32 MaxElems = FMath::Min((int32)NumElems, PrimIdx + 4);
-		for (int32 i = PrimIdx; i < MaxElems; ++i)
+		const uint32 MaxElems = FMath::Min(NumElems, static_cast<uint32>(PrimIdx + 4));
+		for (uint32 i = PrimIdx; i < MaxElems; ++i)
 		{
 			DataArrayHandle->GetElement(i)->ResetToDefault();
 		}
@@ -961,7 +960,7 @@ TSharedRef<SWidget> FCustomPrimitiveDataCustomization::GetUndeclaredParameterWid
 	const FText ToolTipText = FText::Format(
 		LOCTEXT("UndeclaredParameterTooltip", "An assigned material doesn't declare a parameter for primitive index {0}"),
 		FText::AsNumber(PrimIdx));
-	const uint16 FontSize = CustomizationUtils.GetRegularFont().GetClampSize();
+	const float FontSize = CustomizationUtils.GetRegularFont().GetClampSize();
 
 	return SNew(SHorizontalBox)
 		.ToolTipText(ToolTipText)
diff --git a/Engine/Source/Editor/DetailCustomizations/Private/DateTimeStructCustomization.cpp b/Engine/Source/Editor/DetailCustomizations/Private/DateTimeStructCustomization.cpp
--- a/Engine/Source/Editor/DetailCustomizations/Private/DateTimeStructCustomization.cpp
+++ b/Engine/Source/Editor/DetailCustomizations/Private/DateTimeStructCustomization.cpp
@@ -219,7 +219,7 @@ FDateTime FDateTimeStructCustomization::ConvertTime(const FDateTime& InDate, int
 	const int32 MinutesDifference =	OutTimezoneMinutes - InTimezoneMinutes;
 
 	// Calculate our Hours difference
-	const int64 HoursDifference = (OutTimezone / 100) - (InTimezone / 100);
+	const int32 HoursDifference = (OutTimezone / 100) - (InTimezone / 100);
 
 	return FDateTime(InDate + FTimespan(HoursDifference, MinutesDifference, 0));
 }
diff --git a/Engine/Source/Editor/DetailCustomizations/Private/DeviceProfileDetails.cpp b/Engine/Source/Editor/DetailCustomizations/Private/DeviceProfileDetails.cpp
--- a/Engine/Source/Editor/DetailCustomizations/Private/DeviceProfileDetails.cpp
+++ b/Engine/Source/Editor/DetailCustomizations/Private/DeviceProfileDetails.cpp
@@ -391,8 +391,8 @@ void SCVarSelectionPanel::Construct(const FArguments& InArgs, const FString& CVa
 	ChildSlot
 	[
 		SNew(SBox)
-		.WidthOverride(300)
-		.HeightOverride(512)
+		.WidthOverride(300.f)
+		.HeightOverride(512.f)
 		.Content()
 		[
 			SNew(SVerticalBox)
diff --git a/Engine/Source/Editor/DetailCustomizations/Private/DialogueStructsCustomizations.cpp b/Engine/Source/Editor/DetailCustomizations/Private/DialogueStructsCustomizations.cpp
--- a/Engine/Source/Editor/DetailCustomizations/Private/DialogueStructsCustomizations.cpp
+++ b/Engine/Source/Editor/DetailCustomizations/Private/DialogueStructsCustomizations.cpp
@@ -283,8 +283,8 @@ void SSpeakerDropDown::OnSelectionChanged(TSharedPtr<UDialogueVoice*> Speaker, E
 
 TSharedRef<SWidget> SSpeakerDropDown::MakeComboButtonItemWidget(TSharedPtr<UDialogueVoice*> Speaker)
 {
-	const float ThumbnailSizeX = 64.0f;
-	const float ThumbnailSizeY = 64.0f;
+	const uint32 ThumbnailSizeX = 64;
+	const uint32 ThumbnailSizeY = 64;
 
 	const TSharedPtr<FAssetThumbnail> AssetThumbnail = MakeShareable( new FAssetThumbnail( *Speaker, ThumbnailSizeX, ThumbnailSizeY, AssetThumbnailPool ) );
 
@@ -301,8 +301,8 @@ TSharedRef<SWidget> SSpeakerDropDown::MakeComboButtonItemWidget(TSharedPtr<UDial
 			.VAlign(VAlign_Center)
 			[
 				SNew( SBox )
-				.WidthOverride( ThumbnailSizeX ) 
-				.HeightOverride( ThumbnailSizeY )
+				.WidthOverride( static_cast<float>(ThumbnailSizeX) ) 
+				.HeightOverride( static_cast<float>(ThumbnailSizeY) )
 				[
 					AssetThumbnail->MakeThumbnailWidget()
 				]
@@ -519,8 +519,8 @@ TSharedRef<SWidget> STargetSetDropDown::MakeComboButtonItemWidget(TSharedPtr<FTa
 {
 	TSharedPtr<SWidget> Result;
 
-	const float ThumbnailSizeX = 64.0f;
-	const float ThumbnailSizeY = 64.0f;
+	const uint32 ThumbnailSizeX = 64;
+	const uint32 ThumbnailSizeY = 64;
 
 	FSlateFontInfo Font = IDetailLayoutBuilder::GetDetailFont();
 
@@ -544,8 +544,8 @@ TSharedRef<SWidget> STargetSetDropDown::MakeComboButtonItemWidget(TSharedPtr<FTa
 				.VAlign(VAlign_Center)
 				[
 					SNew( SBox )
-					.WidthOverride( ThumbnailSizeX ) 
-					.HeightOverride( ThumbnailSizeY )
+					.WidthOverride( static_cast<float>(ThumbnailSizeX) ) 
+					.HeightOverride( static_cast<float>(ThumbnailSizeY) )
 					[
 						AssetThumbnail->MakeThumbnailWidget()
 					]
@@ -590,8 +590,8 @@ TSharedRef<SWidget> STargetSetDropDown::MakeComboButtonItemWidget(TSharedPtr<FTa
 			.VAlign(VAlign_Center)
 			[
 				SNew( SBox )
-				.WidthOverride( ThumbnailSizeX ) 
-				.HeightOverride( ThumbnailSizeY )
+				.WidthOverride( static_cast<float>(ThumbnailSizeX) ) 
+				.HeightOverride( static_cast<float>(ThumbnailSizeY) )
 				[
 					AssetThumbnail->MakeThumbnailWidget()
 				]
@@ -617,8 +617,8 @@ TSharedRef<SWidget> STargetSetDropDown::MakeComboButtonItemWidget(TSharedPtr<FTa
 		.AutoHeight()
 		[
 			SNew( SBox )
-			.WidthOverride( ThumbnailSizeX )
-			.HeightOverride( ThumbnailSizeY )
+			.WidthOverride( static_cast<float>(ThumbnailSizeX) )
+			.HeightOverride( static_cast<float>(ThumbnailSizeY) )
 			.HAlign(HAlign_Center)
 			.VAlign(VAlign_Center)
 		]
diff --git a/Engine/Source/Editor/DetailCustomizations/Private/DialogueWaveWidgets.cpp b/Engine/Source/Editor/DetailCustomizations/Private/DialogueWaveWidgets.cpp
--- a/Engine/Source/Editor/DetailCustomizations/Private/DialogueWaveWidgets.cpp
+++ b/Engine/Source/Editor/DetailCustomizations/Private/DialogueWaveWidgets.cpp
@@ -72,8 +72,8 @@ void SDialogueVoicePropertyEditor::Construct( const FArguments& InArgs, const TS
 			DialogueVoice = Cast<UDialogueVoice>(Object);
 		}
 
-		const float ThumbnailSizeX = 64.0f;
-		const float ThumbnailSizeY = 64.0f;
+		const uint32 ThumbnailSizeX = 64;
+		const uint32 ThumbnailSizeY = 64;
 		AssetThumbnail = MakeShareable( new FAssetThumbnail( DialogueVoice, ThumbnailSizeX, ThumbnailSizeY, AssetThumbnailPool ) );
 
 		TSharedRef<SWidget> AssetWidget =
@@ -83,8 +83,8 @@ void SDialogueVoicePropertyEditor::Construct( const FArguments& InArgs, const TS
 			.OnAssetsDropped( this, &SDialogueVoicePropertyEditor::OnAssetDropped )
 			[
 				SNew( SBox )
-				.WidthOverride( ThumbnailSizeX ) 
-				.HeightOverride( ThumbnailSizeY )
+				.WidthOverride( static_cast<float>(ThumbnailSizeX) ) 
+				.HeightOverride( static_cast<float>(ThumbnailSizeY) )
 				.HAlign(HAlign_Fill)
 				.VAlign(VAlign_Fill)
 				.Padding(0.0f)
diff --git a/Engine/Source/Editor/DetailCustomizations/Private/EnvQueryParamInstanceCustomization.cpp b/Engine/Source/Editor/DetailCustomizations/Private/EnvQueryParamInstanceCustomization.cpp
--- a/Engine/Source/Editor/DetailCustomizations/Private/EnvQueryParamInstanceCustomization.cpp
+++ b/Engine/Source/Editor/DetailCustomizations/Private/EnvQueryParamInstanceCustomization.cpp
@@ -107,7 +107,7 @@ TOptional<float> FEnvQueryParamInstanceCustomization::GetParamNumValue() const
 		if (PropResult == FPropertyAccess::Success)
 		{
 			const int32 IntValue = *((int32*)&StoreValue);
-			return IntValue;
+			return static_cast<float>(IntValue);
 		}
 	}
 
@@ -211,7 +211,7 @@ void FEnvQueryParamInstanceCustomization::OnTypeChanged()
 		switch (ParamType)
 		{
 		case EAIParamType::Float:	OnParamNumValueChanged(CachedFloat); break;
-		case EAIParamType::Int:	OnParamNumValueChanged(CachedInt); break;
+		case EAIParamType::Int:	OnParamNumValueChanged(static_cast<float>(CachedInt)); break;
 		case EAIParamType::Bool:	OnParamBoolValueChanged(CachedBool ? ECheckBoxState::Checked : ECheckBoxState::Unchecked); break;
 		}
 	}
diff --git a/Engine/Source/Editor/DetailCustomizations/Private/HardwareTargetingSettingsDetails.cpp b/Engine/Source/Editor/DetailCustomizations/Private/HardwareTargetingSettingsDetails.cpp
--- a/Engine/Source/Editor/DetailCustomizations/Private/HardwareTargetingSettingsDetails.cpp
+++ b/Engine/Source/Editor/DetailCustomizations/Private/HardwareTargetingSettingsDetails.cpp
@@ -110,7 +110,7 @@ public:
 
 				+ SHorizontalBox::Slot()
 				.AutoWidth()
-				.Padding(6)
+				.Padding(6.f)
 				[
 					SNew(SButton)
 					.Text(LOCTEXT("RestartEditor", "Restart Editor"))
@@ -120,7 +120,7 @@ public:
 
 				+ SHorizontalBox::Slot()
 				.AutoWidth()
-				.Padding(6)
+				.Padding(6.f)
 				[
 					SNew(SButton)
 					.Text(LOCTEXT("ApplyLater", "Apply Later"))
@@ -341,7 +341,7 @@ void FHardwareTargetingSettingsDetails::CustomizeDetails(IDetailLayoutBuilder& D
 		.Font(DetailBuilder.GetDetailFont())
 	]
 	.ValueContent()
-	.MaxDesiredWidth(0)
+	.MaxDesiredWidth(0.f)
 	[
 		SNew(SHorizontalBox)
 		+ SHorizontalBox::Slot()
diff --git a/Engine/Source/Editor/DetailCustomizations/Private/InstancedStaticMeshComponentDetails.cpp b/Engine/Source/Editor/DetailCustomizations/Private/InstancedStaticMeshComponentDetails.cpp
--- a/Engine/Source/Editor/DetailCustomizations/Private/InstancedStaticMeshComponentDetails.cpp
+++ b/Engine/Source/Editor/DetailCustomizations/Private/InstancedStaticMeshComponentDetails.cpp
@@ -31,7 +31,7 @@ class UObject;
 
 #define LOCTEXT_NAMESPACE "InstancedStaticMeshComponentDetails"
 
-TAutoConsoleVariable<float> CVarMaxNumInstancesDetails(
+TAutoConsoleVariable<int32> CVarMaxNumInstancesDetails(
 	TEXT("r.Editor.MaxNumInstancesDetails"),
 	512,
 	TEXT("Maximum number of instances shown in the details panel. Above this value, instances are hidden by default. \n")
diff --git a/Engine/Source/Editor/DetailCustomizations/Private/LandscapeProxyUIDetails.cpp b/Engine/Source/Editor/DetailCustomizations/Private/LandscapeProxyUIDetails.cpp
--- a/Engine/Source/Editor/DetailCustomizations/Private/LandscapeProxyUIDetails.cpp
+++ b/Engine/Source/Editor/DetailCustomizations/Private/LandscapeProxyUIDetails.cpp
@@ -256,7 +256,7 @@ void FLandscapeProxyUIDetails::CustomizeDetails( IDetailLayoutBuilder& DetailBui
 				SNew(SButton)
 				.VAlign(VAlign_Center)
 				.HAlign(HAlign_Center)
-				.ContentPadding(2)
+				.ContentPadding(2.f)
 				.Text(LOCTEXT("Button_CreateVolumes", "Create Volumes"))
 				.OnClicked_Lambda([LandscapeActor, this]() { return CreateRuntimeVirtualTextureVolume(LandscapeActor.Get()); })
 				.IsEnabled_Lambda([LandscapeActor, this]() { return IsCreateRuntimeVirtualTextureVolumeEnabled(LandscapeActor.Get()); })
@@ -371,7 +371,7 @@ void FLandscapeProxyUIDetails::CustomizeDetails( IDetailLayoutBuilder& DetailBui
 				}
 			}
 
-			return Algo::CountIf(ProxiesToBuild, [bForceRebuild](TWeakObjectPtr<ALandscapeProxy> InProxy) { return (InProxy.IsValid() && (bForceRebuild || !InProxy->IsNaniteMeshUpToDate())); });
+			return static_cast<int32>(Algo::CountIf(ProxiesToBuild, [bForceRebuild](TWeakObjectPtr<ALandscapeProxy> InProxy) { return (InProxy.IsValid() && (bForceRebuild || !InProxy->IsNaniteMeshUpToDate())); }));
 		};
 
 		auto HasAtLeastOneNaniteLandscape = [LandscapeActors]() -> bool
diff --git a/Engine/Source/Editor/DetailCustomizations/Private/LevelSequenceActorDetails.cpp b/Engine/Source/Editor/DetailCustomizations/Private/LevelSequenceActorDetails.cpp
--- a/Engine/Source/Editor/DetailCustomizations/Private/LevelSequenceActorDetails.cpp
+++ b/Engine/Source/Editor/DetailCustomizations/Private/LevelSequenceActorDetails.cpp
@@ -107,10 +107,10 @@ void FLevelSequenceActorDetails::CustomizeDetails( IDetailLayoutBuilder& DetailL
 		SNew(SHorizontalBox)
 		+SHorizontalBox::Slot()
 		.FillWidth(1.f)
-		.Padding(0, 5, 10, 5)
+		.Padding(0.f, 5.f, 10.f, 5.f)
 		[
 			SNew(SButton)
-			.ContentPadding(3)
+			.ContentPadding(3.f)
 			.VAlign(VAlign_Center)
 			.HAlign(HAlign_Center)
 			.IsEnabled( this, &FLevelSequenceActorDetails::CanOpenLevelSequenceForActor )
diff --git a/Engine/Source/Editor/DetailCustomizations/Private/LightingChannelsCustomization.cpp b/Engine/Source/Editor/DetailCustomizations/Private/LightingChannelsCustomization.cpp
--- a/Engine/Source/Editor/DetailCustomizations/Private/LightingChannelsCustomization.cpp
+++ b/Engine/Source/Editor/DetailCustomizations/Private/LightingChannelsCustomization.cpp
@@ -58,10 +58,10 @@ void FLightingChannelsCustomization::CustomizeHeader(TSharedRef<IPropertyHandle>
 		.HAlign(HAlign_Left)
 		.VAlign(VAlign_Center)
 		.AutoWidth()
-		.Padding(0, 0, bIsLastChild ? 0 : 8, 0)
+		.Padding(0.f, 0.f, bIsLastChild ? 0.f : 8.f, 0.f)
 		[
 			SNew(SBox)
-			.WidthOverride(20)
+			.WidthOverride(20.f)
 			.HAlign(HAlign_Fill)
 			.IsEnabled(this, &FLightingChannelsCustomization::IsLightingChannelButtonEditable, ChildIndex)
 			[
@@ -71,7 +71,7 @@ void FLightingChannelsCustomization::CustomizeHeader(TSharedRef<IPropertyHandle>
 				.OnCheckStateChanged(this, &FLightingChannelsCustomization::OnButtonCheckedStateChanged, ChildIndex)
 				.IsChecked(this, &FLightingChannelsCustomization::GetButtonCheckedState, ChildIndex)
 				.HAlign(HAlign_Center)
-				.Padding(FMargin(0, 2))
+				.Padding(FMargin(0.f, 2.f))
 				[
 					SNew(STextBlock)
 					.Font(FAppStyle::Get().GetFontStyle("SmallText"))
diff --git a/Engine/Source/Editor/DetailCustomizations/Private/MathStructCustomizations.cpp b/Engine/Source/Editor/DetailCustomizations/Private/MathStructCustomizations.cpp
--- a/Engine/Source/Editor/DetailCustomizations/Private/MathStructCustomizations.cpp
+++ b/Engine/Source/Editor/DetailCustomizations/Private/MathStructCustomizations.cpp
@@ -140,8 +140,8 @@ void FMathStructCustomization::MakeHeaderRow(TSharedRef<class IPropertyHandle>&
 	]
 	.ValueContent()
 	// Make enough space for each child handle
-	.MinDesiredWidth(125.0f * SortedChildHandles.Num())
-	.MaxDesiredWidth(125.0f * SortedChildHandles.Num())
+	.MinDesiredWidth(125.0f * static_cast<float>(SortedChildHandles.Num()))
+	.MaxDesiredWidth(125.0f * static_cast<float>(SortedChildHandles.Num()))
 	[
 		SAssignNew(HorizontalBox, SHorizontalBox)
 		.IsEnabled(this, &FMathStructCustomization::IsValueEnabled, StructWeakHandlePtr)
diff --git a/Engine/Source/Editor/DetailCustomizations/Private/MotionControllerSourceCustomization.cpp b/Engine/Source/Editor/DetailCustomizations/Private/MotionControllerSourceCustomization.cpp
--- a/Engine/Source/Editor/DetailCustomizations/Private/MotionControllerSourceCustomization.cpp
+++ b/Engine/Source/Editor/DetailCustomizations/Private/MotionControllerSourceCustomization.cpp
@@ -31,7 +31,7 @@ void SMotionSourceWidget::Construct(const FArguments& InArgs)
 	this->ChildSlot
 	[
 		SNew(SComboButton)
-		.ContentPadding(1)
+		.ContentPadding(1.f)
 		.OnGetMenuContent(this, &SMotionSourceWidget::BuildMotionSourceMenu)
 		.ButtonContent()
 		[
diff --git a/Engine/Source/Editor/DetailCustomizations/Private/PoseAssetDetails.cpp b/Engine/Source/Editor/DetailCustomizations/Private/PoseAssetDetails.cpp
--- a/Engine/Source/Editor/DetailCustomizations/Private/PoseAssetDetails.cpp
+++ b/Engine/Source/Editor/DetailCustomizations/Private/PoseAssetDetails.cpp
@@ -129,7 +129,7 @@ void FPoseAssetDetails::CustomizeDetails(IDetailLayoutBuilder& DetailBuilder)
 		.OnComboBoxOpening(this, &FPoseAssetDetails::OnRetargetSourceComboOpening)
 		.InitiallySelectedItem(InitialSelected)
 		.IsEnabled( FSlateApplication::Get().GetNormalExecutionAttribute() )
-		.ContentPadding(0)
+		.ContentPadding(0.f)
 		.Content()
 		[
 			SNew( STextBlock )
@@ -294,7 +294,7 @@ void FPoseAssetDetails::CustomizeDetails(IDetailLayoutBuilder& DetailBuilder)
 		SourceAnimationPropertyHandle->CreatePropertyNameWidget()
 	]
 	.ValueContent()
-	.MinDesiredWidth(200)
+	.MinDesiredWidth(200.f)
 	[
 		SNew(SVerticalBox)
 		+ SVerticalBox::Slot()
diff --git a/Engine/Source/Editor/DetailCustomizations/Private/RawDistributionVectorStructCustomization.cpp b/Engine/Source/Editor/DetailCustomizations/Private/RawDistributionVectorStructCustomization.cpp
--- a/Engine/Source/Editor/DetailCustomizations/Private/RawDistributionVectorStructCustomization.cpp
+++ b/Engine/Source/Editor/DetailCustomizations/Private/RawDistributionVectorStructCustomization.cpp
@@ -228,7 +228,7 @@ namespace
 	{
 		FVector Result;
 		ensure(StructHandle->GetValue(Result) == FPropertyAccess::Success);
-		return Result[ElementIndex];
+		return static_cast<float>(Result[ElementIndex]);
 	}
 
 	void SetColorElementValue(const TSharedPtr<IPropertyHandle>& StructHandle, int32 ElementIndex, float Value)
@@ -243,7 +243,7 @@ namespace
 	{
 		FVector Result;
 		ensure(StructHandle->GetValue(Result) == FPropertyAccess::Success);
-		return FLinearColor(Result.X, Result.Y, Result.Z, 1.0f);
+		return FLinearColor(static_cast<float>(Result.X), static_cast<float>(Result.Y), static_cast<float>(Result.Z), 1.0f);
 	}
 
 	void SetColorValue(const TSharedPtr<IPropertyHandle>& StructHandle, FLinearColor Value)
diff --git a/Engine/Source/Editor/DetailCustomizations/Private/SkeletalMeshComponentDetails.cpp b/Engine/Source/Editor/DetailCustomizations/Private/SkeletalMeshComponentDetails.cpp
--- a/Engine/Source/Editor/DetailCustomizations/Private/SkeletalMeshComponentDetails.cpp
+++ b/Engine/Source/Editor/DetailCustomizations/Private/SkeletalMeshComponentDetails.cpp
@@ -159,7 +159,7 @@ void FSkeletalMeshComponentDetails::UpdateAnimationCategory(IDetailLayoutBuilder
 			[
 				SAssignNew(ClassPickerComboButton, SComboButton)
 				.OnGetMenuContent(this, &FSkeletalMeshComponentDetails::GetClassPickerMenuContent)
-				.ContentPadding(0)
+				.ContentPadding(0.f)
 				.ButtonContent()
 				[
 					SNew(STextBlock)
@@ -211,8 +211,8 @@ void FSkeletalMeshComponentDetails::UpdateAnimationCategory(IDetailLayoutBuilder
 					NameWidget.ToSharedRef()
 				]
 				.ValueContent()
-				.MinDesiredWidth(600)
-				.MaxDesiredWidth(600)
+				.MinDesiredWidth(600.f)
+				.MaxDesiredWidth(600.f)
 				[
 					PropWidget
 				]
@@ -342,7 +342,7 @@ TSharedRef<SWidget> FSkeletalMeshComponentDetails::GetClassPickerMenuContent()
 		.ForegroundColor(FAppStyle::GetColor("DefaultForeground"))
 		[
 			SNew(SBox)
-			.WidthOverride(280)
+			.WidthOverride(280.f)
 			[
 				ClassViewerModule.CreateClassViewer(InitOptions, FOnClassPicked::CreateSP(this, &FSkeletalMeshComponentDetails::OnClassPicked))
 			]
diff --git a/Engine/Source/Editor/DetailCustomizations/Private/SkyLightComponentDetails.cpp b/Engine/Source/Editor/DetailCustomizations/Private/SkyLightComponentDetails.cpp
--- a/Engine/Source/Editor/DetailCustomizations/Private/SkyLightComponentDetails.cpp
+++ b/Engine/Source/Editor/DetailCustomizations/Private/SkyLightComponentDetails.cpp
@@ -105,7 +105,7 @@ void FSkyLightComponentDetails::CustomizeDetails( IDetailLayoutBuilder& DetailLa
 		.MinDesiredWidth(125.f)
 		[
 			SNew(SButton)
-			.ContentPadding(2)
+			.ContentPadding(2.f)
 			.VAlign(VAlign_Center)
 			.HAlign(HAlign_Center)
 			.OnClicked( this, &FSkyLightComponentDetails::OnUpdateSkyCapture )
diff --git a/Engine/Source/Editor/DetailCustomizations/Private/SlateBrushCustomization.cpp b/Engine/Source/Editor/DetailCustomizations/Private/SlateBrushCustomization.cpp
--- a/Engine/Source/Editor/DetailCustomizations/Private/SlateBrushCustomization.cpp
+++ b/Engine/Source/Editor/DetailCustomizations/Private/SlateBrushCustomization.cpp
@@ -250,7 +250,7 @@ public:
 			]
 		);
 
-		CachedTextureSize = FVector2D::ZeroVector;
+		CachedTextureSize = FVector2f::ZeroVector;
 		CachePropertyValues();
 		SetDefaultAlignment();
 		UpdatePreviewImageSize();
@@ -372,13 +372,13 @@ private:
 
 	virtual FReply OnMouseMove( const FGeometry& MyGeometry, const FPointerEvent& MouseEvent ) override
 	{
-		const FVector2D LocalMouseCoordinates( MyGeometry.AbsoluteToLocal( MouseEvent.GetScreenSpacePosition() ) );
+		const FVector2f LocalMouseCoordinates( MyGeometry.AbsoluteToLocal( MouseEvent.GetScreenSpacePosition() ) );
 
 		if( bUserIsResizing )
 		{
 			if( MouseZone >= WZ_RightBorder && MouseZone <= WZ_BottomRightBorder )
 			{
-				FVector2D Delta( LocalMouseCoordinates - ResizeAnchorPosition );
+				FVector2f Delta( LocalMouseCoordinates - ResizeAnchorPosition );
 
 				if( MouseZone == WZ_RightBorder )
 				{
@@ -403,7 +403,7 @@ private:
 	
 	virtual void OnMouseEnter( const FGeometry& MyGeometry, const FPointerEvent& MouseEvent ) override
 	{
-		const FVector2D LocalMouseCoordinates( MyGeometry.AbsoluteToLocal( MouseEvent.GetScreenSpacePosition() ) );
+		const FVector2f LocalMouseCoordinates( MyGeometry.AbsoluteToLocal( MouseEvent.GetScreenSpacePosition() ) );
 		MouseZone = FindMouseZone( LocalMouseCoordinates );
 		SBorder::OnMouseEnter( MyGeometry, MouseEvent );
 	}
@@ -444,10 +444,10 @@ private:
 	/**
 	 * Determine which zone of the widget that the mouse is in
 	 */
-	EWidgetZone FindMouseZone( const FVector2D& LocalMouseCoordinates ) const
+	EWidgetZone FindMouseZone( const FVector2f& LocalMouseCoordinates ) const
 	{
 		EWidgetZone InMouseZone = WZ_NotInWidget;
-		const FVector2D DesiredZoneSize( GetDesiredSize() );
+		const FVector2f DesiredZoneSize( GetDesiredSize() );
 
 		if( LocalMouseCoordinates.X > DesiredZoneSize.X - BorderHitSize )
 		{
@@ -664,13 +664,13 @@ private:
 					UTexture* const BaseTexture = BrushTexture;
 					FTextureCompilingManager::Get().FinishCompilation( MakeArrayView(&BaseTexture, 1) );
 				}
-				CachedTextureSize = FVector2D( BrushTexture->GetSizeX(), BrushTexture->GetSizeY() );
+				CachedTextureSize = FVector2f( static_cast<float>(BrushTexture->GetSizeX()), static_cast<float>(BrushTexture->GetSizeY()) );
 			}
 			else if ( ISlateTextureAtlasInterface* AtlasedTextureObject = Cast<ISlateTextureAtlasInterface>(ResourceObject) )
 			{
-				CachedTextureSize = AtlasedTextureObject->GetSlateAtlasData().GetSourceDimensions();
+				CachedTextureSize = UE::Slate::CastToVector2f(AtlasedTextureObject->GetSlateAtlasData().GetSourceDimensions());
 			}
-			else if( CachedTextureSize == FVector2D::ZeroVector )
+			else if( CachedTextureSize == FVector2f::ZeroVector )
 			{
 				// If the cached texture size is not initialized, create a default value now for materials 
 				CachedTextureSize = CachedImageSizeValue;
@@ -739,10 +739,10 @@ private:
 	 */
 	void UpdateMarginLinePositions()
 	{
-		const FVector2D DrawSize( (HorizontalAlignment == EHorizontalAlignment::HAlign_Fill || PreviewImageSize.X < CachedImageSizeValue.X) ? PreviewImageSize.X : CachedImageSizeValue.X,
+		const FVector2f DrawSize( (HorizontalAlignment == EHorizontalAlignment::HAlign_Fill || PreviewImageSize.X < CachedImageSizeValue.X) ? PreviewImageSize.X : CachedImageSizeValue.X,
 								  (VerticalAlignment == EVerticalAlignment::VAlign_Fill || PreviewImageSize.Y < CachedImageSizeValue.Y) ? PreviewImageSize.Y : CachedImageSizeValue.Y );
 
-		FVector2D Position( 0.0f, 0.0f );
+		FVector2f Position( 0.0f, 0.0f );
 
 		if( PreviewImageSize.X > DrawSize.X )
 		{
@@ -785,10 +785,10 @@ private:
 			BottomMargin = TopMargin;
 		}
 
-		MarginLinePositions[ MarginLine_Left ] = FVector2D( ImagePadding + Position.X + LeftMargin, 1.0f );
-		MarginLinePositions[ MarginLine_Right ] = FVector2D( ImagePadding + Position.X + RightMargin, 1.0f );
-		MarginLinePositions[ MarginLine_Top ] = FVector2D( 1.0f, ImagePadding + Position.Y + TopMargin );
-		MarginLinePositions[ MarginLine_Bottom ] = FVector2D( 1.0f, ImagePadding + Position.Y + BottomMargin );
+		MarginLinePositions[ MarginLine_Left ] = FVector2f( ImagePadding + Position.X + LeftMargin, 1.0f );
+		MarginLinePositions[ MarginLine_Right ] = FVector2f( ImagePadding + Position.X + RightMargin, 1.0f );
+		MarginLinePositions[ MarginLine_Top ] = FVector2f( 1.0f, ImagePadding + Position.Y + TopMargin );
+		MarginLinePositions[ MarginLine_Bottom ] = FVector2f( 1.0f, ImagePadding + Position.Y + BottomMargin );
 	}
 
 	/**
@@ -944,7 +944,7 @@ private:
 	 */
 	FVector2D GetLeftMarginLinePosition() const
 	{
-		return MarginLinePositions[ MarginLine_Left ];
+		return FVector2D(MarginLinePositions[ MarginLine_Left ]);
 	}
 
 	/**
@@ -952,7 +952,7 @@ private:
 	 */
 	FVector2D GetRightMarginLinePosition() const
 	{
-		return MarginLinePositions[ MarginLine_Right ];
+		return FVector2D(MarginLinePositions[ MarginLine_Right ]);
 	}
 
 	/**
@@ -960,7 +960,7 @@ private:
 	 */
 	FVector2D GetTopMarginLinePosition() const
 	{
-		return MarginLinePositions[ MarginLine_Top ];
+		return FVector2D(MarginLinePositions[ MarginLine_Top ]);
 	}
 
 	/**
@@ -968,7 +968,7 @@ private:
 	 */
 	FVector2D GetBottomMarginLinePosition() const
 	{
-		return MarginLinePositions[ MarginLine_Bottom ];
+		return FVector2D(MarginLinePositions[ MarginLine_Bottom ]);
 	}
 
 	/** Padding between the preview image and the checkerboard background */
@@ -997,8 +997,8 @@ private:
 	TSharedPtr<IPropertyHandle> ImageTypeProperty;
 
 	/** Cached Slate Brush property values */
-	FVector2D CachedTextureSize;
-	FVector2D CachedImageSizeValue;
+	FVector2f CachedTextureSize;
+	FVector2f CachedImageSizeValue;
 	ESlateBrushDrawType::Type CachedDrawAsType;
 	ESlateBrushTileType::Type CachedTilingType;
 	FMargin CachedMarginPropertyValue;
@@ -1008,10 +1008,10 @@ private:
 	EVerticalAlignment VerticalAlignment;
 
 	/** Preview image size */
-	FVector2D PreviewImageSize;
+	FVector2f PreviewImageSize;
 
 	/** Margin line positions */
-	FVector2D MarginLinePositions[ MarginLine_Count ];
+	FVector2f MarginLinePositions[ MarginLine_Count ];
 
 	/** The current widget zone the mouse is in */
 	EWidgetZone MouseZone;
@@ -1020,10 +1020,10 @@ private:
 	bool bUserIsResizing;
 
 	/** Preview resize anchor position */
-	FVector2D ResizeAnchorPosition;
+	FVector2f ResizeAnchorPosition;
 
 	/** Size of the preview image on begin of resize */
-	FVector2D ResizeAnchorSize;
+	FVector2f ResizeAnchorSize;
 };
 
 const float SSlateBrushPreview::ImagePadding = 5.0f;
@@ -1056,7 +1056,7 @@ class SSlateBrushStaticPreview : public SCompoundWidget
 				.BorderImage(this, &SSlateBrushStaticPreview::GetPropertyBrush)
 				[
 					SNew(SSpacer)
-					.Size(FVector2D(1, 1))
+					.Size(FVector2f(1.f, 1.f))
 				]
 			]
 
@@ -1091,7 +1091,7 @@ private:
 	{
 		if (TemporaryBrush.DrawAs == ESlateBrushDrawType::Image)
 		{
-			const FVector2D& Size = TemporaryBrush.ImageSize;
+			const FVector2f& Size = TemporaryBrush.ImageSize;
 			if (Size.X > 0 && Size.Y > 0)
 			{
 				return Size.X * TargetHeight / Size.Y;
@@ -1287,7 +1287,7 @@ private:
 					// GetSizeX/Y will return the incorrect value if this texture is being compiled so we need to wait for it here
 					FTextureCompilingManager::Get().FinishCompilation( MakeArrayView(&BaseTexture, 1) );
 				}
-				CachedTextureSize = ImageSizeType(BrushTexture->GetSizeX(), BrushTexture->GetSizeY());
+				CachedTextureSize = ImageSizeType(static_cast<float>(BrushTexture->GetSizeX()), static_cast<float>(BrushTexture->GetSizeY()));
 			}
 			else if ( ISlateTextureAtlasInterface* AtlasedTextureObject = Cast<ISlateTextureAtlasInterface>(ResourceObject) )
 			{
@@ -1557,10 +1557,14 @@ bool FSlateBrushStructCustomization::IsImageSizeResetToDefaultVisible(TSharedPtr
 	if (FPropertyAccess::Success == ResourceObjectProperty->GetValue(ResourceObject) && ResourceObject)
 	{
 		// get texture size from ResourceObjectProperty and compare to image size prop value
-		const FVector2D SizeDefault = GetDefaultImageSize();
+		const FVector2f SizeDefault = GetDefaultImageSize();
 
-		FVector2D Size;
-		ImageSizeProperty->GetValue(Size);
+		FVector2f Size;
+		{
+			FVector2D SizeAsDouble;
+			ImageSizeProperty->GetValue(SizeAsDouble);
+			Size = UE::Slate::CastToVector2f(SizeAsDouble);
+		}
 
 		if (PropertyHandle->GetProperty() == ImageSizeProperty->GetProperty())
 		{
@@ -1592,12 +1596,12 @@ void FSlateBrushStructCustomization::OnImageSizeResetToDefault(TSharedPtr<IPrope
 	if (FPropertyAccess::Success == ResourceObjectProperty->GetValue(ResourceObject) && ResourceObject)
 	{
 		// Set image size prop value to the texture size in ResourceObjectProperty
-		const FVector2D SizeDefault = GetDefaultImageSize();
+		const FVector2f SizeDefault = GetDefaultImageSize();
 
 		if (PropertyHandle->GetProperty() == ImageSizeProperty->GetProperty())
 		{
 			// reseting the whole vector
-			PropertyHandle->SetValue(SizeDefault);
+			PropertyHandle->SetValue(FVector2D(SizeDefault));
 		}
 		else if (PropertyHandle->GetProperty() == ImageSizeProperty->GetChildHandle(0)->GetProperty())	// X
 		{
@@ -1621,7 +1625,7 @@ void FSlateBrushStructCustomization::OnImageSizeResetToDefault(TSharedPtr<IPrope
 	}
 }
 
-FVector2D FSlateBrushStructCustomization::GetDefaultImageSize() const
+FVector2f FSlateBrushStructCustomization::GetDefaultImageSize() const
 {
 	// Custom default behavior using the texture's size, if one is set as the resource object
 	UObject* ResourceObject;
@@ -1635,14 +1639,14 @@ FVector2D FSlateBrushStructCustomization::GetDefaultImageSize() const
 				UTexture* const BaseTexture = Texture;
 				FTextureCompilingManager::Get().FinishCompilation(MakeArrayView(&BaseTexture, 1));
 			}
-			return FVector2D(Texture->GetSizeX(), Texture->GetSizeY());
+			return FVector2f(static_cast<float>(Texture->GetSizeX()), static_cast<float>(Texture->GetSizeY()));
 		}
 		else if ( ISlateTextureAtlasInterface* AtlasedTextureObject = Cast<ISlateTextureAtlasInterface>(ResourceObject) )
 		{
-			return AtlasedTextureObject->GetSlateAtlasData().GetSourceDimensions();
+			return UE::Slate::CastToVector2f(AtlasedTextureObject->GetSlateAtlasData().GetSourceDimensions());
 		}
 	}
 
 	// Fall back on the standard default size for brush images
-	return FVector2D(SlateBrushDefs::DefaultImageSize, SlateBrushDefs::DefaultImageSize);
+	return FVector2f(SlateBrushDefs::DefaultImageSize, SlateBrushDefs::DefaultImageSize);
 }
diff --git a/Engine/Source/Editor/DetailCustomizations/Private/SlateFontInfoCustomization.cpp b/Engine/Source/Editor/DetailCustomizations/Private/SlateFontInfoCustomization.cpp
--- a/Engine/Source/Editor/DetailCustomizations/Private/SlateFontInfoCustomization.cpp
+++ b/Engine/Source/Editor/DetailCustomizations/Private/SlateFontInfoCustomization.cpp
@@ -220,7 +220,7 @@ FText FSlateFontInfoStructCustomization::GetFontSizeTooltipText() const
 float FSlateFontInfoStructCustomization::ConvertFontSizeFromNativeToDisplay(float NativeFontSize)
 {
 	const UUserInterfaceSettings* UISettings = GetDefault<UUserInterfaceSettings>();
-	const float FontDisplayDPI = UISettings->GetFontDisplayDPI();
+	const float FontDisplayDPI = static_cast<float>(UISettings->GetFontDisplayDPI());
 	const float DisplayedSize = NativeFontSize * static_cast<float>(FontConstants::RenderDPI) / FontDisplayDPI;
 	return DisplayedSize;
 }
@@ -228,7 +228,7 @@ float FSlateFontInfoStructCustomization::ConvertFontSizeFromNativeToDisplay(floa
 float FSlateFontInfoStructCustomization::ConvertFontSizeFromDisplayToNative(float DisplayFontSize)
 {
 	const UUserInterfaceSettings* UISettings = GetDefault<UUserInterfaceSettings>();
-	const float FontDisplayDPI = UISettings->GetFontDisplayDPI();
+	const float FontDisplayDPI = static_cast<float>(UISettings->GetFontDisplayDPI());
 	const float NativeSize = DisplayFontSize * FontDisplayDPI / static_cast<float>(FontConstants::RenderDPI);
 	const float RoundedSize = FMath::GridSnap(NativeSize, 0.01f);
 	return RoundedSize;
diff --git a/Engine/Source/Editor/DetailCustomizations/Private/SourceCodeAccessSettingsDetails.cpp b/Engine/Source/Editor/DetailCustomizations/Private/SourceCodeAccessSettingsDetails.cpp
--- a/Engine/Source/Editor/DetailCustomizations/Private/SourceCodeAccessSettingsDetails.cpp
+++ b/Engine/Source/Editor/DetailCustomizations/Private/SourceCodeAccessSettingsDetails.cpp
@@ -66,14 +66,14 @@ void FSourceCodeAccessSettingsDetails::CustomizeDetails( IDetailLayoutBuilder& D
 		PreferredProviderPropertyHandle->CreatePropertyNameWidget()
 	]
 	.ValueContent()
-	.MinDesiredWidth(113)
-	.MaxDesiredWidth(113)
+	.MinDesiredWidth(113.f)
+	.MaxDesiredWidth(113.f)
 	[
 		SNew(SComboBox< TSharedPtr<FAccessorItem>>)
 		.ToolTipText(LOCTEXT("PreferredAccessorToolTip", "Choose the way to access source code."))
 		.OptionsSource(&Accessors)
 		.OnSelectionChanged(this, &FSourceCodeAccessSettingsDetails::OnSelectionChanged, PreferredProviderPropertyHandle)
-		.ContentPadding(2)
+		.ContentPadding(2.f)
 		.OnGenerateWidget(this, &FSourceCodeAccessSettingsDetails::OnGenerateWidget)
 		.Content()
 		[
diff --git a/Engine/Source/Editor/DetailCustomizations/Private/SplineComponentDetails.cpp b/Engine/Source/Editor/DetailCustomizations/Private/SplineComponentDetails.cpp
--- a/Engine/Source/Editor/DetailCustomizations/Private/SplineComponentDetails.cpp
+++ b/Engine/Source/Editor/DetailCustomizations/Private/SplineComponentDetails.cpp
@@ -191,9 +191,9 @@ private:
 			return !X.IsSet() || !Y.IsSet() || !Z.IsSet();
 		}
 
-		TOptional<float> X;
-		TOptional<float> Y;
-		TOptional<float> Z;
+		TOptional<FVector::FReal> X;
+		TOptional<FVector::FReal> Y;
+		TOptional<FVector::FReal> Z;
 		bool bInitialized;
 	};
 
@@ -230,9 +230,9 @@ private:
 			return !Roll.IsSet() || !Pitch.IsSet() || !Yaw.IsSet();
 		}
 
-		TOptional<float> Roll;
-		TOptional<float> Pitch;
-		TOptional<float> Yaw;
+		TOptional<FRotator::FReal> Roll;
+		TOptional<FRotator::FReal> Pitch;
+		TOptional<FRotator::FReal> Yaw;
 		bool bInitialized;
 	};
 
@@ -242,27 +242,27 @@ private:
 	bool ArePointsSelected() const { return (SelectedKeys.Num() > 0); };
 	bool AreNoPointsSelected() const { return (SelectedKeys.Num() == 0); };
 	TOptional<float> GetInputKey() const { return InputKey.Value; }
-	TOptional<float> GetPositionX() const { return Position.X; }
-	TOptional<float> GetPositionY() const { return Position.Y; }
-	TOptional<float> GetPositionZ() const { return Position.Z; }
-	TOptional<float> GetArriveTangentX() const { return ArriveTangent.X; }
-	TOptional<float> GetArriveTangentY() const { return ArriveTangent.Y; }
-	TOptional<float> GetArriveTangentZ() const { return ArriveTangent.Z; }
-	TOptional<float> GetLeaveTangentX() const { return LeaveTangent.X; }
-	TOptional<float> GetLeaveTangentY() const { return LeaveTangent.Y; }
-	TOptional<float> GetLeaveTangentZ() const { return LeaveTangent.Z; }
-	TOptional<float> GetRotationRoll() const { return Rotation.Roll; }
-	TOptional<float> GetRotationPitch() const { return Rotation.Pitch; }
-	TOptional<float> GetRotationYaw() const { return Rotation.Yaw; }
-	TOptional<float> GetScaleX() const { return Scale.X; }
-	TOptional<float> GetScaleY() const { return Scale.Y; }
-	TOptional<float> GetScaleZ() const { return Scale.Z; }
+	TOptional<FVector::FReal> GetPositionX() const { return Position.X; }
+	TOptional<FVector::FReal> GetPositionY() const { return Position.Y; }
+	TOptional<FVector::FReal> GetPositionZ() const { return Position.Z; }
+	TOptional<FVector::FReal> GetArriveTangentX() const { return ArriveTangent.X; }
+	TOptional<FVector::FReal> GetArriveTangentY() const { return ArriveTangent.Y; }
+	TOptional<FVector::FReal> GetArriveTangentZ() const { return ArriveTangent.Z; }
+	TOptional<FVector::FReal> GetLeaveTangentX() const { return LeaveTangent.X; }
+	TOptional<FVector::FReal> GetLeaveTangentY() const { return LeaveTangent.Y; }
+	TOptional<FVector::FReal> GetLeaveTangentZ() const { return LeaveTangent.Z; }
+	TOptional<FRotator::FReal> GetRotationRoll() const { return Rotation.Roll; }
+	TOptional<FRotator::FReal> GetRotationPitch() const { return Rotation.Pitch; }
+	TOptional<FRotator::FReal> GetRotationYaw() const { return Rotation.Yaw; }
+	TOptional<FVector::FReal> GetScaleX() const { return Scale.X; }
+	TOptional<FVector::FReal> GetScaleY() const { return Scale.Y; }
+	TOptional<FVector::FReal> GetScaleZ() const { return Scale.Z; }
 	void OnSetInputKey(float NewValue, ETextCommit::Type CommitInfo);
-	void OnSetPosition(float NewValue, ETextCommit::Type CommitInfo, EAxis::Type Axis);
-	void OnSetArriveTangent(float NewValue, ETextCommit::Type CommitInfo, EAxis::Type Axis);
-	void OnSetLeaveTangent(float NewValue, ETextCommit::Type CommitInfo, EAxis::Type Axis);
-	void OnSetRotation(float NewValue, ETextCommit::Type CommitInfo, EAxis::Type Axis);
-	void OnSetScale(float NewValue, ETextCommit::Type CommitInfo, EAxis::Type Axis);
+	void OnSetPosition(FVector::FReal NewValue, ETextCommit::Type CommitInfo, EAxis::Type Axis);
+	void OnSetArriveTangent(FVector::FReal NewValue, ETextCommit::Type CommitInfo, EAxis::Type Axis);
+	void OnSetLeaveTangent(FVector::FReal NewValue, ETextCommit::Type CommitInfo, EAxis::Type Axis);
+	void OnSetRotation(FRotator::FReal NewValue, ETextCommit::Type CommitInfo, EAxis::Type Axis);
+	void OnSetScale(FVector::FReal NewValue, ETextCommit::Type CommitInfo, EAxis::Type Axis);
 	FText GetPointType() const;
 	void OnSplinePointTypeChanged(TSharedPtr<FString> NewValue, ESelectInfo::Type SelectInfo);
 	TSharedRef<SWidget> OnGenerateComboWidget(TSharedPtr<FString> InComboString);
@@ -291,7 +291,7 @@ private:
 	bool IsTransformEditingAbsolute(ESplinePointProperty SplinePointProperty) const;
 	bool IsTransformEditingRelative(ESplinePointProperty SplinePointProperty) const;
 	FText GetSplinePointPropertyText(ESplinePointProperty SplinePointProp) const;
-	void SetSplinePointProperty(ESplinePointProperty SplinePointProp, FVector NewValue, EAxisList::Type Axis, bool bCommitted);
+	void SetSplinePointProperty(ESplinePointProperty SplinePointProp, FVector3f NewValue, EAxisList::Type Axis, bool bCommitted);
 
 	FUIAction CreateCopyAction(ESplinePointProperty SplinePointProp);
 	FUIAction CreatePasteAction(ESplinePointProperty SplinePointProp);
@@ -305,7 +305,7 @@ private:
 
 	void OnBeginPositionSlider();
 	void OnBeginScaleSlider();
-	void OnEndSlider(float);
+	void OnEndSlider(FVector::FReal);
 
 	USplineComponent* SplineComp;
 	USplineComponent* SplineCompArchetype;
@@ -348,7 +348,7 @@ FSplinePointDetails::FSplinePointDetails(USplineComponent* InOwningSplineCompone
 	check(SplinePointTypeEnum);
 	for (int32 EnumIndex = 0; EnumIndex < SplinePointTypeEnum->NumEnums() - 1; ++EnumIndex)
 	{
-		const int32 Value = SplinePointTypeEnum->GetValueByIndex(EnumIndex);
+		const ESplinePointType::Type Value = static_cast<ESplinePointType::Type>(SplinePointTypeEnum->GetValueByIndex(EnumIndex));
 		if (EnabledSplinePointTypes.Contains(Value))
 		{
 			SplinePointTypes.Add(MakeShareable(new FString(SplinePointTypeEnum->GetNameStringByIndex(EnumIndex))));
@@ -567,7 +567,7 @@ void FSplinePointDetails::GenerateChildContent(IDetailChildrenBuilder& ChildrenB
 			.MinDesiredWidth(375.0f)
 			.MaxDesiredWidth(375.0f)
 			[
-				SNew(SVectorInputBox)
+				SNew(SNumericVectorInputBox<FVector::FReal>)
 				.X(this, &FSplinePointDetails::GetPositionX)
 				.Y(this, &FSplinePointDetails::GetPositionY)
 				.Z(this, &FSplinePointDetails::GetPositionZ)
@@ -606,7 +606,7 @@ void FSplinePointDetails::GenerateChildContent(IDetailChildrenBuilder& ChildrenB
 			.MinDesiredWidth(375.0f)
 			.MaxDesiredWidth(375.0f)
 			[
-				SNew(SRotatorInputBox)
+				SNew(SNumericRotatorInputBox<FRotator::FReal>)
 				.Roll(this, &FSplinePointDetails::GetRotationRoll)
 				.Pitch(this, &FSplinePointDetails::GetRotationPitch)
 				.Yaw(this, &FSplinePointDetails::GetRotationYaw)
@@ -641,21 +641,21 @@ void FSplinePointDetails::GenerateChildContent(IDetailChildrenBuilder& ChildrenB
 			.MinDesiredWidth(375.0f)
 			.MaxDesiredWidth(375.0f)
 			[
-				SNew(SVectorInputBox)
+				SNew(SNumericVectorInputBox<FVector::FReal>)
 				.X(this, &FSplinePointDetails::GetScaleX)
-			.Y(this, &FSplinePointDetails::GetScaleY)
-			.Z(this, &FSplinePointDetails::GetScaleZ)
-			.AllowSpin(true)
-			.bColorAxisLabels(true)
-			.OnXChanged(this, &FSplinePointDetails::OnSetScale, ETextCommit::Default, EAxis::X)
-			.OnYChanged(this, &FSplinePointDetails::OnSetScale, ETextCommit::Default, EAxis::Y)
-			.OnZChanged(this, &FSplinePointDetails::OnSetScale, ETextCommit::Default, EAxis::Z)
-			.OnXCommitted(this, &FSplinePointDetails::OnSetScale, EAxis::X)
-			.OnYCommitted(this, &FSplinePointDetails::OnSetScale, EAxis::Y)
-			.OnZCommitted(this, &FSplinePointDetails::OnSetScale, EAxis::Z)
-			.OnBeginSliderMovement(this, &FSplinePointDetails::OnBeginScaleSlider)
-			.OnEndSliderMovement(this, &FSplinePointDetails::OnEndSlider)
-			.Font(IDetailLayoutBuilder::GetDetailFont())
+				.Y(this, &FSplinePointDetails::GetScaleY)
+				.Z(this, &FSplinePointDetails::GetScaleZ)
+				.AllowSpin(true)
+				.bColorAxisLabels(true)
+				.OnXChanged(this, &FSplinePointDetails::OnSetScale, ETextCommit::Default, EAxis::X)
+				.OnYChanged(this, &FSplinePointDetails::OnSetScale, ETextCommit::Default, EAxis::Y)
+				.OnZChanged(this, &FSplinePointDetails::OnSetScale, ETextCommit::Default, EAxis::Z)
+				.OnXCommitted(this, &FSplinePointDetails::OnSetScale, EAxis::X)
+				.OnYCommitted(this, &FSplinePointDetails::OnSetScale, EAxis::Y)
+				.OnZCommitted(this, &FSplinePointDetails::OnSetScale, EAxis::Z)
+				.OnBeginSliderMovement(this, &FSplinePointDetails::OnBeginScaleSlider)
+				.OnEndSliderMovement(this, &FSplinePointDetails::OnEndSlider)
+				.Font(IDetailLayoutBuilder::GetDetailFont())
 			];
 	}
 
@@ -684,7 +684,7 @@ void FSplinePointDetails::GenerateChildContent(IDetailChildrenBuilder& ChildrenB
 			.MinDesiredWidth(375.0f)
 			.MaxDesiredWidth(375.0f)
 			[
-				SNew(SVectorInputBox)
+				SNew(SNumericVectorInputBox<FVector::FReal>)
 				.X(this, &FSplinePointDetails::GetArriveTangentX)
 				.Y(this, &FSplinePointDetails::GetArriveTangentY)
 				.Z(this, &FSplinePointDetails::GetArriveTangentZ)
@@ -721,7 +721,7 @@ void FSplinePointDetails::GenerateChildContent(IDetailChildrenBuilder& ChildrenB
 			.MinDesiredWidth(375.0f)
 			.MaxDesiredWidth(375.0f)
 			[
-				SNew(SVectorInputBox)
+				SNew(SNumericVectorInputBox<FVector::FReal>)
 				.X(this, &FSplinePointDetails::GetLeaveTangentX)
 				.Y(this, &FSplinePointDetails::GetLeaveTangentY)
 				.Z(this, &FSplinePointDetails::GetLeaveTangentZ)
@@ -988,7 +988,7 @@ void FSplinePointDetails::OnSetInputKey(float NewValue, ETextCommit::Type Commit
 	GEditor->RedrawLevelEditingViewports(true);
 }
 
-void FSplinePointDetails::OnSetPosition(float NewValue, ETextCommit::Type CommitInfo, EAxis::Type Axis)
+void FSplinePointDetails::OnSetPosition(FVector::FReal NewValue, ETextCommit::Type CommitInfo, EAxis::Type Axis)
 {
 	if (!SplineComp)
 	{
@@ -1043,7 +1043,7 @@ void FSplinePointDetails::OnSetPosition(float NewValue, ETextCommit::Type Commit
 	GEditor->RedrawLevelEditingViewports(true);
 }
 
-void FSplinePointDetails::OnSetArriveTangent(float NewValue, ETextCommit::Type CommitInfo, EAxis::Type Axis)
+void FSplinePointDetails::OnSetArriveTangent(FVector::FReal NewValue, ETextCommit::Type CommitInfo, EAxis::Type Axis)
 {
 	if (!SplineComp)
 	{
@@ -1086,7 +1086,7 @@ void FSplinePointDetails::OnSetArriveTangent(float NewValue, ETextCommit::Type C
 	GEditor->RedrawLevelEditingViewports(true);
 }
 
-void FSplinePointDetails::OnSetLeaveTangent(float NewValue, ETextCommit::Type CommitInfo, EAxis::Type Axis)
+void FSplinePointDetails::OnSetLeaveTangent(FVector::FReal NewValue, ETextCommit::Type CommitInfo, EAxis::Type Axis)
 {
 	if (!SplineComp)
 	{
@@ -1129,7 +1129,7 @@ void FSplinePointDetails::OnSetLeaveTangent(float NewValue, ETextCommit::Type Co
 	GEditor->RedrawLevelEditingViewports(true);
 }
 
-void FSplinePointDetails::OnSetRotation(float NewValue, ETextCommit::Type CommitInfo, EAxis::Type Axis)
+void FSplinePointDetails::OnSetRotation(FRotator::FReal NewValue, ETextCommit::Type CommitInfo, EAxis::Type Axis)
 {
 	if (!SplineComp)
 	{
@@ -1210,7 +1210,7 @@ void FSplinePointDetails::OnSetRotation(float NewValue, ETextCommit::Type Commit
 	GEditor->RedrawLevelEditingViewports(true);
 }
 
-void FSplinePointDetails::OnSetScale(float NewValue, ETextCommit::Type CommitInfo, EAxis::Type Axis)
+void FSplinePointDetails::OnSetScale(FVector::FReal NewValue, ETextCommit::Type CommitInfo, EAxis::Type Axis)
 {
 	if (!SplineComp)
 	{
@@ -1498,7 +1498,7 @@ TSharedRef<SWidget> FSplinePointDetails::BuildSplinePointPropertyLabel(ESplinePo
 
 	return
 		SNew(SComboButton)
-		.ContentPadding(0)
+		.ContentPadding(0.f)
 		.ButtonStyle(FAppStyle::Get(), "NoBorder")
 		.ForegroundColor(FSlateColor::UseForeground())
 		.MenuContent()
@@ -1578,7 +1578,7 @@ FText FSplinePointDetails::GetSplinePointPropertyText(ESplinePointProperty Splin
 	return FText::GetEmpty();
 }
 
-void FSplinePointDetails::SetSplinePointProperty(ESplinePointProperty SplinePointProp, FVector NewValue, EAxisList::Type Axis, bool bCommitted)
+void FSplinePointDetails::SetSplinePointProperty(ESplinePointProperty SplinePointProp, FVector3f NewValue, EAxisList::Type Axis, bool bCommitted)
 {
 	switch (SplinePointProp)
 	{
@@ -1721,7 +1721,7 @@ void FSplinePointDetails::PasteFromText(
 	{
 	case ESplinePointProperty::Location:
 		{
-			FVector NewLocation;
+			FVector3f NewLocation;
 			if (NewLocation.InitFromString(PastedText))
 			{
 				FScopedTransaction Transaction(LOCTEXT("PasteLocation", "Paste Location"));
@@ -1731,7 +1731,7 @@ void FSplinePointDetails::PasteFromText(
 		}
 	case ESplinePointProperty::Rotation:
 		{
-			FVector NewRotation;
+			FVector3f NewRotation;
 			PastedText.ReplaceInline(TEXT("Pitch="), TEXT("X="));
 			PastedText.ReplaceInline(TEXT("Yaw="), TEXT("Y="));
 			PastedText.ReplaceInline(TEXT("Roll="), TEXT("Z="));
@@ -1744,7 +1744,7 @@ void FSplinePointDetails::PasteFromText(
 		}
 	case ESplinePointProperty::Scale:
 		{
-			FVector NewScale;
+			FVector3f NewScale;
 			if (NewScale.InitFromString(PastedText))
 			{
 				FScopedTransaction Transaction(LOCTEXT("PasteScale", "Paste Scale"));
@@ -1754,7 +1754,7 @@ void FSplinePointDetails::PasteFromText(
 		}
 	case ESplinePointProperty::ArriveTangent:
 		{
-			FVector NewArrive;
+			FVector3f NewArrive;
 			if (NewArrive.InitFromString(PastedText))
 			{
 				FScopedTransaction Transaction(LOCTEXT("PasteArriveTangent", "Paste Arrive Tangent"));
@@ -1764,7 +1764,7 @@ void FSplinePointDetails::PasteFromText(
 		}
 	case ESplinePointProperty::LeaveTangent:
 		{
-			FVector NewLeave;
+			FVector3f NewLeave;
 			if (NewLeave.InitFromString(PastedText))
 			{
 				FScopedTransaction Transaction(LOCTEXT("PasteLeaveTangent", "Paste Leave Tangent"));
@@ -1797,7 +1797,7 @@ void FSplinePointDetails::OnBeginScaleSlider()
 	GEditor->BeginTransaction(LOCTEXT("SetSplinePointScale", "Set spline point scale"));
 }
 
-void FSplinePointDetails::OnEndSlider(float)
+void FSplinePointDetails::OnEndSlider(FVector::FReal)
 {
 	bInSliderTransaction = false;
 	GEditor->EndTransaction();
diff --git a/Engine/Source/Editor/DetailCustomizations/Private/TextureLODSettingsDetails.cpp b/Engine/Source/Editor/DetailCustomizations/Private/TextureLODSettingsDetails.cpp
--- a/Engine/Source/Editor/DetailCustomizations/Private/TextureLODSettingsDetails.cpp
+++ b/Engine/Source/Editor/DetailCustomizations/Private/TextureLODSettingsDetails.cpp
@@ -225,7 +225,7 @@ void FTextureLODGroupLayout::GenerateChildContent(IDetailChildrenBuilder& Childr
 				.OnGenerateWidget(this, &FTextureLODGroupLayout::MakeMinMagFilterComboWidget)
 				.OnSelectionChanged(this, &FTextureLODGroupLayout::OnMinMagFilterChanged)
 				.InitiallySelectedItem(FilterComboList[0])
-				.ContentPadding(0)
+				.ContentPadding(0.f)
 				.Content()
 				[
 					SNew(STextBlock)
@@ -251,7 +251,7 @@ void FTextureLODGroupLayout::GenerateChildContent(IDetailChildrenBuilder& Childr
 				.OnGenerateWidget(this, &FTextureLODGroupLayout::MakeMipFilterComboWidget)
 				.OnSelectionChanged(this, &FTextureLODGroupLayout::OnMipFilterChanged)
 				.InitiallySelectedItem(FilterComboList[0])
-				.ContentPadding(0)
+				.ContentPadding(0.f)
 				.Content()
 				[
 					SNew(STextBlock)
@@ -279,7 +279,7 @@ void FTextureLODGroupLayout::GenerateChildContent(IDetailChildrenBuilder& Childr
 				.OnGenerateWidget(this, &FTextureLODGroupLayout::MakeMipGenSettingsComboWidget)
 				.OnSelectionChanged(this, &FTextureLODGroupLayout::OnMipGenSettingsChanged)
 				.InitiallySelectedItem(MipGenSettingsComboList[0])
-				.ContentPadding(0)
+				.ContentPadding(0.f)
 				.Content()
 				[
 					SNew(STextBlock)
diff --git a/Engine/Source/Editor/DetailCustomizations/Private/WorldSettingsDetails.cpp b/Engine/Source/Editor/DetailCustomizations/Private/WorldSettingsDetails.cpp
--- a/Engine/Source/Editor/DetailCustomizations/Private/WorldSettingsDetails.cpp
+++ b/Engine/Source/Editor/DetailCustomizations/Private/WorldSettingsDetails.cpp
@@ -401,8 +401,8 @@ TSharedRef<SWidget> FLightmapCustomNodeBuilder::MakeLightMapList(TSharedPtr<FLig
 			.VAlign(VAlign_Center)
 			[
 				SNew( SBox )
-				.WidthOverride( ThumbnailResolution + ThumbnailBoxPadding * 2 )
-				.HeightOverride( ThumbnailResolution + ThumbnailBoxPadding * 2 )
+				.WidthOverride( static_cast<float>(ThumbnailResolution + ThumbnailBoxPadding * 2) )
+				.HeightOverride( static_cast<float>(ThumbnailResolution + ThumbnailBoxPadding * 2) )
 				[
 					// Drop shadow border
 					SNew(SBorder)
@@ -416,14 +416,14 @@ TSharedRef<SWidget> FLightmapCustomNodeBuilder::MakeLightMapList(TSharedPtr<FLig
 
 			+SHorizontalBox::Slot()
 			.AutoWidth()
-			.Padding(6, 0, 0, 0)
+			.Padding(6.f, 0.f, 0.f, 0.f)
 			.VAlign(VAlign_Center)
 			[
 				SNew(SVerticalBox)
 
 				+SVerticalBox::Slot()
 				.AutoHeight()
-				.Padding(0, 1)
+				.Padding(0.f, 1.f)
 				[
 					SNew(STextBlock)
 					.Font(IDetailLayoutBuilder::GetDetailFont())
@@ -432,7 +432,7 @@ TSharedRef<SWidget> FLightmapCustomNodeBuilder::MakeLightMapList(TSharedPtr<FLig
 
 				+SVerticalBox::Slot()
 				.AutoHeight()
-				.Padding(0, 1)
+				.Padding(0.f, 1.f)
 				[
 					// Class
 					SNew(STextBlock)
diff --git a/Engine/Source/Editor/DetailCustomizations/Public/Customizations/SlateBrushCustomization.h b/Engine/Source/Editor/DetailCustomizations/Public/Customizations/SlateBrushCustomization.h
--- a/Engine/Source/Editor/DetailCustomizations/Public/Customizations/SlateBrushCustomization.h
+++ b/Engine/Source/Editor/DetailCustomizations/Public/Customizations/SlateBrushCustomization.h
@@ -47,7 +47,7 @@ private:
 	void OnImageSizeResetToDefault(TSharedPtr<IPropertyHandle> PropertyHandle) const;
 
 	/** Gets the current default image size based on the current texture resource */
-	FVector2D GetDefaultImageSize() const;
+	FVector2f GetDefaultImageSize() const;
 
 	/** Slate Brush DrawAs property */
 	TSharedPtr<IPropertyHandle> DrawAsProperty;
diff --git a/Engine/Source/Editor/KismetWidgets/Private/SScrubWidget.cpp b/Engine/Source/Editor/KismetWidgets/Private/SScrubWidget.cpp
--- a/Engine/Source/Editor/KismetWidgets/Private/SScrubWidget.cpp
+++ b/Engine/Source/Editor/KismetWidgets/Private/SScrubWidget.cpp
@@ -40,8 +40,13 @@ struct FSlateBrush;
 
 static const float MinStepLengh=15.f;
 
-/** This function is used by a few random widgets and is mostly arbitrary. It could be moved anywhere. */
 int32 SScrubWidget::GetDivider(float InputMinX, float InputMaxX, FVector2D WidgetSize, float SequenceLength, int32 NumFrames)
+{
+	return GetDivider(InputMinX, InputMaxX, UE::Slate::CastToVector2f(WidgetSize), SequenceLength, NumFrames);
+}
+
+/** This function is used by a few random widgets and is mostly arbitrary. It could be moved anywhere. */
+int32 SScrubWidget::GetDivider(float InputMinX, float InputMaxX, FVector2f WidgetSize, float SequenceLength, int32 NumFrames)
 {
 	FTrackScaleInfo TimeScaleInfo(InputMinX, InputMaxX, 0.f, 0.f, WidgetSize);
 	const FAnimKeyHelper Helper(SequenceLength, NumFrames);
diff --git a/Engine/Source/Editor/KismetWidgets/Public/SScrubWidget.h b/Engine/Source/Editor/KismetWidgets/Public/SScrubWidget.h
--- a/Engine/Source/Editor/KismetWidgets/Public/SScrubWidget.h
+++ b/Engine/Source/Editor/KismetWidgets/Public/SScrubWidget.h
@@ -105,7 +105,9 @@ public:
 	virtual FReply OnMouseWheel( const FGeometry& MyGeometry, const FPointerEvent& MouseEvent ) override;
 	virtual FCursorReply OnCursorQuery( const FGeometry& MyGeometry, const FPointerEvent& CursorEvent ) const override;
 
+	UE_DEPRECATED(5.6, "Use the helper function accepting FVector2f instead of FVector2D")
 	static int32 GetDivider(float InputMinX, float InputMaxX, FVector2D WidgetSize, float SequenceLength, int32 NumFrames);
+	static int32 GetDivider(float InputMinX, float InputMaxX, FVector2f WidgetSize, float SequenceLength, int32 NumFrames);
 protected:
 
 	/**
diff --git a/Engine/Source/Editor/Persona/Private/SAnimNotifyPanel.cpp b/Engine/Source/Editor/Persona/Private/SAnimNotifyPanel.cpp
--- a/Engine/Source/Editor/Persona/Private/SAnimNotifyPanel.cpp
+++ b/Engine/Source/Editor/Persona/Private/SAnimNotifyPanel.cpp
@@ -68,7 +68,7 @@ DECLARE_DELEGATE_OneParam( FOnDeleteNotify, struct FAnimNotifyEvent*)
 DECLARE_DELEGATE_RetVal_FourParams( FReply, FOnNotifyNodeDragStarted, TSharedRef<SAnimNotifyNode>, const FPointerEvent&, const FVector2D&, const bool)
 DECLARE_DELEGATE_RetVal_FiveParams(FReply, FOnNotifyNodesDragStarted, TArray<TSharedPtr<SAnimNotifyNode>>, TSharedRef<SWidget>, const FVector2D&, const FVector2D&, const bool)
 DECLARE_DELEGATE_RetVal( float, FOnGetDraggedNodePos )
-DECLARE_DELEGATE_TwoParams( FPanTrackRequest, int32, FVector2D)
+DECLARE_DELEGATE_TwoParams( FPanTrackRequest, int32, FVector2f)
 DECLARE_DELEGATE(FCopyNodes)
 DECLARE_DELEGATE_FourParams(FPasteNodes, SAnimNotifyTrack*, float, ENotifyPasteMode::Type, ENotifyPasteMultipleMode::Type)
 DECLARE_DELEGATE_RetVal_OneParam(EVisibility, FOnGetTimingNodeVisibilityForNode, TSharedPtr<SAnimNotifyNode>)
@@ -587,7 +587,7 @@ private:
 
 	TAttribute<float>			ViewInputMin;
 	TAttribute<float>			ViewInputMax;
-	FVector2D					CachedAllotedGeometrySize;
+	FVector2f					CachedAllotedGeometrySize;
 	FVector2D					ScreenPosition;
 	float						LastSnappedTime;
 
@@ -1320,7 +1320,7 @@ public:
 
 		SelectionBeginPosition.X += SnapMovement;
 
-		CurrentDragXPosition = static_cast<float>(TrackGeom.AbsoluteToLocal(FVector2D(SelectionBeginPosition.X,0.0f)).X);
+		CurrentDragXPosition = static_cast<float>(TrackGeom.AbsoluteToLocal(FVector2f(SelectionBeginPosition.X,0.0f)).X);
 
 		CursorDecoratorWindow->MoveWindowTo(TrackGeom.LocalToAbsolute(TrackGeom.AbsoluteToLocal(SelectionBeginPosition) - SelectedNodes[0]->GetNotifyPositionOffset()));
 		NodeGroupPosition = SelectionBeginPosition;
@@ -1332,12 +1332,12 @@ public:
 		if(LocalMouseXPos < LocalViewportMin && LocalViewportMin > LocalTrackMin - 10.0f)
 		{
 			const float ScreenDelta = FMath::Max(LocalMouseXPos - LocalViewportMin, -10.0f);
-			RequestTrackPan.Execute(static_cast<int32>(ScreenDelta), FVector2D(LocalTrackWidth, 1.f));
+			RequestTrackPan.Execute(static_cast<int32>(ScreenDelta), FVector2f(LocalTrackWidth, 1.f));
 		}
 		else if(LocalMouseXPos > LocalViewportMax && LocalViewportMax < LocalTrackMax + 10.0f)
 		{
 			const float ScreenDelta =  FMath::Max(LocalMouseXPos - LocalViewportMax, 10.0f);
-			RequestTrackPan.Execute(static_cast<int32>(ScreenDelta), FVector2D(LocalTrackWidth, 1.f));
+			RequestTrackPan.Execute(static_cast<int32>(ScreenDelta), FVector2f(LocalTrackWidth, 1.f));
 		}
 
 		OnNodesBeingDragged.ExecuteIfBound(SelectedNodes, DragDropEvent, CurrentDragXPosition, TrackScaleInfo.LocalXToInput(CurrentDragXPosition));
diff --git a/Engine/Source/Editor/Persona/Private/SAnimTrackPanel.cpp b/Engine/Source/Editor/Persona/Private/SAnimTrackPanel.cpp
--- a/Engine/Source/Editor/Persona/Private/SAnimTrackPanel.cpp
+++ b/Engine/Source/Editor/Persona/Private/SAnimTrackPanel.cpp
@@ -77,7 +77,7 @@ TSharedRef<class S2ColumnWidget> SAnimTrackPanel::Create2ColumnWidget( TSharedRe
 	return NewTrack.ToSharedRef();
 }
 
-void SAnimTrackPanel::PanInputViewRange(int32 ScreenDelta, FVector2D ScreenViewSize)
+void SAnimTrackPanel::PanInputViewRange(int32 ScreenDelta, FVector2f ScreenViewSize)
 {
 	FTrackScaleInfo ScaleInfo(ViewInputMin.Get(),  ViewInputMax.Get(), 0.f, 0.f, ScreenViewSize);
 
@@ -89,6 +89,11 @@ void SAnimTrackPanel::PanInputViewRange(int32 ScreenDelta, FVector2D ScreenViewS
 	InputViewRangeChanged(NewViewInputMin, NewViewInputMax);
 }
 
+void SAnimTrackPanel::PanInputViewRange(int32 ScreenDelta, FVector2D ScreenViewSize)
+{
+	PanInputViewRange(ScreenDelta, UE::Slate::CastToVector2f(ScreenViewSize));
+}
+
 void SAnimTrackPanel::InputViewRangeChanged(float ViewMin, float ViewMax)
 {
 	OnSetInputViewRange.ExecuteIfBound(ViewMin, ViewMax);
diff --git a/Engine/Source/Editor/Persona/Private/SAnimTrackPanel.h b/Engine/Source/Editor/Persona/Private/SAnimTrackPanel.h
--- a/Engine/Source/Editor/Persona/Private/SAnimTrackPanel.h
+++ b/Engine/Source/Editor/Persona/Private/SAnimTrackPanel.h
@@ -93,7 +93,9 @@ public:
 
 	void Construct(const FArguments& InArgs);
 
+	UE_DEPRECATED(5.6, "Use the helper function accepting FVector2f instead of FVector2D")
 	void PanInputViewRange(int32 ScreenDelta, FVector2D ScreenViewSize);
+	void PanInputViewRange(int32 ScreenDelta, FVector2f ScreenViewSize);
 
 	virtual float GetSequenceLength() const {return 0.0f;}
 
diff --git a/Engine/Source/Editor/UnrealEd/Public/SCurveEditor.h b/Engine/Source/Editor/UnrealEd/Public/SCurveEditor.h
--- a/Engine/Source/Editor/UnrealEd/Public/SCurveEditor.h
+++ b/Engine/Source/Editor/UnrealEd/Public/SCurveEditor.h
@@ -51,21 +51,31 @@ struct FTrackScaleInfo
 	float ViewOutputRange;
 	float PixelsPerOutput;
 
-	FVector2D WidgetSize;
+	FVector2f WidgetSize;
 
-	FTrackScaleInfo(float InViewMinInput, float InViewMaxInput, float InViewMinOutput, float InViewMaxOutput, const FVector2D InWidgetSize)
+	FTrackScaleInfo(float InViewMinInput, float InViewMaxInput, float InViewMinOutput, float InViewMaxOutput, const FVector2f& InWidgetSize)
 	{
 		WidgetSize = InWidgetSize;
 
 		ViewMinInput = InViewMinInput;
 		ViewMaxInput = InViewMaxInput;
 		ViewInputRange = ViewMaxInput - ViewMinInput;
-		PixelsPerInput = ViewInputRange > 0 ? ( WidgetSize.X / ViewInputRange ) : 0;
+		PixelsPerInput = ViewInputRange > 0.f ? ( WidgetSize.X / ViewInputRange ) : 0.f;
 
 		ViewMinOutput = InViewMinOutput;
 		ViewMaxOutput = InViewMaxOutput;
 		ViewOutputRange = InViewMaxOutput - InViewMinOutput;
-		PixelsPerOutput = ViewOutputRange > 0 ? ( WidgetSize.Y / ViewOutputRange ) : 0;
+		PixelsPerOutput = ViewOutputRange > 0.f ? ( WidgetSize.Y / ViewOutputRange ) : 0.f;
+	}
+
+	UE_DEPRECATED(5.6, "Use the FTrackScaleInfo constructor accepting FVector2f instead of FVector2D")
+	FTrackScaleInfo(float InViewMinInput, float InViewMaxInput, float InViewMinOutput, float InViewMaxOutput, const FVector2D InWidgetSize)
+		: FTrackScaleInfo(InViewMinInput
+		, InViewMaxInput
+		, InViewMinOutput
+		, InViewMaxOutput
+		, UE::Slate::CastToVector2f(InWidgetSize))
+	{
 	}
 
 	/** Local Widget Space -> Curve Input domain. */
diff --git a/Engine/Source/Runtime/AdvancedWidgets/Public/Widgets/ColorGrading/SColorGradingComponentViewer.h b/Engine/Source/Runtime/AdvancedWidgets/Public/Widgets/ColorGrading/SColorGradingComponentViewer.h
--- a/Engine/Source/Runtime/AdvancedWidgets/Public/Widgets/ColorGrading/SColorGradingComponentViewer.h
+++ b/Engine/Source/Runtime/AdvancedWidgets/Public/Widgets/ColorGrading/SColorGradingComponentViewer.h
@@ -31,12 +31,12 @@ public:
 	DECLARE_DELEGATE_FourParams(FOnDynamicSliderMinMaxValueChanged, float, TWeakPtr<SWidget>, bool, bool);
 
 	SLATE_BEGIN_ARGS(SColorGradingComponentViewer)
-		: _Value(0)
-		, _MinValue(0)
-		, _MaxValue(2)
+		: _Value(0.f)
+		, _MinValue(0.f)
+		, _MaxValue(2.f)
 		, _MinFractionalDigits(DefaultMinFractionalDigits)
 		, _MaxFractionalDigits(DefaultMaxFractionalDigits)
-		, _Delta(0)
+		, _Delta(0.f)
 		, _ShiftMultiplier(10.f)
 		, _CtrlMultiplier(0.1f)
 		, _SupportDynamicSliderMaxValue(false)
diff --git a/Engine/Source/Runtime/Slate/Public/Widgets/Input/NumericTypeInterface.h b/Engine/Source/Runtime/Slate/Public/Widgets/Input/NumericTypeInterface.h
--- a/Engine/Source/Runtime/Slate/Public/Widgets/Input/NumericTypeInterface.h
+++ b/Engine/Source/Runtime/Slate/Public/Widgets/Input/NumericTypeInterface.h
@@ -141,7 +141,7 @@ struct TDefaultNumericTypeInterface : INumericTypeInterface<NumericType>
 		// Attempt to parse it as an expression
 		static FBasicMathExpressionEvaluator Parser;
 
-		TValueOrError<double, FExpressionError> Result = Parser.Evaluate(*InString, InExistingValue);
+		TValueOrError<double, FExpressionError> Result = Parser.Evaluate(*InString, static_cast<double>(InExistingValue));
 		if (Result.IsValid())
 		{
 			return FMath::Clamp((NumericType)Result.GetValue(), TNumericLimits<NumericType>::Lowest(), TNumericLimits<NumericType>::Max());
