[Truncation Warnings] Fix Mac CIS in DetailsCustomization module

#rnx

diff --git a/Engine/Source/Editor/DetailCustomizations/Private/MacTargetSettingsDetails.cpp b/Engine/Source/Editor/DetailCustomizations/Private/MacTargetSettingsDetails.cpp
--- a/Engine/Source/Editor/DetailCustomizations/Private/MacTargetSettingsDetails.cpp
+++ b/Engine/Source/Editor/DetailCustomizations/Private/MacTargetSettingsDetails.cpp
@@ -347,7 +347,7 @@ bool FMacTargetSettingsDetails::HandlePostExternalIconCopy(const FString& InChos
 	return true;
 }
 
-static uint32 GMacTargetSettingsMinOSVers[][3] = {
+static uint8 GMacTargetSettingsMinOSVers[][3] = {
 	{12,0,0},
     {13,0,0}
 };
