Fixed type mismatch which led to ensure and value not set in the body instance object type.

#rnx

diff --git a/Engine/Source/Editor/DetailCustomizations/Private/BodyInstanceCustomization.cpp b/Engine/Source/Editor/DetailCustomizations/Private/BodyInstanceCustomization.cpp
--- a/Engine/Source/Editor/DetailCustomizations/Private/BodyInstanceCustomization.cpp
+++ b/Engine/Source/Editor/DetailCustomizations/Private/BodyInstanceCustomization.cpp
@@ -748,7 +748,7 @@ void FBodyInstanceCustomization::OnObjectTypeChanged( TSharedPtr<FString> NewSel
 				NewEnumVal = ObjectTypeValues[Iter.GetIndex()];
 			}
 		}
-		ensure(ObjectTypeHandle->SetValue(NewEnumVal) == FPropertyAccess::Result::Success);
+		ensure(ObjectTypeHandle->SetValue((uint8)NewEnumVal) == FPropertyAccess::Result::Success);
 	}
 }
 
