Fixed nullptr crash when changing to Custom filter type on abstract body instances like in PCG.

#rnx

diff --git a/Engine/Source/Editor/DetailCustomizations/Private/BodyInstanceCustomization.cpp b/Engine/Source/Editor/DetailCustomizations/Private/BodyInstanceCustomization.cpp
--- a/Engine/Source/Editor/DetailCustomizations/Private/BodyInstanceCustomization.cpp
+++ b/Engine/Source/Editor/DetailCustomizations/Private/BodyInstanceCustomization.cpp
@@ -879,7 +879,10 @@ void FBodyInstanceCustomization::OnCollisionProfileChanged(TSharedPtr<FString> N
 		if( NewSelection == CollisionProfileComboList[GetCustomIndex()])
 		{
 			// Force expansion when the user chooses the selected item
-			DetailBuilder->EditPropertyFromRoot(BodyInstanceHandle)->ShouldAutoExpand(true);
+			if (IDetailPropertyRow* BodyInstanceProperty = DetailBuilder->EditPropertyFromRoot(BodyInstanceHandle))
+			{
+				BodyInstanceProperty->ShouldAutoExpand(true);
+			}
 		}
 
 		// if none of them found, clear it
