Update the details customizations for transforms to support the EAxisList::LeftUpForward display coordinate system.

FComponentTransformDetails : When the display coordinate system is EAxisList::LeftUpForward, the value of CachedLocation.Y must be inverted for display or copying to the clipboard and when OnSetTransform is called the NewValue.Y should be inverted for writing back to the underlying data.

FMatrixStructCustomization, FTransformStructCustomization: When using the LeftUpForward display coordinate system, invert the value of CachedTranslationY when reading in the underlying data in CacheValues and before writing back the underlying data in FlushValues.

diff --git a/Engine/Source/Editor/DetailCustomizations/Private/ComponentTransformDetails.cpp b/Engine/Source/Editor/DetailCustomizations/Private/ComponentTransformDetails.cpp
--- a/Engine/Source/Editor/DetailCustomizations/Private/ComponentTransformDetails.cpp
+++ b/Engine/Source/Editor/DetailCustomizations/Private/ComponentTransformDetails.cpp
@@ -33,6 +33,7 @@
 #include "Math/UnitConversion.h"
 #include "Misc/AssertionMacros.h"
 #include "Misc/Attribute.h"
+#include "Misc/AxisDisplayInfo.h"
 #include "Misc/ConfigCacheIni.h"
 #include "Misc/NotifyHook.h"
 #include "PropertyEditorCopyPaste.h"
@@ -137,6 +138,7 @@ FComponentTransformDetails::FComponentTransformDetails( const TArray< TWeakObjec
 	, bIsSliderTransaction( false )
 	, HiddenFieldMask( 0 )
 	, bIsEnabledCache( false )
+	, bIsAxisDisplayLeftUpForward( AxisDisplayInfo::GetAxisDisplayCoordinateSystem() == EAxisList::LeftUpForward )
 {
 	GConfig->GetBool(TEXT("SelectionDetails"), TEXT("PreserveScaleRatio"), bPreserveScaleRatio, GEditorPerProjectIni);
 	FCoreUObjectDelegates::OnObjectsReplaced.AddRaw(this, &FComponentTransformDetails::OnObjectsReplaced);
@@ -301,7 +303,7 @@ void FComponentTransformDetails::OnCopy( ETransformField::Type TransformField )
 	switch (TransformField)
 	{
 	case ETransformField::Location:
-		CopyStr = FString::Printf(TEXT("(X=%f,Y=%f,Z=%f)"), CachedLocation.X.GetValue(), CachedLocation.Y.GetValue(), CachedLocation.Z.GetValue());
+		CopyStr = FString::Printf(TEXT("(X=%f,Y=%f,Z=%f)"), GetLocationX().GetValue(), GetLocationY().GetValue(), GetLocationZ().GetValue());
 		break;
 	case ETransformField::Rotation:
 		CopyStr = FString::Printf(TEXT("(Pitch=%f,Yaw=%f,Roll=%f)"), CachedRotation.Y.GetValue(), CachedRotation.Z.GetValue(), CachedRotation.X.GetValue());
@@ -480,6 +482,9 @@ void FComponentTransformDetails::GenerateChildContent( IDetailChildrenBuilder& C
 			.X(this, &FComponentTransformDetails::GetLocationX)
 			.Y(this, &FComponentTransformDetails::GetLocationY)
 			.Z(this, &FComponentTransformDetails::GetLocationZ)
+			.XDisplayName(AxisDisplayInfo::GetAxisDisplayName(EAxisList::Forward))
+			.YDisplayName(AxisDisplayInfo::GetAxisDisplayName(EAxisList::Left))
+			.ZDisplayName(AxisDisplayInfo::GetAxisDisplayName(EAxisList::Up))
 			.bColorAxisLabels(true)
 			.IsEnabled(this, &FComponentTransformDetails::GetIsLocationEnabled)
 			.OnXChanged(this, &FComponentTransformDetails::OnSetTransformAxis, ETextCommit::Default, ETransformField::Location, EAxisList::X, false)
@@ -586,6 +591,9 @@ void FComponentTransformDetails::GenerateChildContent( IDetailChildrenBuilder& C
 			.X( this, &FComponentTransformDetails::GetScaleX )
 			.Y( this, &FComponentTransformDetails::GetScaleY )
 			.Z( this, &FComponentTransformDetails::GetScaleZ )
+			.XDisplayName(AxisDisplayInfo::GetAxisDisplayName(EAxisList::Forward))
+			.YDisplayName(AxisDisplayInfo::GetAxisDisplayName(EAxisList::Left))
+			.ZDisplayName(AxisDisplayInfo::GetAxisDisplayName(EAxisList::Up))
 			.bColorAxisLabels( true )
 			.IsEnabled( this, &FComponentTransformDetails::GetIsScaleEnabled )
 			.OnXChanged( this, &FComponentTransformDetails::OnSetTransformAxis, ETextCommit::Default, ETransformField::Scale, EAxisList::X, false )
@@ -618,18 +626,21 @@ void FComponentTransformDetails::Tick( float DeltaTime )
 
 void FComponentTransformDetails::CacheCommonLocationUnits()
 {
+	const TOptional<FVector::FReal> LocationX = GetLocationX();
+	const TOptional<FVector::FReal> LocationY = GetLocationY();
+	const TOptional<FVector::FReal> LocationZ = GetLocationZ();
 	FVector::FReal LargestValue = 0.0;
-	if (CachedLocation.X.IsSet() && CachedLocation.X.GetValue() > LargestValue)
+	if (LocationX.IsSet() && LocationX.GetValue() > LargestValue)
 	{
-		LargestValue = CachedLocation.X.GetValue();
+		LargestValue = LocationX.GetValue();
 	}
-	if (CachedLocation.Y.IsSet() && CachedLocation.Y.GetValue() > LargestValue)
+	if (LocationY.IsSet() && LocationY.GetValue() > LargestValue)
 	{
-		LargestValue = CachedLocation.Y.GetValue();
+		LargestValue = LocationY.GetValue();
 	}
-	if (CachedLocation.Z.IsSet() && CachedLocation.Z.GetValue() > LargestValue)
+	if (LocationZ.IsSet() && LocationZ.GetValue() > LargestValue)
 	{
-		LargestValue = CachedLocation.Z.GetValue();
+		LargestValue = LocationZ.GetValue();
 	}
 
 	SetupFixedDisplay(LargestValue);
@@ -942,6 +953,15 @@ struct FGetRootComponentArchetype
 	}
 };
 
+TOptional<FVector::FReal> FComponentTransformDetails::GetLocationY() const
+{
+	if (bIsAxisDisplayLeftUpForward && CachedLocation.Y.IsSet())
+	{
+		return TOptional<FVector::FReal>(-CachedLocation.Y.GetValue());
+	}
+	return CachedLocation.Y;
+}
+
 bool FComponentTransformDetails::GetLocationResetVisibility() const
 {
 	const USceneComponent* Archetype = FGetRootComponentArchetype::Get(SelectedObjects[0].Get());
@@ -1259,6 +1279,11 @@ void FComponentTransformDetails::OnSetTransform(ETransformField::Type TransformF
 				{
 				case ETransformField::Location:
 					OldComponentValue = SceneComponent->GetRelativeLocation();
+
+					if (bIsAxisDisplayLeftUpForward)
+					{
+						NewValue.Y = -NewValue.Y;
+					}
 					break;
 				case ETransformField::Rotation:
 					// Pull from the actual component or from the cache
diff --git a/Engine/Source/Editor/DetailCustomizations/Private/ComponentTransformDetails.h b/Engine/Source/Editor/DetailCustomizations/Private/ComponentTransformDetails.h
--- a/Engine/Source/Editor/DetailCustomizations/Private/ComponentTransformDetails.h
+++ b/Engine/Source/Editor/DetailCustomizations/Private/ComponentTransformDetails.h
@@ -233,7 +233,7 @@ private:
 	/** @return The X component of location */
 	TOptional<FVector::FReal> GetLocationX() const { return CachedLocation.X; }
 	/** @return The Y component of location */
-	TOptional<FVector::FReal> GetLocationY() const { return CachedLocation.Y; }
+	TOptional<FVector::FReal> GetLocationY() const;
 	/** @return The Z component of location */
 	TOptional<FVector::FReal> GetLocationZ() const { return CachedLocation.Z; }
 	/** @return The visibility of the "Reset to Default" button for the location component */
@@ -339,4 +339,6 @@ private:
 	TArray< TWeakObjectPtr<UObject> > CachedHandlesObjects;
 	/** Cached enabled value of the selected set */
 	bool bIsEnabledCache;
+	/** Whether or not the axis display coordinate system is LeftUpForward  */
+	bool bIsAxisDisplayLeftUpForward;
 };
diff --git a/Engine/Source/Editor/DetailCustomizations/Private/MathStructProxyCustomizations.cpp b/Engine/Source/Editor/DetailCustomizations/Private/MathStructProxyCustomizations.cpp
--- a/Engine/Source/Editor/DetailCustomizations/Private/MathStructProxyCustomizations.cpp
+++ b/Engine/Source/Editor/DetailCustomizations/Private/MathStructProxyCustomizations.cpp
@@ -176,6 +176,10 @@ void FMatrixStructCustomization<T>::CustomizeLocation(TSharedRef<class IProperty
 		Row.OnPasteFromTextDelegate.Pin()->AddSP(this, &FMatrixStructCustomization<T>::OnPasteFromText, FTransformField::Location, WeakHandlePtr);
 	}
 	
+	const EAxisList::Type XAxis = bUseLeftUpForwardAxisDisplayCoordinateSystem ? EAxisList::Forward : EAxisList::X;
+	const EAxisList::Type YAxis = bUseLeftUpForwardAxisDisplayCoordinateSystem ? EAxisList::Left : EAxisList::Y;
+	const EAxisList::Type ZAxis = bUseLeftUpForwardAxisDisplayCoordinateSystem ? EAxisList::Up : EAxisList::Z;
+
 	Row
 	.CopyAction(FUIAction(FExecuteAction::CreateSP(this, &FMatrixStructCustomization<T>::OnCopy, FTransformField::Location, WeakHandlePtr)))
 	.PasteAction(FUIAction(FExecuteAction::CreateSP(this, &FMatrixStructCustomization<T>::OnPaste, FTransformField::Location, WeakHandlePtr)))
@@ -191,17 +195,17 @@ void FMatrixStructCustomization<T>::CustomizeLocation(TSharedRef<class IProperty
 		+ SHorizontalBox::Slot()
 		.Padding(FMargin(0.0f, 2.0f, 3.0f, 2.0f))
 		[
-			MakeNumericProxyWidget<UE::Math::TVector<T>, T>(StructPropertyHandle, CachedTranslationX, AxisDisplayInfo::GetAxisDisplayName(EAxisList::X), false, AxisDisplayInfo::GetAxisColor(EAxisList::X))
+			MakeNumericProxyWidget<UE::Math::TVector<T>, T>(StructPropertyHandle, CachedTranslationX, AxisDisplayInfo::GetAxisDisplayName(XAxis), false, AxisDisplayInfo::GetAxisColor(XAxis))
 		]
 		+ SHorizontalBox::Slot()
 		.Padding(FMargin(0.0f, 2.0f, 3.0f, 2.0f))
 		[
-			MakeNumericProxyWidget<UE::Math::TVector<T>, T>(StructPropertyHandle, CachedTranslationY, AxisDisplayInfo::GetAxisDisplayName(EAxisList::Y), false, AxisDisplayInfo::GetAxisColor(EAxisList::Y))
+			MakeNumericProxyWidget<UE::Math::TVector<T>, T>(StructPropertyHandle, CachedTranslationY, AxisDisplayInfo::GetAxisDisplayName(YAxis), false, AxisDisplayInfo::GetAxisColor(YAxis))
 		]
 		+ SHorizontalBox::Slot()
 		.Padding(FMargin(0.0f, 2.0f, 0.0f, 2.0f))
 		[
-			MakeNumericProxyWidget<UE::Math::TVector<T>, T>(StructPropertyHandle, CachedTranslationZ, AxisDisplayInfo::GetAxisDisplayName(EAxisList::Z), false, AxisDisplayInfo::GetAxisColor(EAxisList::Z))
+			MakeNumericProxyWidget<UE::Math::TVector<T>, T>(StructPropertyHandle, CachedTranslationZ, AxisDisplayInfo::GetAxisDisplayName(ZAxis), false, AxisDisplayInfo::GetAxisColor(ZAxis))
 		]
 	];
 }
@@ -216,6 +220,10 @@ void FMatrixStructCustomization<T>::CustomizeRotation(TSharedRef<class IProperty
 		Row.OnPasteFromTextDelegate.Pin()->AddSP(this, &FMatrixStructCustomization<T>::OnPasteFromText, FTransformField::Rotation, WeakHandlePtr);
 	}
 
+	const EAxisList::Type XAxis = bUseLeftUpForwardAxisDisplayCoordinateSystem ? EAxisList::Forward : EAxisList::X;
+	const EAxisList::Type YAxis = bUseLeftUpForwardAxisDisplayCoordinateSystem ? EAxisList::Left : EAxisList::Y;
+	const EAxisList::Type ZAxis = bUseLeftUpForwardAxisDisplayCoordinateSystem ? EAxisList::Up : EAxisList::Z;
+
 	Row
 	.CopyAction(FUIAction(FExecuteAction::CreateSP(this, &FMatrixStructCustomization<T>::OnCopy, FTransformField::Rotation, WeakHandlePtr)))
 	.PasteAction(FUIAction(FExecuteAction::CreateSP(this, &FMatrixStructCustomization<T>::OnPaste, FTransformField::Rotation, WeakHandlePtr)))
@@ -231,17 +239,17 @@ void FMatrixStructCustomization<T>::CustomizeRotation(TSharedRef<class IProperty
 		+ SHorizontalBox::Slot()
 		.Padding(FMargin(0.0f, 2.0f, 3.0f, 2.0f))
 		[
-			MakeNumericProxyWidget<UE::Math::TRotator<T>, T>(StructPropertyHandle, CachedRotationRoll, LOCTEXT("RotationRoll", "Roll"), true, AxisDisplayInfo::GetAxisColor(EAxisList::X))
+			MakeNumericProxyWidget<UE::Math::TRotator<T>, T>(StructPropertyHandle, CachedRotationRoll, LOCTEXT("RotationRoll", "Roll"), true, AxisDisplayInfo::GetAxisColor(XAxis))
 		]
 		+ SHorizontalBox::Slot()
 		.Padding(FMargin(0.0f, 2.0f, 3.0f, 2.0f))
 		[
-			MakeNumericProxyWidget<UE::Math::TRotator<T>, T>(StructPropertyHandle, CachedRotationPitch, LOCTEXT("RotationPitch", "Pitch"), true, AxisDisplayInfo::GetAxisColor(EAxisList::Y))
+			MakeNumericProxyWidget<UE::Math::TRotator<T>, T>(StructPropertyHandle, CachedRotationPitch, LOCTEXT("RotationPitch", "Pitch"), true, AxisDisplayInfo::GetAxisColor(YAxis))
 		]
 		+ SHorizontalBox::Slot()
 		.Padding(FMargin(0.0f, 2.0f, 0.0f, 2.0f))
 		[
-			MakeNumericProxyWidget<UE::Math::TRotator<T>, T>(StructPropertyHandle, CachedRotationYaw, LOCTEXT("RotationYaw", "Yaw"), true, AxisDisplayInfo::GetAxisColor(EAxisList::Z))
+			MakeNumericProxyWidget<UE::Math::TRotator<T>, T>(StructPropertyHandle, CachedRotationYaw, LOCTEXT("RotationYaw", "Yaw"), true, AxisDisplayInfo::GetAxisColor(ZAxis))
 		]
 	];
 }
@@ -256,6 +264,10 @@ void FMatrixStructCustomization<T>::CustomizeScale(TSharedRef<class IPropertyHan
 		Row.OnPasteFromTextDelegate.Pin()->AddSP(this, &FMatrixStructCustomization<T>::OnPasteFromText, FTransformField::Scale, WeakHandlePtr);
 	}
 
+	const EAxisList::Type XAxis = bUseLeftUpForwardAxisDisplayCoordinateSystem ? EAxisList::Forward : EAxisList::X;
+	const EAxisList::Type YAxis = bUseLeftUpForwardAxisDisplayCoordinateSystem ? EAxisList::Left : EAxisList::Y;
+	const EAxisList::Type ZAxis = bUseLeftUpForwardAxisDisplayCoordinateSystem ? EAxisList::Up : EAxisList::Z;
+
 	Row
 	.CopyAction(FUIAction(FExecuteAction::CreateSP(this, &FMatrixStructCustomization<T>::OnCopy, FTransformField::Scale, WeakHandlePtr)))
 	.PasteAction(FUIAction(FExecuteAction::CreateSP(this, &FMatrixStructCustomization<T>::OnPaste, FTransformField::Scale, WeakHandlePtr)))
@@ -271,17 +283,17 @@ void FMatrixStructCustomization<T>::CustomizeScale(TSharedRef<class IPropertyHan
 		+ SHorizontalBox::Slot()
 		.Padding(FMargin(0.0f, 2.0f, 3.0f, 2.0f))
 		[
-			MakeNumericProxyWidget<UE::Math::TVector<T>, T>(StructPropertyHandle, CachedScaleX, AxisDisplayInfo::GetAxisDisplayName(EAxisList::X), false, AxisDisplayInfo::GetAxisColor(EAxisList::X))
+			MakeNumericProxyWidget<UE::Math::TVector<T>, T>(StructPropertyHandle, CachedScaleX, AxisDisplayInfo::GetAxisDisplayName(XAxis), false, AxisDisplayInfo::GetAxisColor(XAxis))
 		]
 		+ SHorizontalBox::Slot()
 		.Padding(FMargin(0.0f, 2.0f, 3.0f, 2.0f))
 		[
-			MakeNumericProxyWidget<UE::Math::TVector<T>, T>(StructPropertyHandle, CachedScaleY, AxisDisplayInfo::GetAxisDisplayName(EAxisList::Y), false, AxisDisplayInfo::GetAxisColor(EAxisList::Y))
+			MakeNumericProxyWidget<UE::Math::TVector<T>, T>(StructPropertyHandle, CachedScaleY, AxisDisplayInfo::GetAxisDisplayName(YAxis), false, AxisDisplayInfo::GetAxisColor(YAxis))
 		]
 		+ SHorizontalBox::Slot()
 		.Padding(FMargin(0.0f, 2.0f, 0.0f, 2.0f))
 		[
-			MakeNumericProxyWidget<UE::Math::TVector<T>, T>(StructPropertyHandle, CachedScaleZ, AxisDisplayInfo::GetAxisDisplayName(EAxisList::Z), false, AxisDisplayInfo::GetAxisColor(EAxisList::Z))
+			MakeNumericProxyWidget<UE::Math::TVector<T>, T>(StructPropertyHandle, CachedScaleZ, AxisDisplayInfo::GetAxisDisplayName(ZAxis), false, AxisDisplayInfo::GetAxisColor(ZAxis))
 		]
 	];
 }
@@ -293,6 +305,10 @@ void FMatrixStructCustomization<T>::CustomizeChildren(TSharedRef<class IProperty
 
 	TWeakPtr<IPropertyHandle> WeakHandlePtr = StructPropertyHandle;
 
+	bUseLeftUpForwardAxisDisplayCoordinateSystem =
+		AxisDisplayInfo::GetAxisDisplayCoordinateSystem() == EAxisList::LeftUpForward
+		&& !StructPropertyHandle->GetProperty()->HasAnyPropertyFlags(CPF_BlueprintVisible | CPF_BlueprintReadOnly);
+
 	CustomizeLocation(StructPropertyHandle, StructBuilder.AddCustomRow(LOCTEXT("RotationLabel", "Rotation")));
 	CustomizeRotation(StructPropertyHandle, StructBuilder.AddCustomRow(LOCTEXT("LocationLabel", "Location")));
 	CustomizeScale(StructPropertyHandle, StructBuilder.AddCustomRow(LOCTEXT("ScaleLabel", "Scale")));
@@ -301,7 +317,7 @@ void FMatrixStructCustomization<T>::CustomizeChildren(TSharedRef<class IProperty
 template<typename T>
 void FMatrixStructCustomization<T>::OnCopy(FTransformField::Type Type, TWeakPtr<IPropertyHandle> PropertyHandlePtr)
 {
-	auto PropertyHandle = PropertyHandlePtr.Pin();
+	TSharedPtr<IPropertyHandle> PropertyHandle = PropertyHandlePtr.Pin();
 
 	if (!PropertyHandle.IsValid())
 	{
@@ -368,7 +384,7 @@ void FMatrixStructCustomization<T>::PasteFromText(
 	FTransformField::Type Type,
 	TWeakPtr<IPropertyHandle> PropertyHandlePtr)
 {
-	auto PropertyHandle = PropertyHandlePtr.Pin();
+	TSharedPtr<IPropertyHandle> PropertyHandle = PropertyHandlePtr.Pin();
 	if (!PropertyHandle.IsValid())
 	{
 		return;
@@ -428,7 +444,7 @@ void FMatrixStructCustomization<T>::PasteFromText(
 template<typename T>
 bool FMatrixStructCustomization<T>::CacheValues( TWeakPtr<IPropertyHandle> PropertyHandlePtr ) const
 {
-	auto PropertyHandle = PropertyHandlePtr.Pin();
+	TSharedPtr<IPropertyHandle> PropertyHandle = PropertyHandlePtr.Pin();
 
 	if (!PropertyHandle.IsValid())
 	{
@@ -463,6 +479,10 @@ bool FMatrixStructCustomization<T>::CacheValues( TWeakPtr<IPropertyHandle> Prope
 	if(FirstMatrixValue)
 	{
 		CachedTranslation->Set(FirstMatrixValue->GetOrigin());
+		if (bUseLeftUpForwardAxisDisplayCoordinateSystem)
+		{
+			CachedTranslationY->Set(-CachedTranslationY->Get());
+		}
 		CachedRotation->Set(FirstMatrixValue->Rotator());
 		CachedScale->Set(FirstMatrixValue->GetScaleVector());
 		return true;
@@ -474,7 +494,7 @@ bool FMatrixStructCustomization<T>::CacheValues( TWeakPtr<IPropertyHandle> Prope
 template<typename T>
 bool FMatrixStructCustomization<T>::FlushValues( TWeakPtr<IPropertyHandle> PropertyHandlePtr ) const
 {
-	auto PropertyHandle = PropertyHandlePtr.Pin();
+	TSharedPtr<IPropertyHandle> PropertyHandle = PropertyHandlePtr.Pin();
 	if (!PropertyHandle.IsValid())
 	{
 		return false;
@@ -510,7 +530,9 @@ bool FMatrixStructCustomization<T>::FlushValues( TWeakPtr<IPropertyHandle> Prope
 				);
 			UE::Math::TVector<T> Translation(
 				CachedTranslationX->IsSet() ? CachedTranslationX->Get() : CurrentTranslation.X,
-				CachedTranslationY->IsSet() ? CachedTranslationY->Get() : CurrentTranslation.Y,
+				CachedTranslationY->IsSet()
+					? (bUseLeftUpForwardAxisDisplayCoordinateSystem ? -CachedTranslationY->Get() : CachedTranslationY->Get())
+					: CurrentTranslation.Y,
 				CachedTranslationZ->IsSet() ? CachedTranslationZ->Get() : CurrentTranslation.Z
 				);
 			UE::Math::TVector<T> Scale(
@@ -588,7 +610,7 @@ TSharedRef<IPropertyTypeCustomization> FTransformStructCustomization<T>::MakeIns
 template<typename T>
 bool FTransformStructCustomization<T>::CacheValues( TWeakPtr<IPropertyHandle> PropertyHandlePtr ) const
 {
-	auto PropertyHandle = PropertyHandlePtr.Pin();
+	TSharedPtr<IPropertyHandle> PropertyHandle = PropertyHandlePtr.Pin();
 
 	if (!PropertyHandle.IsValid())
 	{
@@ -623,6 +645,10 @@ bool FTransformStructCustomization<T>::CacheValues( TWeakPtr<IPropertyHandle> Pr
 	if(FirstTransformValue)
 	{
 		this->CachedTranslation->Set(FirstTransformValue->GetTranslation());
+		if (this->bUseLeftUpForwardAxisDisplayCoordinateSystem)
+		{
+			this->CachedTranslationY->Set(-this->CachedTranslationY->Get());
+		}
 		this->CachedRotation->Set(FirstTransformValue->GetRotation().Rotator());
 		this->CachedScale->Set(FirstTransformValue->GetScale3D());
 		return true;
@@ -634,7 +660,7 @@ bool FTransformStructCustomization<T>::CacheValues( TWeakPtr<IPropertyHandle> Pr
 template<typename T>
 bool FTransformStructCustomization<T>::FlushValues( TWeakPtr<IPropertyHandle> PropertyHandlePtr ) const
 {
-	auto PropertyHandle = PropertyHandlePtr.Pin();
+	TSharedPtr<IPropertyHandle> PropertyHandle = PropertyHandlePtr.Pin();
 
 	if (!PropertyHandle.IsValid())
 	{
@@ -671,7 +697,9 @@ bool FTransformStructCustomization<T>::FlushValues( TWeakPtr<IPropertyHandle> Pr
 				);
 			UE::Math::TVector<T> Translation(
 				this->CachedTranslationX->IsSet() ? this->CachedTranslationX->Get() : CurrentTranslation.X,
-				this->CachedTranslationY->IsSet() ? this->CachedTranslationY->Get() : CurrentTranslation.Y,
+				this->CachedTranslationY->IsSet() 
+					? (this->bUseLeftUpForwardAxisDisplayCoordinateSystem ? -this->CachedTranslationY->Get() : this->CachedTranslationY->Get())
+					: CurrentTranslation.Y,
 				this->CachedTranslationZ->IsSet() ? this->CachedTranslationZ->Get() : CurrentTranslation.Z
 				);
 			UE::Math::TVector<T> Scale(
@@ -762,7 +790,7 @@ void FQuatStructCustomization<T>::CustomizeChildren(TSharedRef<class IPropertyHa
 template<typename T>
 bool FQuatStructCustomization<T>::CacheValues(TWeakPtr<IPropertyHandle> PropertyHandlePtr) const
 {
-	auto PropertyHandle = PropertyHandlePtr.Pin();
+	TSharedPtr<IPropertyHandle> PropertyHandle = PropertyHandlePtr.Pin();
 
 	if (!PropertyHandle.IsValid())
 	{
@@ -788,7 +816,7 @@ bool FQuatStructCustomization<T>::CacheValues(TWeakPtr<IPropertyHandle> Property
 template<typename T>
 bool FQuatStructCustomization<T>::FlushValues(TWeakPtr<IPropertyHandle> PropertyHandlePtr) const
 {
-	auto PropertyHandle = PropertyHandlePtr.Pin();
+	TSharedPtr<IPropertyHandle> PropertyHandle = PropertyHandlePtr.Pin();
 
 	if (!PropertyHandle.IsValid())
 	{
diff --git a/Engine/Source/Editor/DetailCustomizations/Public/Customizations/MathStructProxyCustomizations.h b/Engine/Source/Editor/DetailCustomizations/Public/Customizations/MathStructProxyCustomizations.h
--- a/Engine/Source/Editor/DetailCustomizations/Public/Customizations/MathStructProxyCustomizations.h
+++ b/Engine/Source/Editor/DetailCustomizations/Public/Customizations/MathStructProxyCustomizations.h
@@ -275,6 +275,7 @@ public:
 		, CachedScaleX(MakeShareable( new TProxyProperty<UE::Math::TVector<T>, T>(CachedScale, CachedScale->Get().X)))
 		, CachedScaleY(MakeShareable( new TProxyProperty<UE::Math::TVector<T>, T>(CachedScale, CachedScale->Get().Y)))
 		, CachedScaleZ(MakeShareable( new TProxyProperty<UE::Math::TVector<T>, T>(CachedScale, CachedScale->Get().Z)))
+		, bUseLeftUpForwardAxisDisplayCoordinateSystem( false )
 	{
 	}
 
@@ -328,6 +329,9 @@ protected:
 	mutable TSharedRef< TProxyProperty<UE::Math::TVector<T>, T> > CachedScaleX;
 	mutable TSharedRef< TProxyProperty<UE::Math::TVector<T>, T> > CachedScaleY;
 	mutable TSharedRef< TProxyProperty<UE::Math::TVector<T>, T> > CachedScaleZ;
+
+	/** Whether or not this customization is using the LeftUpForward axis display coordinate system */
+	bool bUseLeftUpForwardAxisDisplayCoordinateSystem;
 };
 
 /** 
