PR #7740: FCollisionProfileName property tooltip and editable fixes

diff --git a/Engine/Source/Editor/DetailCustomizations/Private/CollisionProfileNameCustomization.cpp b/Engine/Source/Editor/DetailCustomizations/Private/CollisionProfileNameCustomization.cpp
--- a/Engine/Source/Editor/DetailCustomizations/Private/CollisionProfileNameCustomization.cpp
+++ b/Engine/Source/Editor/DetailCustomizations/Private/CollisionProfileNameCustomization.cpp
@@ -58,9 +58,7 @@ void FCollisionProfileNameCustomization::CustomizeChildren(TSharedRef<class IPro
 	CollisionGroup.HeaderRow()
 	.NameContent()
 	[
-		SNew(STextBlock)
-		.Text(StructPropertyHandle->GetPropertyDisplayName())
-		.Font(IDetailLayoutBuilder::GetDetailFont())
+		StructPropertyHandle->CreatePropertyNameWidget()
 	]
 	.ValueContent()
 	[
@@ -76,7 +74,7 @@ void FCollisionProfileNameCustomization::CustomizeChildren(TSharedRef<class IPro
 			.OnSelectionChanged(this, &FCollisionProfileNameCustomization::OnSelectionChanged, &CollisionGroup)
 			.OnComboBoxOpening(this, &FCollisionProfileNameCustomization::OnComboBoxOpening)
 			.InitiallySelectedItem(InitialSelectedName)
-			.IsEnabled(FSlateApplication::Get().GetNormalExecutionAttribute())
+			.IsEnabled(StructPropertyHandle->IsEditable())
 			.ContentPadding(FMargin(2.0f, 2.0f))
 			.Content()
 			[
