Remove bUseAudioStreamCaching from configs

This setting doesn't exist anymore.

diff --git a/Engine/Source/Editor/DetailCustomizations/Private/WindowsTargetSettingsDetails.cpp b/Engine/Source/Editor/DetailCustomizations/Private/WindowsTargetSettingsDetails.cpp
--- a/Engine/Source/Editor/DetailCustomizations/Private/WindowsTargetSettingsDetails.cpp
+++ b/Engine/Source/Editor/DetailCustomizations/Private/WindowsTargetSettingsDetails.cpp
@@ -323,35 +323,6 @@ void FWindowsTargetSettingsDetails::CustomizeDetails( IDetailLayoutBuilder& Deta
 			.OnPostExternalImageCopy(FOnPostExternalImageCopy::CreateSP(this, &FWindowsTargetSettingsDetails::HandlePostExternalIconCopy))
 		]
 	];
-
-
-	AudioPluginWidgetManager.BuildAudioCategory(DetailBuilder, TEXT("Windows"));
-	IDetailCategoryBuilder& AudioCategory = DetailBuilder.EditCategory("Audio");
-
-	// Here we add a callback when the 
-// 	TSharedPtr<IPropertyHandle> AudioStreamCachingPropertyHandle = DetailBuilder.GetProperty("bUseAudioStreamCaching");
-// 	IDetailCategoryBuilder& AudioStreamCachingCategory = DetailBuilder.EditCategory("Audio");
-// 	IDetailPropertyRow& AudioStreamCachingPropertyRow = AudioCategory.AddProperty(AudioStreamCachingPropertyHandle);
-// 	AudioStreamCachingPropertyRow.CustomWidget()
-// 		.NameContent()
-// 		[
-// 			AudioStreamCachingPropertyHandle->CreatePropertyNameWidget()
-// 		]
-// 		.ValueContent()
-// 		.MaxDesiredWidth(500.0f)
-// 		.MinDesiredWidth(100.0f)
-// 		[
-// 			SNew(SHorizontalBox)
-// 
-// 			+ SHorizontalBox::Slot()
-// 			.FillWidth(1.0f)
-// 			[
-// 				SNew(SCheckBox)
-// 				.OnCheckStateChanged(this, &FWindowsTargetSettingsDetails::HandleAudioStreamCachingToggled, AudioStreamCachingPropertyHandle)
-// 				.IsChecked(this, &FWindowsTargetSettingsDetails::GetAudioStreamCachingToggled, AudioStreamCachingPropertyHandle)
-// 				.ToolTipText(AudioStreamCachingPropertyHandle->GetToolTipText())
-// 			]
-// 		];
 }
 
 bool FWindowsTargetSettingsDetails::HandlePreExternalIconCopy(const FString& InChosenImage)
