Fix AnimSequence segment length fix button in AnimMontage not working due to property access type mismatch

diff --git a/Engine/Source/Editor/DetailCustomizations/Private/AnimMontageSegmentDetails.cpp b/Engine/Source/Editor/DetailCustomizations/Private/AnimMontageSegmentDetails.cpp
--- a/Engine/Source/Editor/DetailCustomizations/Private/AnimMontageSegmentDetails.cpp
+++ b/Engine/Source/Editor/DetailCustomizations/Private/AnimMontageSegmentDetails.cpp
@@ -373,7 +373,7 @@ void FAnimMontageSegmentDetails::CustomizeDetails( IDetailLayoutBuilder& DetailB
 							{									
 								if(const IAnimationDataModel* DataModel = AnimSegment->GetAnimReference()->GetDataModel())
 								{
-									const double NewPlayLength = DataModel->GetPlayLength();
+									const float NewPlayLength = static_cast<float>(DataModel->GetPlayLength());
 									PlayLengthHandle->SetValue(NewPlayLength);
 								}
 							}
@@ -397,7 +397,7 @@ void FAnimMontageSegmentDetails::CustomizeDetails( IDetailLayoutBuilder& DetailB
 							{
 								if(const IAnimationDataModel* DataModel = AnimSegment->GetAnimReference()->GetDataModel())
 								{
-									const double NewPlayLength = DataModel->GetPlayLength();
+									const float NewPlayLength = static_cast<float>(DataModel->GetPlayLength());
 									AnimEndTimeProperty->SetValue(NewPlayLength);
 								}
 							}
