Restore removed code for registering Windows target platform audio customizations.

#rnx

diff --git a/Engine/Source/Editor/DetailCustomizations/Private/WindowsTargetSettingsDetails.cpp b/Engine/Source/Editor/DetailCustomizations/Private/WindowsTargetSettingsDetails.cpp
--- a/Engine/Source/Editor/DetailCustomizations/Private/WindowsTargetSettingsDetails.cpp
+++ b/Engine/Source/Editor/DetailCustomizations/Private/WindowsTargetSettingsDetails.cpp
@@ -323,6 +323,8 @@ void FWindowsTargetSettingsDetails::CustomizeDetails( IDetailLayoutBuilder& Deta
 			.OnPostExternalImageCopy(FOnPostExternalImageCopy::CreateSP(this, &FWindowsTargetSettingsDetails::HandlePostExternalIconCopy))
 		]
 	];
+
+	AudioPluginWidgetManager.BuildAudioCategory(DetailBuilder, TEXT("Windows"));
 }
 
 bool FWindowsTargetSettingsDetails::HandlePreExternalIconCopy(const FString& InChosenImage)
