- Fixed Animation Track Names actually showing array of names part of the AnimSequence
- Added details customization builder to display/show compressed animation data information (statistics, size, hash etc.) With the ability to select the platform you want to view (data will be compressed if necessary)

diff --git a/Engine/Source/Editor/DetailCustomizations/Private/AnimSequenceDetails.cpp b/Engine/Source/Editor/DetailCustomizations/Private/AnimSequenceDetails.cpp
--- a/Engine/Source/Editor/DetailCustomizations/Private/AnimSequenceDetails.cpp
+++ b/Engine/Source/Editor/DetailCustomizations/Private/AnimSequenceDetails.cpp
@@ -59,6 +59,8 @@
 #include "Widgets/Text/STextBlock.h"
 #include "PropertyCustomizationHelpers.h"
 #include "IDetailChildrenBuilder.h"
+#include "IDetailGroup.h"
+#include "CompressedAnimationDataNodeBuilder.h"
 
 class SWidget;
 struct FGeometry;
@@ -207,36 +209,88 @@ void FAnimSequenceDetails::CustomizeDetails(IDetailLayoutBuilder& DetailBuilder)
 	DetailBuilder.HideProperty(RetargetSourceAssetHandle);
 
 	// *** Animation Track Names ***
-	PRAGMA_DISABLE_DEPRECATION_WARNINGS
-	AnimationTrackNamesHandle = DetailBuilder.GetProperty(GET_MEMBER_NAME_CHECKED(UAnimSequence, AnimationTrackNames));
-	PRAGMA_ENABLE_DEPRECATION_WARNINGS
-
-	if (AnimationTrackNamesHandle.IsValid())
+	TArray<TWeakObjectPtr<UAnimSequence>> AnimSequences = DetailBuilder.GetSelectedObjectsOfType<UAnimSequence>();
+	AnimationTrackNamesList.Empty();
+	if (!AnimSequences.IsEmpty())
 	{
-		TArray<TWeakObjectPtr<UAnimSequence>> AnimSequences = DetailBuilder.GetSelectedObjectsOfType<UAnimSequence>();
-		AnimationTrackNamesList.Empty();
-
-		if (!AnimSequences.IsEmpty())
+		for (const TWeakObjectPtr<UAnimSequence>& AnimSequenceWeak : AnimSequences)
 		{
-			for (const TWeakObjectPtr<UAnimSequence>& AnimSequenceWeak : AnimSequences)
-			{
-				if(AnimSequenceWeak.IsValid())
-				{ 
-					if (AnimSequenceWeak->IsDataModelValid())
-					{
-						AnimSequenceWeak->GetDataModelInterface()->GetBoneTrackNames(AnimationTrackNamesList);
-					}
+			if(AnimSequenceWeak.IsValid())
+			{ 
+				if (AnimSequenceWeak->IsDataModelValid())
+				{
+					AnimSequenceWeak->GetDataModelInterface()->GetBoneTrackNames(AnimationTrackNamesList);
 				}
 			}
 		}
+	}
 
-		TSharedRef<FDetailArrayBuilder> AnimationTrackNamesArrayBuilder = MakeShareable(new FDetailArrayBuilder(AnimationTrackNamesHandle.ToSharedRef(), true, false, true));
-		AnimationTrackNamesArrayBuilder->OnGenerateArrayElementWidget(FOnGenerateArrayElementWidget::CreateSP(this, &FAnimSequenceDetails::GenerateAnimationTrackNameArrayElementWidget, &DetailBuilder));
-		AnimationCategory.AddCustomBuilder(AnimationTrackNamesArrayBuilder);
+	if (AnimationTrackNamesList.Num())
+	{
+		const FName AnimTrackNamesGroupName("AnimationTrackNames");
 
-		DetailBuilder.HideProperty(AnimationTrackNamesHandle);
+		const FText TrackNameDisplayText = LOCTEXT("AnimTrackNamesLabel", "Animation Track Names");
+		IDetailGroup& AnimTrackNamesGroup = AnimationCategory.AddGroup(AnimTrackNamesGroupName, TrackNameDisplayText);
+		AnimTrackNamesGroup.HeaderRow()
+		.NameContent()
+		[
+			SNew( SButton )
+			.ButtonStyle( FAppStyle::Get(), "NoBorder" )
+			.ContentPadding(FMargin(0,2,0,2))
+			.OnClicked_Lambda( [&AnimTrackNamesGroup]() { AnimTrackNamesGroup.ToggleExpansion(!AnimTrackNamesGroup.GetExpansionState()); return FReply::Handled(); } )				
+			.ForegroundColor( FSlateColor::UseForeground() )
+			.Content()
+			[
+				SNew(STextBlock)
+				.Font(IDetailLayoutBuilder::GetDetailFont())
+				.Text(TrackNameDisplayText)
+			]
+		]
+		.ValueWidget
+		[
+			SNew(STextBlock)
+			.Text_Lambda([NumTracks = AnimationTrackNamesList.Num()]()
+			{
+				return FText::Format(LOCTEXT("NumTracksFormat", "Number of Tracks: {0}"), FText::AsNumber(NumTracks));
+			})
+			.Font(DetailBuilder.GetDetailFont())
+		];
+		for (int32 Index = 0; Index < AnimationTrackNamesList.Num(); Index++)
+		{
+			FDetailWidgetRow& PropRow = AnimTrackNamesGroup.AddWidgetRow();
+			PropRow.OverrideResetToDefault(FResetToDefaultOverride::Hide());
+			PropRow.WholeRowWidget
+			[
+				SNew(SHorizontalBox)
+					+ SHorizontalBox::Slot()
+					.VAlign(VAlign_Center)
+					.HAlign(HAlign_Fill)
+					.Padding(5, 0, 0, 0)
+					.AutoWidth()
+					[
+						SNew(STextBlock)
+							.Text_Lambda([this, Index]()
+							{
+								if (AnimationTrackNamesList.IsValidIndex(Index))
+								{
+									return FText::FromName(AnimationTrackNamesList[Index]);
+								}
+						
+								return FText::GetEmpty();
+							})
+							.Font(DetailBuilder.GetDetailFont())
+					]
+			];
+		}
 	}
 
+	// *** Compressed Animation Data information ***
+	if (AnimSequences.Num() == 1)
+	{
+		IDetailCategoryBuilder& CompressionSettingsCategory = DetailBuilder.EditCategory("Compression");
+		UAnimSequence* AnimSequence = AnimSequences[0].Get();
+		CompressionSettingsCategory.AddCustomBuilder(MakeShared<FCompressedAnimationDataNodeBuilder>(AnimSequence));
+	}
 
 	/////////////////////////////////////////////////////////////////////////////
 	// Additive settings category
diff --git a/Engine/Source/Editor/DetailCustomizations/Private/CompressedAnimationDataNodeBuilder.cpp b/Engine/Source/Editor/DetailCustomizations/Private/CompressedAnimationDataNodeBuilder.cpp
new file mode 100644
--- /dev/null
+++ b/Engine/Source/Editor/DetailCustomizations/Private/CompressedAnimationDataNodeBuilder.cpp
@@ -0,0 +1,564 @@
+// Copyright Epic Games, Inc. All Rights Reserved.
+
+#include "CompressedAnimationDataNodeBuilder.h"
+
+#include "DetailLayoutBuilder.h"
+#include "DetailWidgetRow.h"
+#include "Editor.h"
+#include "IDetailChildrenBuilder.h"
+#include "IDetailGroup.h"
+#include "PlatformInfo.h"
+#include "Animation/AnimSequence.h"
+#include "Animation/Skeleton.h"
+#include "Interfaces/ITargetPlatform.h"
+#include "Interfaces/ITargetPlatformManagerModule.h"
+#include "Styling/AppStyle.h"
+#include "Widgets/SNullWidget.h"
+#include "Widgets/SBoxPanel.h"
+#include "Widgets/Images/SImage.h"
+#include "Widgets/Images/SThrobber.h"
+#include "Widgets/Input/SButton.h"
+#include "Widgets/Input/SComboBox.h"
+#include "Widgets/Text/STextBlock.h"
+
+#define LOCTEXT_NAMESPACE "CompressedAnimationDataNodeBuilder"
+
+FCompressedAnimationDataNodeBuilder::FCompressedAnimationDataNodeBuilder(UAnimSequence* InAnimSequence): WeakAnimSequence(InAnimSequence)
+{
+	CurrentTargetPlatform = GetTargetPlatformManagerRef().GetRunningTargetPlatform();
+	SelectedPlatformName = CurrentTargetPlatform->CookingDeviceProfileName();
+		
+	for (const FName& PlatformName : PlatformInfo::GetAllVanillaPlatformNames())
+	{
+		PlatformsList.Add(MakeShared<FString>(PlatformName.ToString()));
+	}
+
+	for (const PlatformInfo::FTargetPlatformInfo* PlatformInfo : PlatformInfo::GetPlatformInfoArray())
+	{
+		if(PlatformInfo->PlatformType == EBuildTargetType::Game || PlatformInfo->PlatformType == EBuildTargetType::Program)
+		{
+			continue;
+		}
+
+		PlatformsList.Add(MakeShared<FString>(PlatformInfo->Name.ToString()));
+	}
+
+	PlatformsList.Sort([](TSharedPtr<FString> A, TSharedPtr<FString> B) { return *A < *B; });
+	ensure(PlatformsList.ContainsByPredicate([this](TSharedPtr<FString> SharedPlatform) { return *SharedPlatform == CurrentTargetPlatform->IniPlatformName(); }));
+}
+
+void FCompressedAnimationDataNodeBuilder::Tick(float DeltaTime)
+{
+	if (const UAnimSequence* Sequence = WeakAnimSequence.Get())
+	{
+		const bool bHasCompressionData = Sequence->HasCompressedDataForPlatform(CurrentTargetPlatform);
+		if (bCachedHasCompressionData != bHasCompressionData )
+		{
+			bCachedHasCompressionData = bHasCompressionData;
+			OnRegenerateChildren.Execute();
+		}
+	}
+}
+
+const FSlateBrush* FCompressedAnimationDataNodeBuilder::GetSelectedPlatformBrush() const
+{
+	const PlatformInfo::FTargetPlatformInfo* PlatformInfo = PlatformInfo::FindPlatformInfo(FName(SelectedPlatformName));
+	if (PlatformInfo != nullptr)
+	{
+		return FAppStyle::GetBrush(PlatformInfo->GetIconStyleName(EPlatformIconSize::Normal));
+	}
+
+	return FStyleDefaults::GetNoBrush();
+}
+
+TSharedRef<SWidget> FCompressedAnimationDataNodeBuilder::OnGeneratePlatformListWidget(TSharedPtr<FString> Platform)
+{
+	const PlatformInfo::FTargetPlatformInfo* PlatformInfo = PlatformInfo::FindPlatformInfo(FName(*Platform));
+	if (PlatformInfo != nullptr)
+	{
+		const float Indent = PlatformInfo->IsVanilla() ? 0.f : 16.f;
+
+		return SNew(SHorizontalBox)
+			+SHorizontalBox::Slot()
+			 .AutoWidth()
+			 .VAlign(VAlign_Center)
+			 .Padding(4+Indent)
+			[
+				SNew(STextBlock)
+				.Text(PlatformInfo->DisplayName)
+				.Font(IDetailLayoutBuilder::GetDetailFont())
+			];
+	}
+	
+	return SNullWidget::NullWidget;
+}
+
+void FCompressedAnimationDataNodeBuilder::OnPlatformSelectionChanged(TSharedPtr<FString> Platform, ESelectInfo::Type InSelectInfo)
+{
+	if (Platform.IsValid())
+	{				
+		SelectedPlatformName = *Platform;
+	}
+
+	const FName PlatformName(*Platform);
+	const PlatformInfo::FTargetPlatformInfo* PlatformInfo = PlatformInfo::FindPlatformInfo(PlatformName);
+	CurrentTargetPlatform = GetTargetPlatformManagerRef().FindTargetPlatform(PlatformName);
+
+	check(CurrentTargetPlatform);
+	// Request compressed data
+	if (UAnimSequence* Sequence = WeakAnimSequence.Get())
+	{
+		Sequence->BeginCacheDerivedData(CurrentTargetPlatform);
+		bCachedHasCompressionData = false;
+		OnRegenerateChildren.Execute();
+	}
+}
+
+FText FCompressedAnimationDataNodeBuilder::GetSelectedPlatformName() const
+{
+	return FText::FromString(SelectedPlatformName);
+}
+
+void FCompressedAnimationDataNodeBuilder::GenerateHeaderRowContent(FDetailWidgetRow& NodeRow)
+{			
+	NodeRow.NameContent()
+	[
+		SNew(SHorizontalBox)
+		+SHorizontalBox::Slot()
+		[
+			SNew( SButton )
+			.ButtonStyle( FAppStyle::Get(), "NoBorder" )
+			.ContentPadding(FMargin(0,2,0,2))
+			.OnClicked_Lambda( [this]() { bExpanded = !bExpanded; OnToggleExpansion.ExecuteIfBound(bExpanded); return FReply::Handled(); } )
+			.ForegroundColor( FSlateColor::UseForeground() )
+			.Content()
+			[
+				SNew(STextBlock)
+				.Font(IDetailLayoutBuilder::GetDetailFont())
+				.Text_Lambda([this]()
+				{
+					return LOCTEXT("CompressedAnimationDataLabel", "Compressed Animation Data");
+				})
+			]
+		]
+	];
+		
+	NodeRow.ValueContent()
+	[
+		SNew(SHorizontalBox)
+		+SHorizontalBox::Slot()
+		[
+			SNew(SComboBox<TSharedPtr<FString>>)
+			.OptionsSource(&PlatformsList)
+			.InitiallySelectedItem(*PlatformsList.FindByPredicate([this](TSharedPtr<FString> SharedPlatform) { return *SharedPlatform == CurrentTargetPlatform->CookingDeviceProfileName();}))
+			.OnGenerateWidget_Static(&FCompressedAnimationDataNodeBuilder::OnGeneratePlatformListWidget)
+			.OnSelectionChanged_Raw(this, &FCompressedAnimationDataNodeBuilder::OnPlatformSelectionChanged)
+			[
+				SNew(SHorizontalBox)
+							
+				+SHorizontalBox::Slot()
+				 .AutoWidth()
+				 .VAlign(VAlign_Center)
+				 .Padding(4,0)
+				[
+					SNew(SImage)
+					.DesiredSizeOverride(FVector2D(16,16))
+					.Image_Raw(this, &FCompressedAnimationDataNodeBuilder::GetSelectedPlatformBrush)
+				]
+
+				+SHorizontalBox::Slot()
+				 .AutoWidth()
+				 .VAlign(VAlign_Center)
+				 .Padding(4,0)
+				[
+					SNew(STextBlock)
+					.Text_Raw(this, &FCompressedAnimationDataNodeBuilder::GetSelectedPlatformName)
+					.Font(IDetailLayoutBuilder::GetDetailFont())
+				]
+			]
+		]
+			
+		+SHorizontalBox::Slot()
+		 .VAlign(VAlign_Center)
+		 .Padding(2.0f, 0.0f)
+		 .AutoWidth()
+		[
+			SNew(SThrobber)
+			.Visibility_Raw(this, &FCompressedAnimationDataNodeBuilder::GetCompressionIndicatorVisibility)
+			.ToolTipText(LOCTEXT("CompressionTooltip", "Waiting for compressed data..."))
+		]
+	];
+}
+
+EVisibility FCompressedAnimationDataNodeBuilder::GetCompressionIndicatorVisibility() const
+{
+	if (const UAnimSequence* Sequence = WeakAnimSequence.Get())
+	{
+		return Sequence->HasCompressedDataForPlatform(CurrentTargetPlatform) ? EVisibility::Collapsed : EVisibility::Visible;
+	}
+	return EVisibility::Collapsed;
+}
+
+void FCompressedAnimationDataNodeBuilder::GenerateChildContent(IDetailChildrenBuilder& ChildrenBuilder)
+{
+	auto GetCompressedData = [this]() -> UAnimSequence::FScopedCompressedAnimSequence
+	{
+		const UAnimSequence* Sequence = WeakAnimSequence.Get();
+		check(Sequence);
+		return Sequence->GetCompressedData(CurrentTargetPlatform);
+	};
+
+	TAttribute<EVisibility> VisibilityAttribute = TAttribute<EVisibility>::Create([this]() -> EVisibility
+	{
+		const UAnimSequence* Sequence = WeakAnimSequence.Get();
+		check(Sequence);
+		return Sequence->GetCompressedData(CurrentTargetPlatform).Get().IsValid(Sequence) ? EVisibility::Visible : EVisibility::Collapsed;
+	});
+
+	FNumberFormattingOptions Options;
+	Options.MaximumFractionalDigits = 3;
+	const FText InvalidText = LOCTEXT("InvalidData", "Invalid");
+
+	{
+		ChildrenBuilder.AddCustomRow(LOCTEXT("CompressedDataHash", "Compressed Data Hash"))
+		.NameWidget
+		[
+			SNew(STextBlock)
+			.Font(IDetailLayoutBuilder::GetDetailFont())
+			.Text(LOCTEXT("CompressedDataHash", "Compressed Data Hash"))
+		]
+		.ValueWidget
+		[
+			SNew(STextBlock)
+			.Font(IDetailLayoutBuilder::GetDetailFont())
+			.Text_Lambda([this, GetCompressedData, InvalidText]() -> FText
+			{
+				if (const UAnimSequence* AnimSequence = WeakAnimSequence.Get())
+				{
+					const FIoHash Hash = AnimSequence->GetDerivedDataKeyHash(CurrentTargetPlatform);
+					const FString HashString = LexToString(Hash);
+					return FText::FromString(HashString);
+				}
+
+				return InvalidText;
+			})
+		]
+		.Visibility(VisibilityAttribute);
+	}
+
+	{
+		ChildrenBuilder.AddCustomRow(LOCTEXT("UncompressedSize", "Uncompressed (source) Data Size"))
+	    .NameWidget
+		[
+			SNew(STextBlock)
+			.Font(IDetailLayoutBuilder::GetDetailFont())
+			.Text(LOCTEXT("UncompressedSize", "Uncompressed (source) Data Size"))
+		]
+		.ValueWidget
+		[
+			SNew(STextBlock)
+			.Font(IDetailLayoutBuilder::GetDetailFont())
+			.Text_Lambda([this, GetCompressedData, InvalidText, Options]() -> FText
+			{
+				UAnimSequence::FScopedCompressedAnimSequence Data = GetCompressedData();
+				if (Data.Get().IsValid(WeakAnimSequence.Get()))
+				{								
+					return FText::AsMemory(Data.Get().CompressedRawDataSize, &Options);
+				}
+
+				return InvalidText;
+			})	
+		]
+		.Visibility(VisibilityAttribute);
+	}
+
+	{
+		ChildrenBuilder.AddCustomRow(LOCTEXT("CompressedDataSize", "Compressed Data Size"))
+	    .NameWidget
+		[
+			SNew(STextBlock)
+			.Font(IDetailLayoutBuilder::GetDetailFont())
+			.Text(LOCTEXT("CompressedSize", "Compressed Data Size"))
+		]
+		.ValueWidget
+		[
+			SNew(STextBlock)
+			.Font(IDetailLayoutBuilder::GetDetailFont())
+			.Text_Lambda([this, GetCompressedData, InvalidText, Options]() -> FText
+			{
+				UAnimSequence::FScopedCompressedAnimSequence Data = GetCompressedData();
+				if (Data.Get().IsValid(WeakAnimSequence.Get()))
+				{								
+					return FText::AsMemory((Data.Get().CompressedDataStructure ? Data.Get().CompressedDataStructure->GetApproxCompressedSize() : 0) + Data.Get().CompressedCurveByteStream.Num(), &Options);
+				}
+
+				return InvalidText;
+			})
+		]
+		.Visibility(VisibilityAttribute);
+	}
+
+	const FName BoneGroupName("BoneGroup");
+	IDetailGroup& BoneGroup = ChildrenBuilder.AddGroup(BoneGroupName, LOCTEXT("BoneGroupLabel", "Bone Data"), true);
+	{
+		{
+			BoneGroup.AddWidgetRow()
+			.NameWidget
+			[
+				SNew(STextBlock)
+				.Font(IDetailLayoutBuilder::GetDetailFont())
+				.Text(LOCTEXT("CompressedBoneSize", "Compressed Bone Data Size"))
+			]
+			.ValueWidget
+			[
+				SNew(STextBlock)
+				.Font(IDetailLayoutBuilder::GetDetailFont())
+				.Text_Lambda([this, GetCompressedData, InvalidText, Options]() -> FText
+				{
+					UAnimSequence::FScopedCompressedAnimSequence Data = GetCompressedData();
+					if (Data.Get().IsValid(WeakAnimSequence.Get()))
+					{								
+						return FText::AsMemory(Data.Get().CompressedDataStructure ? Data.Get().CompressedDataStructure->GetApproxCompressedSize() : 0, &Options);
+					}
+
+					return InvalidText;
+				})
+			]
+			.Visibility(VisibilityAttribute);
+		}
+		
+		{
+			UAnimSequence::FScopedCompressedAnimSequence Data = GetCompressedData();
+			const TArray<FTrackToSkeletonMap>& Tracks = Data.Get().CompressedTrackToSkeletonMapTable;
+			const int32 NumTracks = Tracks.Num();
+			const UAnimSequence* Sequence = WeakAnimSequence.Get();
+			const USkeleton* Skeleton = Sequence ? Sequence->GetSkeleton() : nullptr;
+			if (NumTracks && Skeleton)
+			{
+				const FName CompressedTrackNamesGroup("CompressedTrackNamesGroup");
+				IDetailGroup& CompressedTrackNameGroup = BoneGroup.AddGroup(CompressedTrackNamesGroup, LOCTEXT("CompressedTrackLabel", "Compressed Track Names"));
+				
+				CompressedTrackNameGroup.HeaderRow()
+				.NameContent()
+				[
+					SNew(STextBlock)
+					.Font(IDetailLayoutBuilder::GetDetailFont())
+					.Text(LOCTEXT("CompressedTrackLabel", "Compressed Track Names"))
+				]
+				.ValueWidget
+				[
+					SNew(STextBlock)
+					.Font(IDetailLayoutBuilder::GetDetailFont())
+					.Text(FText::Format(LOCTEXT("NumCompressedBonesFormat", "Number of Tracks: {0}"), FText::AsNumber(NumTracks)))
+				];
+				for (int32 Index = 0; Index < NumTracks; ++Index)
+				{
+					const int32 SkeletonIndex = Data.Get().GetSkeletonIndexFromTrackIndex(Index);
+					
+					CompressedTrackNameGroup.AddWidgetRow()
+					.WholeRowWidget
+					[
+						SNew(SHorizontalBox)
+						+ SHorizontalBox::Slot()
+						.VAlign(VAlign_Center)
+						.HAlign(HAlign_Fill)
+						.Padding(5, 0, 0, 0)
+						.AutoWidth()
+						[
+							SNew(STextBlock)
+							.Font(IDetailLayoutBuilder::GetDetailFont())
+							.Text(FText::FromName(Skeleton->GetReferenceSkeleton().GetBoneName(SkeletonIndex)))
+						]
+					]
+					.Visibility(VisibilityAttribute);
+				}
+			}
+		}	
+
+		{
+			const FName CompressedBoneErrorStats("CompressedBoneErrorStats");
+			IDetailGroup& CompressedBoneErrorStatGroup = BoneGroup.AddGroup(CompressedBoneErrorStats, LOCTEXT("BoneErrorStatGroupLabel", "Bone Compression Statistics"));
+
+			CompressedBoneErrorStatGroup.AddWidgetRow()
+		    .NameWidget
+			[
+				SNew(STextBlock)
+				.Font(IDetailLayoutBuilder::GetDetailFont())
+				.Text(LOCTEXT("AverageErrorLabel", "Average Error"))
+			]
+			.ValueWidget
+			[
+				SNew(STextBlock)
+				.Font(IDetailLayoutBuilder::GetDetailFont())
+				.Text_Lambda([this, GetCompressedData, InvalidText, Options]() -> FText
+				{
+					UAnimSequence::FScopedCompressedAnimSequence Data = GetCompressedData();
+					if (Data.Get().IsValid(WeakAnimSequence.Get()) && Data.Get().CompressedDataStructure.IsValid())
+					{								
+						return FText::AsNumber(Data.Get().CompressedDataStructure->BoneCompressionErrorStats.AverageError);
+					}
+
+					return InvalidText;
+				})
+			]
+			.Visibility(VisibilityAttribute);
+				
+			CompressedBoneErrorStatGroup.AddWidgetRow()
+		    .NameWidget
+			[
+				SNew(STextBlock)
+				.Font(IDetailLayoutBuilder::GetDetailFont())
+				.Text(LOCTEXT("MaximumErrorLabel", "Maximum Error"))
+			]
+			.ValueWidget
+			[
+				SNew(STextBlock)
+				.Font(IDetailLayoutBuilder::GetDetailFont())
+				.Text_Lambda([this, GetCompressedData, InvalidText, Options]() -> FText
+				{
+					UAnimSequence::FScopedCompressedAnimSequence Data = GetCompressedData();
+					if (Data.Get().IsValid(WeakAnimSequence.Get()) && Data.Get().CompressedDataStructure.IsValid())
+					{								
+						return FText::AsNumber(Data.Get().CompressedDataStructure->BoneCompressionErrorStats.MaxError, &Options);
+					}
+
+					return InvalidText;
+				})
+			]
+			.Visibility(VisibilityAttribute);
+
+			CompressedBoneErrorStatGroup.AddWidgetRow()
+		    .NameWidget
+			[
+				SNew(STextBlock)
+				.Font(IDetailLayoutBuilder::GetDetailFont())
+				.Text(LOCTEXT("MaxErrorTimeLabel", "Maximum Error Time-Interval"))
+			]
+			.ValueWidget
+			[
+				SNew(STextBlock)
+				.Font(IDetailLayoutBuilder::GetDetailFont())
+				.Text_Lambda([this, GetCompressedData, InvalidText, Options]() -> FText
+				{
+					UAnimSequence::FScopedCompressedAnimSequence Data = GetCompressedData();
+					if (Data.Get().IsValid(WeakAnimSequence.Get()) && Data.Get().CompressedDataStructure.IsValid())
+					{								
+						return FText::AsNumber(Data.Get().CompressedDataStructure->BoneCompressionErrorStats.MaxErrorTime, &Options);
+					}
+
+					return InvalidText;
+				})
+			]
+			.Visibility(VisibilityAttribute);
+				
+			CompressedBoneErrorStatGroup.AddWidgetRow()
+		    .NameWidget
+			[
+				SNew(STextBlock)
+				.Font(IDetailLayoutBuilder::GetDetailFont())
+				.Text(LOCTEXT("MaxErrorBoneLabel", "Maximum Error Bone Name"))
+			]
+			.ValueWidget
+			[
+				SNew(STextBlock)
+				.Font(IDetailLayoutBuilder::GetDetailFont())
+				.Text_Lambda([this, GetCompressedData, InvalidText, Options]() -> FText
+				{
+					UAnimSequence::FScopedCompressedAnimSequence Data = GetCompressedData();
+					if (Data.Get().IsValid(WeakAnimSequence.Get()) && Data.Get().CompressedDataStructure.IsValid())
+					{								
+						FName BoneName = NAME_None;							
+						const int32 ErrorBoneIndex = Data.Get().CompressedDataStructure->BoneCompressionErrorStats.MaxErrorBone;
+						if (const UAnimSequence* AnimSequence = WeakAnimSequence.Get())
+						{
+							if (const USkeleton* Skeleton = AnimSequence->GetSkeleton())
+							{
+								BoneName = Skeleton->GetReferenceSkeleton().GetBoneName(ErrorBoneIndex);
+							}
+						}
+						return FText::Format(LOCTEXT("MaxErrorBoneNameFormat", "{0}"), FText::FromName(BoneName));
+					}
+
+					return InvalidText;
+				})
+			]
+			.Visibility(VisibilityAttribute);
+		}
+	}
+
+	const FName CurveGroupName("CurveGroup");
+	IDetailGroup& CurveGroup = ChildrenBuilder.AddGroup(BoneGroupName, LOCTEXT("CurveGroupLabel", "Curve Data"), true);
+	{
+		{
+			CurveGroup.AddWidgetRow()
+			.NameWidget
+			[
+				SNew(STextBlock)
+				.Font(IDetailLayoutBuilder::GetDetailFont())
+				.Text(LOCTEXT("CompressedCurveSize", "Compressed Curve Data Size"))
+			]
+			.ValueWidget
+			[
+				SNew(STextBlock)
+				.Font(IDetailLayoutBuilder::GetDetailFont())
+				.Text_Lambda([this, GetCompressedData, InvalidText, Options]() -> FText
+				{
+					UAnimSequence::FScopedCompressedAnimSequence Data = GetCompressedData();
+					if (Data.Get().IsValid(WeakAnimSequence.Get()))
+					{								
+						return FText::AsMemory(Data.Get().CompressedCurveByteStream.Num(), &Options);
+					}
+
+					return InvalidText;
+				})
+			]
+			.Visibility(VisibilityAttribute);
+		}
+
+		{
+			UAnimSequence::FScopedCompressedAnimSequence Data = GetCompressedData();
+			const TArray<FAnimCompressedCurveIndexedName>& IndexedNames = Data.Get().IndexedCurveNames;
+			if (IndexedNames.Num())
+			{
+				const FName CompressedCurveNamesGroup("CompressedCurveNamesGroup");
+				IDetailGroup& CompressedCurveNameGroup = CurveGroup.AddGroup(CompressedCurveNamesGroup, LOCTEXT("CompressedCurvesLabel", "Compressed Curve Names"));
+				CompressedCurveNameGroup.HeaderRow()
+				.NameContent()
+				[
+					SNew(STextBlock)
+					.Font(IDetailLayoutBuilder::GetDetailFont())
+					.Text(LOCTEXT("CompressedCurvesLabel", "Compressed Curve Names"))
+				]
+				.ValueWidget
+				[
+					SNew(STextBlock)
+					.Font(IDetailLayoutBuilder::GetDetailFont())
+					.Text(FText::Format(LOCTEXT("NumCompressedCurvesFormat", "Number of Curves: {0}"), FText::AsNumber(IndexedNames.Num())))
+				];	
+				
+				const int32 NumCurves = IndexedNames.Num();
+				for (int32 Index = 0; Index < NumCurves; ++Index)
+				{	
+					CompressedCurveNameGroup.AddWidgetRow()
+					.WholeRowWidget
+					[
+						SNew(SHorizontalBox)
+						+ SHorizontalBox::Slot()
+						.VAlign(VAlign_Center)
+						.HAlign(HAlign_Fill)
+						.Padding(5, 0, 0, 0)
+						.AutoWidth()
+						[
+							SNew(STextBlock)
+							.Font(IDetailLayoutBuilder::GetDetailFont())
+							.Text(FText::FromName(IndexedNames[Index].CurveName))
+						]
+					]
+					.Visibility(VisibilityAttribute);
+				}
+			}
+		}		
+	}	
+}
+
+#undef LOCTEXT_NAMESPACE // "CompressedAnimationDataNodeBuilder"
\ No newline at end of file
diff --git a/Engine/Source/Editor/DetailCustomizations/Private/CompressedAnimationDataNodeBuilder.h b/Engine/Source/Editor/DetailCustomizations/Private/CompressedAnimationDataNodeBuilder.h
new file mode 100644
--- /dev/null
+++ b/Engine/Source/Editor/DetailCustomizations/Private/CompressedAnimationDataNodeBuilder.h
@@ -0,0 +1,41 @@
+// Copyright Epic Games, Inc. All Rights Reserved.
+
+#pragma once
+
+#include "IDetailCustomNodeBuilder.h"
+#include "Styling/StyleDefaults.h"
+#include "Widgets/SWidget.h"
+
+class UAnimSequence;
+
+class FCompressedAnimationDataNodeBuilder : public IDetailCustomNodeBuilder
+{
+public:	
+	FCompressedAnimationDataNodeBuilder(UAnimSequence* InAnimSequence);
+
+	virtual void SetOnRebuildChildren(FSimpleDelegate InOnRegenerateChildren) override { OnRegenerateChildren = InOnRegenerateChildren; }
+	virtual void SetOnToggleExpansion(FOnToggleNodeExpansion InOnToggleExpansion) override { OnToggleExpansion = InOnToggleExpansion; }
+	virtual void GenerateHeaderRowContent( FDetailWidgetRow& NodeRow ) override;
+	virtual void GenerateChildContent( IDetailChildrenBuilder& ChildrenBuilder ) override;
+	virtual FName GetName() const override { return FName("FCompressedAnimationDataNodeBuilder"); }
+
+	virtual void Tick(float DeltaTime) override;
+	virtual bool RequiresTick() const override { return true; }
+protected:
+	EVisibility GetCompressionIndicatorVisibility() const;  
+	const FSlateBrush* GetSelectedPlatformBrush() const;                                                        
+	static TSharedRef<SWidget> OnGeneratePlatformListWidget(TSharedPtr<FString> Platform);                     
+	void OnPlatformSelectionChanged( TSharedPtr<FString> Platform, ESelectInfo::Type InSelectInfo );            
+	FText GetSelectedPlatformName() const;      
+
+	TWeakObjectPtr<UAnimSequence> WeakAnimSequence;
+	FOnToggleNodeExpansion OnToggleExpansion;
+	FSimpleDelegate OnRegenerateChildren;
+	const ITargetPlatform* CurrentTargetPlatform = nullptr;
+	TArray<TSharedPtr<FString>> PlatformsList;
+	bool bExpanded = false;
+
+	bool bCachedHasCompressionData = false;
+
+	FString SelectedPlatformName;
+};
\ No newline at end of file
diff --git a/Engine/Source/Runtime/Engine/Classes/Animation/AnimSequence.h b/Engine/Source/Runtime/Engine/Classes/Animation/AnimSequence.h
--- a/Engine/Source/Runtime/Engine/Classes/Animation/AnimSequence.h
+++ b/Engine/Source/Runtime/Engine/Classes/Animation/AnimSequence.h
@@ -972,6 +972,7 @@ public:
 	friend class UE::Anim::FAnimSequenceCompilingManager;
 	friend UE::AnimNext::FDecompressionTools;
 	friend class FAnimSequenceDetails;
+	friend class FCompressedAnimationDataNodeBuilder;
 	friend class UAnimBank;
 	friend class FAnimBankBuildAsyncCacheTask;
 	friend struct FAnimCurveBufferAccess;
