Vector and Rotator InputBox Swizzle + Verse Transform rotation customization changes
* Adds a swizzle attribute to Rotator and Vector input box widget.  Uses Vec4 regardless of number of components for simplicity's sake.  Swizzle will remap the X, Y, Z, W to the respective swizzle indices.
* Adds a default swizzle applied to Rotator and Vector widgets
* Adds customization of display names for rotator and vector customizations to be overridden by AxisDisplayInfo if needed.
* Changes rotation order application for verse transforms to be XYZ
* Display and input of verse transforms to be in right-handed Y-up convention
* Remove display of denormalized rotations after conversion between conventions

diff --git a/Engine/Source/Editor/DetailCustomizations/Private/ComponentTransformDetails.cpp b/Engine/Source/Editor/DetailCustomizations/Private/ComponentTransformDetails.cpp
--- a/Engine/Source/Editor/DetailCustomizations/Private/ComponentTransformDetails.cpp
+++ b/Engine/Source/Editor/DetailCustomizations/Private/ComponentTransformDetails.cpp
@@ -83,7 +83,7 @@ namespace UE::DetailsCustomizations::Internal
 		{ ETransformField::Location, USceneComponent::GetRelativeLocationPropertyName().ToString() },
 		{ ETransformField::Rotation, USceneComponent::GetRelativeRotationPropertyName().ToString() },
 		{ ETransformField::Scale, USceneComponent::GetRelativeScale3DPropertyName().ToString() }		
-	};	
+	};
 }
 
 class FScopedSwitchWorldForObject
@@ -128,6 +128,20 @@ static USceneComponent* GetSceneComponentFromDetailsObject(UObject* InObject)
 	return Cast<USceneComponent>(InObject);
 }
 
+namespace ComponentTransformDetails::Private
+{
+	bool AreRotationsEqual(const FVector& Lhs, const FVector& Rhs)
+	{
+		constexpr double RotationEpsilon = 1.e-4;
+		
+		const double AbsDiffX = FMath::Abs(Lhs.X - Rhs.X);
+		const double AbsDiffY = FMath::Abs(Lhs.Y - Rhs.Y);
+		const double AbsDiffZ = FMath::Abs(Lhs.Z - Rhs.Z);
+							
+		return AbsDiffX < RotationEpsilon && AbsDiffY < RotationEpsilon && AbsDiffZ < RotationEpsilon;
+	}
+}
+
 FComponentTransformDetails::FComponentTransformDetails( const TArray< TWeakObjectPtr<UObject> >& InSelectedObjects, const FSelectedActorInfo& InSelectedActorInfo, IDetailLayoutBuilder& DetailBuilder )
 	: TNumericUnitTypeInterface(GetDefault<UEditorProjectAppearanceSettings>()->bDisplayUnitsOnComponentTransforms ? EUnit::Centimeters : EUnit::Unspecified)
 	, SelectedActorInfo( InSelectedActorInfo )
@@ -501,6 +515,11 @@ void FComponentTransformDetails::GenerateChildContent( IDetailChildrenBuilder& C
 		if( FUnitConversion::Settings().ShouldDisplayUnits() )
 		{
 			TypeInterface = MakeShareable( new TNumericUnitTypeInterface<FRotator::FReal>(EUnit::Degrees) );
+			if (bIsAxisDisplayLeftUpForward)
+			{
+				TypeInterface->SetMaxFractionalDigits(3);
+				TypeInterface->SetIndicateNearlyInteger(false);
+			}
 		}
 
 		ParentCategory.OnPasteFromText()->AddSP(this, &FComponentTransformDetails::OnPasteFromText, ETransformField::Rotation);
@@ -941,6 +960,11 @@ struct FGetRootComponentArchetype
 	}
 };
 
+TOptional<FVector::FReal> FComponentTransformDetails::GetLocationX() const
+{
+	return CachedLocation.X;
+}
+
 TOptional<FVector::FReal> FComponentTransformDetails::GetLocationY() const
 {
 	if (bIsAxisDisplayLeftUpForward && CachedLocation.Y.IsSet())
@@ -950,6 +974,41 @@ TOptional<FVector::FReal> FComponentTransformDetails::GetLocationY() const
 	return CachedLocation.Y;
 }
 
+TOptional<FVector::FReal> FComponentTransformDetails::GetLocationZ() const
+{
+	return CachedLocation.Z;
+}
+
+TOptional<FRotator::FReal> FComponentTransformDetails::GetRotationX() const
+{
+	return CachedRotation.X;
+}
+
+TOptional<FRotator::FReal> FComponentTransformDetails::GetRotationY() const
+{
+	return CachedRotation.Y;
+}
+
+TOptional<FRotator::FReal> FComponentTransformDetails::GetRotationZ() const
+{
+	return CachedRotation.Z;
+}
+
+TOptional<FVector::FReal> FComponentTransformDetails::GetScaleX() const
+{
+	return CachedScale.X;
+}
+
+TOptional<FVector::FReal> FComponentTransformDetails::GetScaleY() const
+{
+	return CachedScale.Y;
+}
+
+TOptional<FVector::FReal> FComponentTransformDetails::GetScaleZ() const
+{
+	return CachedScale.Z;
+}
+
 bool FComponentTransformDetails::GetLocationResetVisibility() const
 {
 	const USceneComponent* Archetype = FGetRootComponentArchetype::Get(SelectedObjects[0].Get());
@@ -976,10 +1035,24 @@ void FComponentTransformDetails::OnLocationResetClicked()
 bool FComponentTransformDetails::GetRotationResetVisibility() const
 {
 	const USceneComponent* Archetype = FGetRootComponentArchetype::Get(SelectedObjects[0].Get());
-	const FVector Data = Archetype ? Archetype->GetRelativeRotation().Euler() : FVector::ZeroVector;
 
 	// unset means multiple differing values, so show "Reset to Default" in that case
-	return CachedRotation.IsSet() && CachedRotation.X.GetValue() == Data.X && CachedRotation.Y.GetValue() == Data.Y && CachedRotation.Z.GetValue() == Data.Z ? false : true;
+	if (!CachedRotation.IsSet())
+	{
+		return true;
+	}
+	
+	if (!bIsAxisDisplayLeftUpForward)
+	{
+		const FVector Data = Archetype ? Archetype->GetRelativeRotation().Euler() : FVector::ZeroVector;
+		// unset means multiple differing values, so show "Reset to Default" in that case
+		return CachedRotation.X.GetValue() != Data.X || CachedRotation.Y.GetValue() != Data.Y || CachedRotation.Z.GetValue() != Data.Z;
+	}
+	else
+	{
+		const FVector Data = Archetype ? ConvertFromUnrealSpace_EulerDeg(Archetype->GetRelativeRotation()) : FVector::ZeroVector;
+		return ComponentTransformDetails::Private::AreRotationsEqual(Data, CachedRotation.ToVector());
+	}
 }
 
 void FComponentTransformDetails::OnRotationResetClicked()
@@ -990,7 +1063,7 @@ void FComponentTransformDetails::OnRotationResetClicked()
 		FScopedTransaction Transaction(TransactionName);
 
 		const USceneComponent* Archetype = FGetRootComponentArchetype::Get(SelectedObjects[0].Get());
-		const FVector Data = Archetype ? Archetype->GetRelativeRotation().Euler() : FVector::ZeroVector;
+		const FVector Data = Archetype ? ConvertFromUnrealSpace_EulerDeg(Archetype->GetRelativeRotation()) : FVector::ZeroVector;
 
 		OnSetTransform(ETransformField::Rotation, EAxisList::All, Data, false, true);
 	}
@@ -1103,6 +1176,13 @@ void FComponentTransformDetails::CacheDetails()
 				Loc = SceneComponent->GetRelativeLocation();
 				FRotator* FoundRotator = ObjectToRelativeRotationMap.Find(SceneComponent);
 				Rot = (bEditingRotationInUI && !Object->IsTemplate() && FoundRotator) ? *FoundRotator : SceneComponent->GetRelativeRotation();
+				if (bIsAxisDisplayLeftUpForward)
+				{
+					FVector Euler = ConvertFromUnrealSpace_EulerDeg(Rot);
+					Rot = FRotator(Euler.X, Euler.Y, Euler.Z);
+					
+					
+				}
 				Scale = SceneComponent->GetRelativeScale3D();
 
 				if( ObjectIndex == 0 )
@@ -1248,6 +1328,50 @@ void FComponentTransformDetails::OnSetTransform(ETransformField::Type TransformF
 	PropertyChain.AddHead(ValueProperty);
 	FPropertyChangedChainEvent PropertyChangedChainEvent(PropertyChain, PropertyChangedEvent);
 
+	EAxisList::Type RemappedAxis = Axis;
+	FVector SwizzledNewValue = NewValue;
+
+	if (bIsAxisDisplayLeftUpForward)
+	{
+		if (TransformField == ETransformField::Location)
+    	{
+    		SwizzledNewValue.Y = -NewValue.Y;
+    	}
+
+		if (TransformField == ETransformField::Rotation)
+		{
+			// Need to convert from Right-handed Y-Up to UE's Left-handed Z-Up however...
+			// NewValue is not the full set of Euler values to be applied, it will only contain
+			// the single value that was changed as specified by Axis
+			// Therefore it is not yet safe to convert it over.
+
+			// However we do need to swizzle the NewValue since the rotation widgets Axis values are set assuming
+			// normal Unreal rotations
+			RemappedAxis = [&]()
+			{
+				switch (Axis)
+				{
+				case EAxisList::X: return EAxisList::Z;
+				case EAxisList::Y: return EAxisList::X;
+				case EAxisList::Z: return EAxisList::Y;
+				case EAxisList::All: return EAxisList::All;
+				default: return EAxisList::X;
+				}
+			}();
+
+			SwizzledNewValue = FVector(NewValue.Y, NewValue.Z, NewValue.X);
+			
+			// Next step is to run SwizzledNewValue through GetAxisFilteredValue() to compose it
+			// with the converted rotator to right-hand coordinate space and get the full set of euler angles.
+			// Finally, these euler angles will be converted back to Unreal Rotator space and applied.
+			
+			// ObjectToRelativeRotationMap stores the rotations in Unreal Rotator space always - this may need to change though
+			// CachedRotation is stored in Right handed Y-up space as this is used to read back the values into the widget for display purposes
+			// See GetRotationY, GetRotationZ
+		}
+	}
+	
+
 	for (int32 ObjectIndex = 0; ObjectIndex < SelectedObjects.Num(); ++ObjectIndex)
 	{
 		TWeakObjectPtr<UObject> ObjectPtr = SelectedObjects[ObjectIndex];
@@ -1260,6 +1384,7 @@ void FComponentTransformDetails::OnSetTransform(ETransformField::Type TransformF
 				AActor* EditedActor = SceneComponent->GetOwner();
 				const bool bIsEditingTemplateObject = Object->IsTemplate();
 
+				FRotator OldComponentRotator;
 				FVector OldComponentValue;
 				FVector NewComponentValue;
 
@@ -1267,19 +1392,18 @@ void FComponentTransformDetails::OnSetTransform(ETransformField::Type TransformF
 				{
 				case ETransformField::Location:
 					OldComponentValue = SceneComponent->GetRelativeLocation();
-
-					if (bIsAxisDisplayLeftUpForward)
-					{
-						NewValue.Y = -NewValue.Y;
-					}
 					break;
 				case ETransformField::Rotation:
 					// Pull from the actual component or from the cache
-					OldComponentValue = SceneComponent->GetRelativeRotation().Euler();
 					if (bEditingRotationInUI && !bIsEditingTemplateObject && ObjectToRelativeRotationMap.Find(SceneComponent))
 					{
-						OldComponentValue = ObjectToRelativeRotationMap.Find(SceneComponent)->Euler();
+						OldComponentRotator = *ObjectToRelativeRotationMap.Find(SceneComponent);
 					}
+					else
+					{
+						OldComponentRotator = SceneComponent->GetRelativeRotation();
+					}
+					OldComponentValue = ConvertFromUnrealSpace_EulerDeg(OldComponentRotator);
 					break;
 				case ETransformField::Scale:
 					OldComponentValue = SceneComponent->GetRelativeScale3D();
@@ -1289,17 +1413,37 @@ void FComponentTransformDetails::OnSetTransform(ETransformField::Type TransformF
 				// Set the incoming value
 				if (bMirror)
 				{
-					NewComponentValue = GetAxisFilteredVector(Axis, -OldComponentValue, OldComponentValue);
+					NewComponentValue = GetAxisFilteredVector(RemappedAxis, -OldComponentValue, OldComponentValue);
 				}
 				else
 				{
-					NewComponentValue = GetAxisFilteredVector(Axis, NewValue, OldComponentValue);
+					NewComponentValue = GetAxisFilteredVector(RemappedAxis, SwizzledNewValue, OldComponentValue);
 				}
+				
+				auto AreValuesEqual = [this, TransformField](const FVector& NewComponentValue_, const FVector& OldComponentValue_)
+				{
+					if (!bIsAxisDisplayLeftUpForward || TransformField != ETransformField::Rotation)
+					{
+						// Bit-wise identical check
+						return NewComponentValue_ == OldComponentValue_;
+					}
+					else
+					{
+						// LeftUpForward uses alternative XYZ intrinsic rotation but rotation is stored
+						// still as FRotator in Yaw-Pitch-Roll intrinsic convention.  The conversion between
+						// these two conventions prevents bit-exact comparisons.  If values set are close enough
+						// to what exists on the component, then skip the setting rotation
+						// This prevents accidental small errors accumulating due to automatic conversion from the cached
+						// euler rotation representations and the underlying data
+						return ComponentTransformDetails::Private::AreRotationsEqual(NewComponentValue_, OldComponentValue_);
+					}
+				};
 
 				// If we're committing during a slider transaction then we need to force it, in order that PostEditChangeChainProperty be called.
 				// Note: this will even happen if the slider hasn't changed the value.
-				if (OldComponentValue != NewComponentValue || (bCommitted && bIsSliderTransaction))
+				if (!AreValuesEqual(NewComponentValue, OldComponentValue) || (bCommitted && bIsSliderTransaction))
 				{
+
 					if (!bBeganTransaction && bCommitted)
 					{
 						// NOTE: One transaction per change, not per actor
@@ -1366,7 +1510,7 @@ void FComponentTransformDetails::OnSetTransform(ETransformField::Type TransformF
 						}
 					case ETransformField::Rotation:
 						{
-							FRotator NewRotation = FRotator::MakeFromEuler(NewComponentValue);
+							FRotator NewRotation = ConvertToUnrealSpace_EulerDeg(NewComponentValue);
 
 							if (!bIsEditingTemplateObject)
 							{
@@ -1380,7 +1524,7 @@ void FComponentTransformDetails::OnSetTransform(ETransformField::Type TransformF
 							if (bIsEditingTemplateObject)
 							{
 								TSet<USceneComponent*> UpdatedInstances;
-								FComponentEditorUtils::PropagateDefaultValueChange(SceneComponent, ValueProperty, FRotator::MakeFromEuler(OldComponentValue), NewRotation, UpdatedInstances);
+								FComponentEditorUtils::PropagateDefaultValueChange(SceneComponent, ValueProperty, ConvertToUnrealSpace_EulerDeg(OldComponentValue), NewRotation, UpdatedInstances);
 							}
 
 							break;
@@ -1693,4 +1837,87 @@ void FComponentTransformDetails::OnObjectsReplaced(const TMap<UObject*, UObject*
 	}
 }
 
+namespace ComponentTransformDetailsPrivate
+{
+	TTuple<FQuat::FReal, FQuat::FReal, FQuat::FReal> RadiansToDegrees(const TTuple<FQuat::FReal, FQuat::FReal, FQuat::FReal>& Rads)
+	{
+		constexpr FQuat::FReal RadsToDegrees = (180.f / UE_PI);
+		TTuple<FQuat::FReal, FQuat::FReal, FQuat::FReal> Output
+		{
+			Rads.Get<0>() * RadsToDegrees,
+			Rads.Get<1>() * RadsToDegrees,
+			Rads.Get<2>() * RadsToDegrees
+		};
+		return Output;
+	}
+
+	TTuple<FQuat::FReal, FQuat::FReal, FQuat::FReal> DegreesToRadians(const TTuple<FQuat::FReal, FQuat::FReal, FQuat::FReal>& Rads)
+	{
+		constexpr FQuat::FReal DegreesToRadians = (UE_PI / 180.f);
+		TTuple<FQuat::FReal, FQuat::FReal, FQuat::FReal> Output
+		{
+			Rads.Get<0>() * DegreesToRadians,
+			Rads.Get<1>() * DegreesToRadians,
+			Rads.Get<2>() * DegreesToRadians
+		};
+		return Output;
+	}
+}
+
+FVector FComponentTransformDetails::ConvertFromUnrealSpace_EulerDeg(const FRotator& Rotator) const
+{
+	using namespace ComponentTransformDetailsPrivate;
+	if (!bIsAxisDisplayLeftUpForward)
+	{
+		return Rotator.Euler();
+	}
+
+	FQuat Q = Rotator.Quaternion().GetNormalized();
+	TTuple<FQuat::FReal, FQuat::FReal, FQuat::FReal> VerseEulerRads = Q.ToLUFEuler();
+
+	// Since the value is converted from quaternion, will likely have denormals.  Clamp those values.
+	auto SanitizeFloat = [](FQuat::FReal Val)->FQuat::FReal
+	{
+		if (FMath::IsNearlyZero(Val))
+		{
+			return 0.0;
+		}
+		return Val;
+	};
+	
+	FVector VerseEulerRadsV
+	{
+		SanitizeFloat(VerseEulerRads.Get<0>()),
+		SanitizeFloat(VerseEulerRads.Get<1>()),
+		SanitizeFloat(VerseEulerRads.Get<2>())
+	};
+	FVector VerseEulerDegrees = FMath::RadiansToDegrees(VerseEulerRadsV);
+	return VerseEulerDegrees;
+}
+
+FRotator FComponentTransformDetails::ConvertToUnrealSpace_EulerDeg(const FVector& Rotation) const
+{
+	using namespace ComponentTransformDetailsPrivate;
+	if (!bIsAxisDisplayLeftUpForward)
+	{
+		return FRotator::MakeFromEuler(Rotation);
+	}
+
+	const FVector RotationRads = FMath::DegreesToRadians(Rotation);
+
+	TTuple<FQuat::FReal, FQuat::FReal, FQuat::FReal> RotationRadsT =
+	{
+		RotationRads.X,
+		RotationRads.Y,
+		RotationRads.Z
+	};
+	
+	
+	FQuat Quat = FQuat::MakeFromLUFEuler(RotationRadsT);
+	Quat.Normalize();
+
+	FRotator Result(Quat);
+	return Result;
+}
+
 #undef LOCTEXT_NAMESPACE
diff --git a/Engine/Source/Editor/DetailCustomizations/Private/ComponentTransformDetails.h b/Engine/Source/Editor/DetailCustomizations/Private/ComponentTransformDetails.h
--- a/Engine/Source/Editor/DetailCustomizations/Private/ComponentTransformDetails.h
+++ b/Engine/Source/Editor/DetailCustomizations/Private/ComponentTransformDetails.h
@@ -231,29 +231,29 @@ private:
 	void OnZScaleMirrored();
 
 	/** @return The X component of location */
-	TOptional<FVector::FReal> GetLocationX() const { return CachedLocation.X; }
+	TOptional<FVector::FReal> GetLocationX() const;
 	/** @return The Y component of location */
 	TOptional<FVector::FReal> GetLocationY() const;
 	/** @return The Z component of location */
-	TOptional<FVector::FReal> GetLocationZ() const { return CachedLocation.Z; }
+	TOptional<FVector::FReal> GetLocationZ() const;
 	/** @return The visibility of the "Reset to Default" button for the location component */
 	bool GetLocationResetVisibility() const;
 
 	/** @return The X component of rotation */
-	TOptional<FRotator::FReal> GetRotationX() const { return CachedRotation.X; }
+	TOptional<FRotator::FReal> GetRotationX() const;
 	/** @return The Y component of rotation */
-	TOptional<FRotator::FReal> GetRotationY() const { return CachedRotation.Y; }
+	TOptional<FRotator::FReal> GetRotationY() const;
 	/** @return The Z component of rotation */
-	TOptional<FRotator::FReal> GetRotationZ() const { return CachedRotation.Z; }
+	TOptional<FRotator::FReal> GetRotationZ() const;
 	/** @return The visibility of the "Reset to Default" button for the rotation component */
 	bool GetRotationResetVisibility() const;
 
 	/** @return The X component of scale */
-	TOptional<FVector::FReal> GetScaleX() const { return CachedScale.X; }
+	TOptional<FVector::FReal> GetScaleX() const;
 	/** @return The Y component of scale */
-	TOptional<FVector::FReal> GetScaleY() const { return CachedScale.Y; }
+	TOptional<FVector::FReal> GetScaleY() const;
 	/** @return The Z component of scale */
-	TOptional<FVector::FReal> GetScaleZ() const { return CachedScale.Z; }
+	TOptional<FVector::FReal> GetScaleZ() const;
 	/** @return The visibility of the "Reset to Default" button for the scale component */
 	bool GetScaleResetVisibility() const;
 
@@ -267,6 +267,15 @@ private:
 
 	/** Delegate called when the OnObjectsReplaced is broadcast by the engine. */
 	void OnObjectsReplaced(const TMap<UObject*, UObject*>& ReplacementMap);
+
+	// Gets Euler angles from Unreal's Rotator space and converts to the display/edit space
+	// If display space is the Rotator space, then just Euler()
+	// Otherwise performs handedness and rotation order calculations
+	FVector ConvertFromUnrealSpace_EulerDeg(const FRotator& Rotator) const;
+	// Converts back to Unreal's Rotator if display space is different than Rotator's space
+	// Assumes that the Rotation provided is in degrees
+	FRotator ConvertToUnrealSpace_EulerDeg(const FVector& Rotation) const;
+	
 private:
 	/** A vector where it may optionally be unset */
 	template <typename NumericType>
@@ -299,6 +308,12 @@ private:
 			return X.IsSet() && Y.IsSet() && Z.IsSet();
 		}
 
+		FVector ToVector() const
+		{
+			check(IsSet());
+			return FVector(X.GetValue(), Y.GetValue(), Z.GetValue());
+		}
+
 		TOptional<NumericType> X;
 		TOptional<NumericType> Y;
 		TOptional<NumericType> Z;
diff --git a/Engine/Source/Editor/DetailCustomizations/Private/MathStructProxyCustomizations.cpp b/Engine/Source/Editor/DetailCustomizations/Private/MathStructProxyCustomizations.cpp
--- a/Engine/Source/Editor/DetailCustomizations/Private/MathStructProxyCustomizations.cpp
+++ b/Engine/Source/Editor/DetailCustomizations/Private/MathStructProxyCustomizations.cpp
@@ -51,7 +51,12 @@ void FMathStructProxyCustomization::MakeHeaderRow( TSharedRef<class IPropertyHan
 }
 
 template<typename ProxyType, typename NumericType>
-TSharedRef<SWidget> FMathStructProxyCustomization::MakeNumericProxyWidget(TSharedRef<IPropertyHandle>& StructPropertyHandle, TSharedRef< TProxyProperty<ProxyType, NumericType> >& ProxyValue, const FText& Label, bool bRotationInDegrees, const FLinearColor& LabelBackgroundColor)
+TSharedRef<SWidget> FMathStructProxyCustomization::MakeNumericProxyWidget(
+	TSharedRef<IPropertyHandle>& StructPropertyHandle,
+	TSharedRef< TProxyProperty<ProxyType, NumericType> >& ProxyValue,
+	const FText& Label,
+	bool bRotationInDegrees,
+	const FLinearColor& LabelBackgroundColor)
 {
 	TWeakPtr<IPropertyHandle> WeakHandlePtr = StructPropertyHandle;
 
@@ -180,6 +185,46 @@ void FMatrixStructCustomization<T>::CustomizeLocation(TSharedRef<class IProperty
 	const EAxisList::Type YAxis = bUseLeftUpForwardAxisDisplayCoordinateSystem ? EAxisList::Left : EAxisList::Y;
 	const EAxisList::Type ZAxis = bUseLeftUpForwardAxisDisplayCoordinateSystem ? EAxisList::Up : EAxisList::Z;
 
+	constexpr int32 NumComponents = 3;
+	TStaticArray<TFunction<TSharedRef<SWidget>()>, NumComponents> ComponentConstructors = {
+		[this, &StructPropertyHandle, XAxis]()->TSharedRef<SWidget>
+		{
+			return MakeNumericProxyWidget<UE::Math::TVector<T>, T>(StructPropertyHandle, CachedTranslationX, AxisDisplayInfo::GetAxisDisplayName(XAxis), false, AxisDisplayInfo::GetAxisColor(XAxis));
+		},
+		[this, &StructPropertyHandle, YAxis]()->TSharedRef<SWidget>
+		{
+			return MakeNumericProxyWidget<UE::Math::TVector<T>, T>(StructPropertyHandle, CachedTranslationY, AxisDisplayInfo::GetAxisDisplayName(YAxis), false, AxisDisplayInfo::GetAxisColor(YAxis));
+		},
+		[this, &StructPropertyHandle, ZAxis]()->TSharedRef<SWidget>
+		{
+			return MakeNumericProxyWidget<UE::Math::TVector<T>, T>(StructPropertyHandle, CachedTranslationZ, AxisDisplayInfo::GetAxisDisplayName(ZAxis), false, AxisDisplayInfo::GetAxisColor(ZAxis));
+		}
+	};
+	const TStaticArray<FMargin, NumComponents> Paddings =
+	{
+		FMargin(0.0f, 2.0f, 3.0f, 2.0f),
+		FMargin(0.0f, 2.0f, 3.0f, 2.0f),
+		FMargin(0.0f, 2.0f, 0.0f, 2.0f)
+	};
+
+	TSharedRef<SHorizontalBox> HBox = MakeShared<SHorizontalBox>();
+	FIntVector4 Swizzle = GetSwizzle();
+	for (int32 ComponentIndex = 0; ComponentIndex < NumComponents; ComponentIndex++)
+	{
+		TSharedRef<SWidget> Widget = SNullWidget::NullWidget;
+		if (ensure(Swizzle[ComponentIndex] < NumComponents))
+		{
+			const int32 SwizzledComponentIndex = Swizzle[ComponentIndex];
+			Widget = ComponentConstructors[SwizzledComponentIndex]();
+		}
+
+		HBox->AddSlot()
+			.Padding(Paddings[ComponentIndex])
+			[
+				Widget
+			];
+	}
+
 	Row
 	.CopyAction(FUIAction(FExecuteAction::CreateSP(this, &FMatrixStructCustomization<T>::OnCopy, FTransformField::Location, WeakHandlePtr)))
 	.PasteAction(FUIAction(FExecuteAction::CreateSP(this, &FMatrixStructCustomization<T>::OnPaste, FTransformField::Location, WeakHandlePtr)))
@@ -191,22 +236,7 @@ void FMatrixStructCustomization<T>::CustomizeLocation(TSharedRef<class IProperty
 	.MinDesiredWidth(375.0f)
 	.MaxDesiredWidth(375.0f)
 	[
-		SNew(SHorizontalBox)
-		+ SHorizontalBox::Slot()
-		.Padding(FMargin(0.0f, 2.0f, 3.0f, 2.0f))
-		[
-			MakeNumericProxyWidget<UE::Math::TVector<T>, T>(StructPropertyHandle, CachedTranslationX, AxisDisplayInfo::GetAxisDisplayName(XAxis), false, AxisDisplayInfo::GetAxisColor(XAxis))
-		]
-		+ SHorizontalBox::Slot()
-		.Padding(FMargin(0.0f, 2.0f, 3.0f, 2.0f))
-		[
-			MakeNumericProxyWidget<UE::Math::TVector<T>, T>(StructPropertyHandle, CachedTranslationY, AxisDisplayInfo::GetAxisDisplayName(YAxis), false, AxisDisplayInfo::GetAxisColor(YAxis))
-		]
-		+ SHorizontalBox::Slot()
-		.Padding(FMargin(0.0f, 2.0f, 0.0f, 2.0f))
-		[
-			MakeNumericProxyWidget<UE::Math::TVector<T>, T>(StructPropertyHandle, CachedTranslationZ, AxisDisplayInfo::GetAxisDisplayName(ZAxis), false, AxisDisplayInfo::GetAxisColor(ZAxis))
-		]
+		HBox
 	];
 }
 
@@ -223,6 +253,46 @@ void FMatrixStructCustomization<T>::CustomizeRotation(TSharedRef<class IProperty
 	const EAxisList::Type XAxis = bUseLeftUpForwardAxisDisplayCoordinateSystem ? EAxisList::Forward : EAxisList::X;
 	const EAxisList::Type YAxis = bUseLeftUpForwardAxisDisplayCoordinateSystem ? EAxisList::Left : EAxisList::Y;
 	const EAxisList::Type ZAxis = bUseLeftUpForwardAxisDisplayCoordinateSystem ? EAxisList::Up : EAxisList::Z;
+	
+	constexpr int32 NumComponents = 3;
+	TStaticArray<TFunction<TSharedRef<SWidget>()>, NumComponents> ComponentConstructors = {
+		[this, &StructPropertyHandle, XAxis]()->TSharedRef<SWidget>
+		{
+			return MakeNumericProxyWidget<UE::Math::TRotator<T>, T>(StructPropertyHandle, CachedRotationRoll, AxisDisplayInfo::GetRotationAxisNameShort(XAxis), true, AxisDisplayInfo::GetAxisColor(XAxis));
+		},
+		[this, &StructPropertyHandle, YAxis]()->TSharedRef<SWidget>
+		{
+			return MakeNumericProxyWidget<UE::Math::TRotator<T>, T>(StructPropertyHandle, CachedRotationPitch, AxisDisplayInfo::GetRotationAxisNameShort(YAxis), true, AxisDisplayInfo::GetAxisColor(YAxis));
+		},
+		[this, &StructPropertyHandle, ZAxis]()->TSharedRef<SWidget>
+		{
+			return MakeNumericProxyWidget<UE::Math::TRotator<T>, T>(StructPropertyHandle, CachedRotationYaw, AxisDisplayInfo::GetRotationAxisNameShort(ZAxis), true, AxisDisplayInfo::GetAxisColor(ZAxis));
+		}
+	};
+	const TStaticArray<FMargin, NumComponents> Paddings =
+	{
+		FMargin(0.0f, 2.0f, 3.0f, 2.0f),
+		FMargin(0.0f, 2.0f, 3.0f, 2.0f),
+		FMargin(0.0f, 2.0f, 0.0f, 2.0f)
+	};
+
+	TSharedRef<SHorizontalBox> HBox = MakeShared<SHorizontalBox>();
+	FIntVector4 Swizzle = GetSwizzle();
+	for (int32 ComponentIndex = 0; ComponentIndex < NumComponents; ComponentIndex++)
+	{
+		TSharedRef<SWidget> Widget = SNullWidget::NullWidget;
+		if (ensure(Swizzle[ComponentIndex] < NumComponents))
+		{
+			const int32 SwizzledComponentIndex = Swizzle[ComponentIndex];
+			Widget = ComponentConstructors[SwizzledComponentIndex]();
+		}
+
+		HBox->AddSlot()
+			.Padding(Paddings[ComponentIndex])
+			[
+				Widget
+			];
+	}
 
 	Row
 	.CopyAction(FUIAction(FExecuteAction::CreateSP(this, &FMatrixStructCustomization<T>::OnCopy, FTransformField::Rotation, WeakHandlePtr)))
@@ -235,22 +305,7 @@ void FMatrixStructCustomization<T>::CustomizeRotation(TSharedRef<class IProperty
 	.MinDesiredWidth(375.0f)
 	.MaxDesiredWidth(375.0f)
 	[
-		SNew(SHorizontalBox)
-		+ SHorizontalBox::Slot()
-		.Padding(FMargin(0.0f, 2.0f, 3.0f, 2.0f))
-		[
-			MakeNumericProxyWidget<UE::Math::TRotator<T>, T>(StructPropertyHandle, CachedRotationRoll, LOCTEXT("RotationRoll", "Roll"), true, AxisDisplayInfo::GetAxisColor(XAxis))
-		]
-		+ SHorizontalBox::Slot()
-		.Padding(FMargin(0.0f, 2.0f, 3.0f, 2.0f))
-		[
-			MakeNumericProxyWidget<UE::Math::TRotator<T>, T>(StructPropertyHandle, CachedRotationPitch, LOCTEXT("RotationPitch", "Pitch"), true, AxisDisplayInfo::GetAxisColor(YAxis))
-		]
-		+ SHorizontalBox::Slot()
-		.Padding(FMargin(0.0f, 2.0f, 0.0f, 2.0f))
-		[
-			MakeNumericProxyWidget<UE::Math::TRotator<T>, T>(StructPropertyHandle, CachedRotationYaw, LOCTEXT("RotationYaw", "Yaw"), true, AxisDisplayInfo::GetAxisColor(ZAxis))
-		]
+		HBox
 	];
 }
 
@@ -268,6 +323,46 @@ void FMatrixStructCustomization<T>::CustomizeScale(TSharedRef<class IPropertyHan
 	const EAxisList::Type YAxis = bUseLeftUpForwardAxisDisplayCoordinateSystem ? EAxisList::Left : EAxisList::Y;
 	const EAxisList::Type ZAxis = bUseLeftUpForwardAxisDisplayCoordinateSystem ? EAxisList::Up : EAxisList::Z;
 
+	constexpr int32 NumComponents = 3;
+	TStaticArray<TFunction<TSharedRef<SWidget>()>, NumComponents> ComponentConstructors = {
+		[this, &StructPropertyHandle, XAxis]()->TSharedRef<SWidget>
+		{
+			return MakeNumericProxyWidget<UE::Math::TVector<T>, T>(StructPropertyHandle, CachedScaleX, AxisDisplayInfo::GetAxisDisplayName(XAxis), false, AxisDisplayInfo::GetAxisColor(XAxis));
+		},
+		[this, &StructPropertyHandle, YAxis]()->TSharedRef<SWidget>
+		{
+			return MakeNumericProxyWidget<UE::Math::TVector<T>, T>(StructPropertyHandle, CachedScaleY, AxisDisplayInfo::GetAxisDisplayName(YAxis), false, AxisDisplayInfo::GetAxisColor(YAxis));
+		},
+		[this, &StructPropertyHandle, ZAxis]()->TSharedRef<SWidget>
+		{
+			return MakeNumericProxyWidget<UE::Math::TVector<T>, T>(StructPropertyHandle, CachedScaleZ, AxisDisplayInfo::GetAxisDisplayName(ZAxis), false, AxisDisplayInfo::GetAxisColor(ZAxis));
+		}
+	};
+	const TStaticArray<FMargin, NumComponents> Paddings =
+	{
+		FMargin(0.0f, 2.0f, 3.0f, 2.0f),
+		FMargin(0.0f, 2.0f, 3.0f, 2.0f),
+		FMargin(0.0f, 2.0f, 0.0f, 2.0f)
+	};
+
+	TSharedRef<SHorizontalBox> HBox = MakeShared<SHorizontalBox>();
+	FIntVector4 Swizzle = GetSwizzle();
+	for (int32 ComponentIndex = 0; ComponentIndex < NumComponents; ComponentIndex++)
+	{
+		TSharedRef<SWidget> Widget = SNullWidget::NullWidget;
+		if (ensure(Swizzle[ComponentIndex] < NumComponents))
+		{
+			const int32 SwizzledComponentIndex = Swizzle[ComponentIndex];
+			Widget = ComponentConstructors[SwizzledComponentIndex]();
+		}
+
+		HBox->AddSlot()
+			.Padding(Paddings[ComponentIndex])
+			[
+				Widget
+			];
+	}
+
 	Row
 	.CopyAction(FUIAction(FExecuteAction::CreateSP(this, &FMatrixStructCustomization<T>::OnCopy, FTransformField::Scale, WeakHandlePtr)))
 	.PasteAction(FUIAction(FExecuteAction::CreateSP(this, &FMatrixStructCustomization<T>::OnPaste, FTransformField::Scale, WeakHandlePtr)))
@@ -279,22 +374,7 @@ void FMatrixStructCustomization<T>::CustomizeScale(TSharedRef<class IPropertyHan
 	.MinDesiredWidth(375.0f)
 	.MaxDesiredWidth(375.0f)
 	[
-		SNew(SHorizontalBox)
-		+ SHorizontalBox::Slot()
-		.Padding(FMargin(0.0f, 2.0f, 3.0f, 2.0f))
-		[
-			MakeNumericProxyWidget<UE::Math::TVector<T>, T>(StructPropertyHandle, CachedScaleX, AxisDisplayInfo::GetAxisDisplayName(XAxis), false, AxisDisplayInfo::GetAxisColor(XAxis))
-		]
-		+ SHorizontalBox::Slot()
-		.Padding(FMargin(0.0f, 2.0f, 3.0f, 2.0f))
-		[
-			MakeNumericProxyWidget<UE::Math::TVector<T>, T>(StructPropertyHandle, CachedScaleY, AxisDisplayInfo::GetAxisDisplayName(YAxis), false, AxisDisplayInfo::GetAxisColor(YAxis))
-		]
-		+ SHorizontalBox::Slot()
-		.Padding(FMargin(0.0f, 2.0f, 0.0f, 2.0f))
-		[
-			MakeNumericProxyWidget<UE::Math::TVector<T>, T>(StructPropertyHandle, CachedScaleZ, AxisDisplayInfo::GetAxisDisplayName(ZAxis), false, AxisDisplayInfo::GetAxisColor(ZAxis))
-		]
+		HBox
 	];
 }
 
@@ -601,6 +681,12 @@ bool FMatrixStructCustomization<T>::FlushValues( TWeakPtr<IPropertyHandle> Prope
 	return true;
 }
 
+template <typename T>
+FIntVector4 FMatrixStructCustomization<T>::GetSwizzle() const
+{
+	return {0, 1, 2, 3};
+}
+
 template<typename T>
 TSharedRef<IPropertyTypeCustomization> FTransformStructCustomization<T>::MakeInstance() 
 {
@@ -768,6 +854,12 @@ bool FTransformStructCustomization<T>::FlushValues( TWeakPtr<IPropertyHandle> Pr
 	return true;
 }
 
+template <typename T>
+FIntVector4 FTransformStructCustomization<T>::GetSwizzle() const
+{
+	return AxisDisplayInfo::DefaultAxisComponentDisplaySwizzle();
+}
+
 template<typename T>
 TSharedRef<IPropertyTypeCustomization> FQuatStructCustomization<T>::MakeInstance()
 {
diff --git a/Engine/Source/Editor/DetailCustomizations/Public/Customizations/MathStructProxyCustomizations.h b/Engine/Source/Editor/DetailCustomizations/Public/Customizations/MathStructProxyCustomizations.h
--- a/Engine/Source/Editor/DetailCustomizations/Public/Customizations/MathStructProxyCustomizations.h
+++ b/Engine/Source/Editor/DetailCustomizations/Public/Customizations/MathStructProxyCustomizations.h
@@ -294,6 +294,7 @@ protected:
 	/** FMathStructProxyCustomization interface */
 	virtual bool CacheValues(TWeakPtr<IPropertyHandle> PropertyHandlePtr) const override;
 	virtual bool FlushValues(TWeakPtr<IPropertyHandle> PropertyHandlePtr) const override;
+	virtual FIntVector4 GetSwizzle() const;
 
 	void OnCopy(FTransformField::Type Type, TWeakPtr<IPropertyHandle> PropertyHandlePtr);
 	void OnPaste(FTransformField::Type Type, TWeakPtr<IPropertyHandle> PropertyHandlePtr);
@@ -347,6 +348,7 @@ protected:
 	/** FMathStructProxyCustomization interface */
 	virtual bool CacheValues(TWeakPtr<IPropertyHandle> PropertyHandlePtr) const override;
 	virtual bool FlushValues(TWeakPtr<IPropertyHandle> PropertyHandlePtr) const override;
+	virtual FIntVector4 GetSwizzle() const override;
 };
 
 /**
diff --git a/Engine/Source/Editor/UnrealEd/Private/EditorAxisDisplayInfo.cpp b/Engine/Source/Editor/UnrealEd/Private/EditorAxisDisplayInfo.cpp
--- a/Engine/Source/Editor/UnrealEd/Private/EditorAxisDisplayInfo.cpp
+++ b/Engine/Source/Editor/UnrealEd/Private/EditorAxisDisplayInfo.cpp
@@ -166,6 +166,7 @@ FText FEditorAxisDisplayInfo::GetAxisDisplayNameShort(EAxisList::Type Axis)
 	}
 	else if (Axis == EAxisList::Forward)
 	{
+		
 		return ForwardDisplayNameShort;
 	}
 
@@ -196,6 +197,11 @@ FLinearColor FEditorAxisDisplayInfo::GetAxisColor(EAxisList::Type Axis)
 	return FLinearColor::Black;
 }
 
+FIntVector4 FEditorAxisDisplayInfo::DefaultAxisComponentDisplaySwizzle() const
+{
+	return FIntVector4(0, 1, 2, 3);
+}
+
 bool FEditorAxisDisplayInfo::UseForwardRightUpDisplayNames()
 {
 	if (!bUseForwardRightUpDisplayNames.IsSet())
@@ -207,6 +213,31 @@ bool FEditorAxisDisplayInfo::UseForwardRightUpDisplayNames()
 	return bUseForwardRightUpDisplayNames.Get(false/*DefaultValue*/);
 }
 
+FText FEditorAxisDisplayInfo::GetRotationAxisName(EAxisList::Type Axis)
+{
+	return FEditorAxisDisplayInfo::GetRotationAxisNameShort(Axis);
+}
+
+FText FEditorAxisDisplayInfo::GetRotationAxisNameShort(EAxisList::Type Axis)
+{
+	if (Axis == EAxisList::X || Axis == EAxisList::Forward)
+	{
+		return LOCTEXT("Roll", "Roll");
+	}
+	else if (Axis == EAxisList::Y || Axis == EAxisList::Left)
+	{
+		return LOCTEXT("Pitch", "Pitch");
+	}
+	else if (Axis == EAxisList::Z || Axis == EAxisList::Up)
+	{
+		return LOCTEXT("Yaw", "Yaw");
+	}
+
+	ensureMsgf(false, TEXT("Unsupported Axis: %d"), Axis);
+
+	return LOCTEXT("UnsupportedDisplayNameShort", "?");
+}
+
 EAxisList::Type FEditorAxisDisplayInfo::MapAxis(EAxisList::Type Axis)
 {
 	if (GetAxisDisplayCoordinateSystem() == EAxisList::XYZ)
diff --git a/Engine/Source/Editor/UnrealEd/Private/UnrealWidgetRender.cpp b/Engine/Source/Editor/UnrealEd/Private/UnrealWidgetRender.cpp
--- a/Engine/Source/Editor/UnrealEd/Private/UnrealWidgetRender.cpp
+++ b/Engine/Source/Editor/UnrealEd/Private/UnrealWidgetRender.cpp
@@ -1208,6 +1208,13 @@ void FWidget::DrawRotationArc(const FSceneView* View, FPrimitiveDrawInterface* P
 				{
 					float OffsetAngle = IsRotationLocalSpace() ? 0 : AdjustedDeltaRotation;
 
+					if (AxisDisplayInfo::GetAxisDisplayCoordinateSystem() == EAxisList::LeftUpForward)
+					{
+						if (InAxis == EAxisList::Y || InAxis == EAxisList::Z)
+						{
+							DeltaRotation = -DeltaRotation;
+						}
+					}
 					CacheRotationHUDText(View, PDI, InLocation, Axis0.RotateAngleAxis(OffsetAngle, ZAxis), Axis1.RotateAngleAxis(OffsetAngle, ZAxis), DeltaRotation, InScale);
 				}
 			}
diff --git a/Engine/Source/Editor/UnrealEd/Public/EditorAxisDisplayInfo.h b/Engine/Source/Editor/UnrealEd/Public/EditorAxisDisplayInfo.h
--- a/Engine/Source/Editor/UnrealEd/Public/EditorAxisDisplayInfo.h
+++ b/Engine/Source/Editor/UnrealEd/Public/EditorAxisDisplayInfo.h
@@ -21,6 +21,11 @@ public:
 
 	virtual bool UseForwardRightUpDisplayNames() override;
 
+	virtual FText GetRotationAxisName(EAxisList::Type Axis) override;
+	virtual FText GetRotationAxisNameShort(EAxisList::Type Axis) override;
+	
+	virtual FIntVector4 DefaultAxisComponentDisplaySwizzle() const override;
+
 private:
 	// Maps the given axis from FLU -> XYZ if the AxisDisplayCoordinateSystem is XYZ
 	EAxisList::Type MapAxis(EAxisList::Type Axis);
diff --git a/Engine/Source/Runtime/Core/Private/Misc/AxisDisplayInfo.cpp b/Engine/Source/Runtime/Core/Private/Misc/AxisDisplayInfo.cpp
--- a/Engine/Source/Runtime/Core/Private/Misc/AxisDisplayInfo.cpp
+++ b/Engine/Source/Runtime/Core/Private/Misc/AxisDisplayInfo.cpp
@@ -1,6 +1,7 @@
 // Copyright Epic Games, Inc. All Rights Reserved.
 
 #include "Misc/AxisDisplayInfo.h"
+#include "Math/IntVector.h"
 
 #include "Features/IModularFeatures.h"
 #include "Internationalization/Internationalization.h"
@@ -52,22 +53,34 @@ FText AxisDisplayInfo::GetAxisDisplayNameShort(EAxisList::Type Axis)
 		return DisplayInfo.GetAxisDisplayNameShort(Axis);
 	}
 
-	if (Axis == EAxisList::X || Axis == EAxisList::Forward)
+	if (Axis == EAxisList::X)
 	{
 		return LOCTEXT("XDisplayNameShort", "X");
 	}
-	else if (Axis == EAxisList::Y || Axis == EAxisList::Left)
+	else if (Axis == EAxisList::Y)
 	{
 		return LOCTEXT("YDisplayNameShort", "Y");
 	}
-	else if (Axis == EAxisList::Z || Axis == EAxisList::Up)
+	else if (Axis == EAxisList::Z)
 	{
 		return LOCTEXT("ZDisplayNameShort", "Z");
 	}
+	else if (Axis == EAxisList::Forward)
+	{
+		return LOCTEXT("ForwardDisplayNameShort", "Forward");
+	}
+	else if (Axis == EAxisList::Left)
+	{
+		return LOCTEXT("LeftDisplayNameShort", "Left");
+	}
+	else if (Axis == EAxisList::Up)
+	{
+		return LOCTEXT("UpDisplayNameShort", "Up");
+	}
 
 	ensureMsgf(false, TEXT("Unsupported Axis: %d"), Axis);
 
-	return LOCTEXT("UnsupportedDisplayNameShort", "?");
+	return LOCTEXT("UnsupportedGetAxisDisplayNameShort", "?");
 }
 
 FLinearColor AxisDisplayInfo::GetAxisColor(EAxisList::Type Axis)
@@ -107,4 +120,53 @@ bool AxisDisplayInfo::UseForwardRightUpDisplayNames()
 	return false;
 }
 
+FText AxisDisplayInfo::GetRotationAxisNameShort(EAxisList::Type Axis)
+{
+	if (IModularFeatures::Get().IsModularFeatureAvailable(IAxisDisplayInfo::GetModularFeatureName()))
+	{
+		IAxisDisplayInfo& DisplayInfo = IModularFeatures::Get().GetModularFeature<IAxisDisplayInfo>(IAxisDisplayInfo::GetModularFeatureName());
+		return DisplayInfo.GetRotationAxisNameShort(Axis);
+	}
+
+	if (Axis == EAxisList::X)
+	{
+		return LOCTEXT("Roll_ToolTipTextFormat", "Roll");
+	}
+	else if (Axis == EAxisList::Y)
+	{
+		return LOCTEXT("Pitch_ToolTipTextFormat", "Pitch");
+	}
+	else if (Axis == EAxisList::Z)
+	{
+		return LOCTEXT("Yaw_ToolTipTextFormat", "Yaw");
+	}
+	else if (Axis == EAxisList::Forward)
+	{
+		return LOCTEXT("Forward_ToolTipTextFormat", "Forward");
+	}
+	else if (Axis == EAxisList::Left)
+	{
+		return LOCTEXT("Left_ToolTipTextFormat", "Left");
+	}
+	else if (Axis == EAxisList::Up)
+	{
+		return LOCTEXT("Up_ToolTipTextFormat", "Up");
+	}
+
+	ensureMsgf(false, TEXT("Unsupported Axis: %d"), Axis);
+
+	return LOCTEXT("UnsupportedRotationAxisNameShort", "?");
+}
+
+FIntVector4 AxisDisplayInfo::DefaultAxisComponentDisplaySwizzle()
+{
+	if (IModularFeatures::Get().IsModularFeatureAvailable(IAxisDisplayInfo::GetModularFeatureName()))
+	{
+		IAxisDisplayInfo& DisplayInfo = IModularFeatures::Get().GetModularFeature<IAxisDisplayInfo>(IAxisDisplayInfo::GetModularFeatureName());
+		return DisplayInfo.DefaultAxisComponentDisplaySwizzle();
+	}
+	
+	return FIntVector4(0, 1, 2, 3);
+}
+
 #undef LOCTEXT_NAMESPACE
diff --git a/Engine/Source/Runtime/Core/Public/Math/Quat.h b/Engine/Source/Runtime/Core/Public/Math/Quat.h
--- a/Engine/Source/Runtime/Core/Public/Math/Quat.h
+++ b/Engine/Source/Runtime/Core/Public/Math/Quat.h
@@ -702,6 +702,11 @@ public:
 	 */
 	static CORE_API void CalcTangents(const TQuat<T>& PrevP, const TQuat<T>& P, const TQuat<T>& NextP, T Tension, TQuat<T>& OutTan);
 
+	[[nodiscard]] TTuple<T, T, T> ToLUFEuler() const;
+
+	// Converts to euler angles (in rads) in verse conventions
+	[[nodiscard]] static TQuat MakeFromLUFEuler(const TTuple<T, T, T>& InLUFEuler);
+
 public:
 
 	bool Serialize(FArchive& Ar)
@@ -1410,6 +1415,41 @@ FORCEINLINE uint32 GetTypeHash(const TQuat<T>& Quat)
 	return FCrc::MemCrc_DEPRECATED(&Quat, sizeof(Quat));
 }
 
+template<typename T>
+[[nodiscard]] TTuple<T, T, T> TQuat<T>::ToLUFEuler() const
+{
+	auto ThreeAxisRot = [](T R11, T R12, T R21, T R31, T R32, T Euler[])
+	{
+		Euler[0] = FMath::Atan2(R31, R32);
+		Euler[1] = FMath::Asin(R21);
+		Euler[2] = FMath::Atan2(R11, R12);
+	};
+	
+	T Euler[3];
+	ThreeAxisRot(2*(Y*Z + W*X),
+				 W*W - X*X + Y*Y - Z*Z,
+				 -2*(X*Y - W*Z),
+				 2*(X*Z + W*Y),
+				 W*W + X*X - Y*Y - Z*Z,
+				 Euler);
+	// Going from left-handed to right-handed
+	return {Euler[0], -Euler[1], -Euler[2]};
+}
+
+template<typename T>
+[[nodiscard]] TQuat<T> TQuat<T>::MakeFromLUFEuler(const TTuple<T, T, T>& InLUFEuler)
+{
+	const T LeftAxisAngle = InLUFEuler.template Get<0>();
+	const T UpAxisAngle = InLUFEuler.template Get<1>();
+	const T ForwardAxisAngle = InLUFEuler.template Get<2>();
+		
+	// Going from left-handed to right-handed
+	TQuat LeftRotation(TVector<T>::LeftVector, -LeftAxisAngle);
+	TQuat UpRotation(TVector<T>::UpVector, -UpAxisAngle);
+	TQuat ForwardRotation(TVector<T>::ForwardVector, -ForwardAxisAngle);
+	return ForwardRotation * UpRotation * LeftRotation;
+}
+
 } // namespace UE::Math
 } // namespace UE
 
diff --git a/Engine/Source/Runtime/Core/Public/Misc/AxisDisplayInfo.h b/Engine/Source/Runtime/Core/Public/Misc/AxisDisplayInfo.h
--- a/Engine/Source/Runtime/Core/Public/Misc/AxisDisplayInfo.h
+++ b/Engine/Source/Runtime/Core/Public/Misc/AxisDisplayInfo.h
@@ -6,6 +6,7 @@
 #include "Internationalization/Text.h"
 #include "Math/Axis.h"
 #include "Math/Color.h"
+#include "Math/MathFwd.h"
 
 namespace AxisDisplayInfo
 {
@@ -59,6 +60,15 @@ namespace AxisDisplayInfo
 	 * it should not be cached until the editor has finished loading.
 	 */
 	CORE_API bool UseForwardRightUpDisplayNames();
+	
+	CORE_API FText GetRotationAxisName(EAxisList::Type Axis);
+	CORE_API FText GetRotationAxisNameShort(EAxisList::Type Axis);
+	
+
+	/**
+	 * Default swizzle for axis displays in areas like the details panel
+	 */
+	CORE_API FIntVector4 DefaultAxisComponentDisplaySwizzle();
 };
 
 class IAxisDisplayInfo : public IModularFeature
@@ -106,6 +116,16 @@ public:
 	 */
 	virtual FLinearColor GetAxisColor(EAxisList::Type Axis) = 0;
 
+	virtual FText GetRotationAxisName(EAxisList::Type Axis) = 0;
+	virtual FText GetRotationAxisNameShort(EAxisList::Type Axis) = 0;
+
+	/**
+	 * Defines the swizzle for slate widgets displaying axis components
+	 * such as SVectorInputBox and SRotatorInputBox
+	 * Generally, these widgets will be ordered X, Y and Z
+	 */
+	virtual FIntVector4 DefaultAxisComponentDisplaySwizzle() const = 0;
+
 	/**
 	 * Gets whether or not the engine uses forward/right/up nomenclature for display names.
 	 */
diff --git a/Engine/Source/Runtime/Slate/Public/Widgets/Input/SRotatorInputBox.h b/Engine/Source/Runtime/Slate/Public/Widgets/Input/SRotatorInputBox.h
--- a/Engine/Source/Runtime/Slate/Public/Widgets/Input/SRotatorInputBox.h
+++ b/Engine/Source/Runtime/Slate/Public/Widgets/Input/SRotatorInputBox.h
@@ -22,6 +22,8 @@ class FArrangedChildren;
 template<typename NumericType>
 class SNumericRotatorInputBox : public SCompoundWidget
 {
+	using ThisClass = SNumericRotatorInputBox<NumericType>;
+	static constexpr int32 ComponentNum = 3;
 public:
 	/** Notification for float value change */
 	DECLARE_DELEGATE_OneParam(FOnNumericValueChanged, NumericType);
@@ -30,7 +32,11 @@ public:
 	DECLARE_DELEGATE_TwoParams(FOnNumericValueCommitted, NumericType, ETextCommit::Type);
 	
 	SLATE_BEGIN_ARGS( SNumericRotatorInputBox<NumericType> )
-		: _bColorAxisLabels(false)
+		: _RollDisplayName(AxisDisplayInfo::GetRotationAxisNameShort(EAxisList::X))
+		, _PitchDisplayName(AxisDisplayInfo::GetRotationAxisNameShort(EAxisList::Y))
+		, _YawDisplayName(AxisDisplayInfo::GetRotationAxisNameShort(EAxisList::Z))
+		, _bColorAxisLabels(false)
+		, _Swizzle(AxisDisplayInfo::DefaultAxisComponentDisplaySwizzle())
 		, _Font(FAppStyle::Get().GetFontStyle("NormalFont"))
 		, _AllowSpin(true)
 		, _SpinDelta(1.f)
@@ -54,9 +60,20 @@ public:
 		/** Yaw component of the rotator */
 		SLATE_ATTRIBUTE( TOptional<NumericType>, Yaw )
 
+		SLATE_ATTRIBUTE( FText, RollDisplayName )
+		
+		SLATE_ATTRIBUTE( FText, PitchDisplayName )
+		
+		SLATE_ATTRIBUTE( FText, YawDisplayName )
+
 		/** Should the axis labels be colored */
 		SLATE_ARGUMENT( bool, bColorAxisLabels )
 
+		/** Swizzles the display of 0: Roll, 1: Pitch, 2: Yaw
+		 * 4th component unused
+		 */
+		SLATE_ARGUMENT( FIntVector4, Swizzle)
+
 		/** Font to use for the text in this box */
 		SLATE_ATTRIBUTE( FSlateFontInfo, Font )
 
@@ -149,6 +166,127 @@ public:
 
 	SLATE_END_ARGS()
 
+private:
+
+	TAttribute<TOptional<FTextFormat>> BuildToolTipTextFormatString(FStringView DisplayNameString)
+	{
+		TStringBuilder<32> ToolTipTextFormatString;
+		ToolTipTextFormatString.Append(DisplayNameString);
+		ToolTipTextFormatString.Append(TEXT(": {0}"));
+		TAttribute<TOptional<FTextFormat>> Attribute(FText::FromString(ToolTipTextFormatString.ToString()));
+		return Attribute;
+	}
+
+	TSharedRef<SWidget> ConstructRoll(const FArguments& InArgs)
+	{
+		TSharedRef<SWidget> Widget =
+			SNew(SNumericEntryBox<NumericType>)
+			.AllowSpin(InArgs._AllowSpin)
+			.Delta(InArgs._SpinDelta)
+			.LinearDeltaSensitivity(InArgs._LinearDeltaSensitivity)
+			.MinValue(InArgs._MinSliderValue)
+			.MaxValue(InArgs._MaxSliderValue)
+			.MinSliderValue(InArgs._MinSliderValue)
+			.MaxSliderValue(InArgs._MaxSliderValue)
+			.LabelPadding(FMargin(3.0f))
+			.LabelLocation(SNumericEntryBox<NumericType>::ELabelLocation::Inside)
+			.Label()
+			[
+				InArgs._bColorAxisLabels ? SNumericEntryBox<NumericType>::BuildNarrowColorLabel(AxisDisplayInfo::GetAxisColor(EAxisList::X)) : SNullWidget::NullWidget
+			]
+			.Font( InArgs._Font )
+			.Value( InArgs._Roll )
+			.OnValueChanged( InArgs._OnRollChanged )
+			.OnValueCommitted( InArgs._OnRollCommitted )
+			.OnBeginSliderMovement( CreatePerComponentSliderMovementEvent(InArgs._OnBeginSliderMovement, InArgs._OnRollBeginSliderMovement) )
+			.OnEndSliderMovement( CreatePerComponentSliderMovementEvent< FOnNumericValueChanged, NumericType >(InArgs._OnEndSliderMovement, InArgs._OnRollEndSliderMovement) )
+			.UndeterminedString( NSLOCTEXT("SRotatorInputBox", "MultipleValues", "Multiple Values") )
+			.ToolTipTextFormat(BuildToolTipTextFormatString(InArgs._RollDisplayName.Get().ToString()))
+			.TypeInterface(InArgs._TypeInterface)
+			.MinFractionalDigits(InArgs._TypeInterface ? InArgs._TypeInterface->GetMinFractionalDigits() : TOptional<int32>())
+			.MaxFractionalDigits(InArgs._TypeInterface ? InArgs._TypeInterface->GetMaxFractionalDigits() : TOptional<int32>())
+			.DisplayToggle(InArgs._DisplayToggle)
+			.ToggleChecked(InArgs._ToggleRollChecked)
+			.OnToggleChanged(InArgs._OnToggleRollChanged)
+			.TogglePadding(InArgs._TogglePadding)
+			.PreventThrottling(InArgs._PreventThrottling);
+		return Widget;
+	};
+	
+	TSharedRef<SWidget> ConstructPitch(const FArguments& InArgs)
+	{
+		TSharedRef<SWidget> Widget =
+			SNew(SNumericEntryBox<NumericType>)
+			.AllowSpin(InArgs._AllowSpin)
+			.Delta(InArgs._SpinDelta)
+			.LinearDeltaSensitivity(InArgs._LinearDeltaSensitivity)
+			.MinValue(InArgs._MinSliderValue)
+			.MaxValue(InArgs._MaxSliderValue)
+			.MinSliderValue(InArgs._MinSliderValue)
+			.MaxSliderValue(InArgs._MaxSliderValue)
+			.LabelPadding(FMargin(3.0f))
+			.LabelLocation(SNumericEntryBox<NumericType>::ELabelLocation::Inside)
+			.Label()
+			[
+				InArgs._bColorAxisLabels ? SNumericEntryBox<NumericType>::BuildNarrowColorLabel(AxisDisplayInfo::GetAxisColor(EAxisList::Y)) : SNullWidget::NullWidget
+			]
+			.Font( InArgs._Font )
+			.Value( InArgs._Pitch )
+			.OnValueChanged( InArgs._OnPitchChanged )
+			.OnValueCommitted( InArgs._OnPitchCommitted )
+			.OnBeginSliderMovement( CreatePerComponentSliderMovementEvent( InArgs._OnBeginSliderMovement, InArgs._OnPitchBeginSliderMovement ) )
+			.OnEndSliderMovement( CreatePerComponentSliderMovementEvent< FOnNumericValueChanged, NumericType >( InArgs._OnEndSliderMovement, InArgs._OnPitchEndSliderMovement ) )
+			.UndeterminedString( NSLOCTEXT("SRotatorInputBox", "MultipleValues", "Multiple Values") )
+			.ToolTipTextFormat(BuildToolTipTextFormatString(InArgs._PitchDisplayName.Get().ToString()))
+			.TypeInterface(InArgs._TypeInterface)
+			.MinFractionalDigits(InArgs._TypeInterface ? InArgs._TypeInterface->GetMinFractionalDigits() : TOptional<int32>())
+			.MaxFractionalDigits(InArgs._TypeInterface ? InArgs._TypeInterface->GetMaxFractionalDigits() : TOptional<int32>())
+			.DisplayToggle(InArgs._DisplayToggle)
+			.ToggleChecked(InArgs._TogglePitchChecked)
+			.OnToggleChanged(InArgs._OnTogglePitchChanged)
+			.TogglePadding(InArgs._TogglePadding)
+			.PreventThrottling(InArgs._PreventThrottling);
+		return Widget;
+	};
+
+	TSharedRef<SWidget> ConstructYaw(const FArguments& InArgs)
+	{
+		TSharedRef<SWidget> Widget =
+			SNew(SNumericEntryBox<NumericType>)
+			.AllowSpin(InArgs._AllowSpin)
+			.Delta(InArgs._SpinDelta)
+			.LinearDeltaSensitivity(InArgs._LinearDeltaSensitivity)
+			.MinValue(InArgs._MinSliderValue)
+			.MaxValue(InArgs._MaxSliderValue)
+			.MinSliderValue(InArgs._MinSliderValue)
+			.MaxSliderValue(InArgs._MaxSliderValue)
+			.LabelPadding(FMargin(3.0f))
+			.LabelLocation(SNumericEntryBox<NumericType>::ELabelLocation::Inside)
+			.Label()
+			[
+				InArgs._bColorAxisLabels ? SNumericEntryBox<NumericType>::BuildNarrowColorLabel(AxisDisplayInfo::GetAxisColor(EAxisList::Z)) : SNullWidget::NullWidget
+			]
+			.Font( InArgs._Font )
+			.Value( InArgs._Yaw )
+			.OnValueChanged( InArgs._OnYawChanged )
+			.OnValueCommitted( InArgs._OnYawCommitted )
+			.OnBeginSliderMovement( CreatePerComponentSliderMovementEvent(InArgs._OnBeginSliderMovement, InArgs._OnYawBeginSliderMovement) )
+			.OnEndSliderMovement( CreatePerComponentSliderMovementEvent< FOnNumericValueChanged, NumericType >(InArgs._OnEndSliderMovement, InArgs._OnYawEndSliderMovement) )
+			.UndeterminedString( NSLOCTEXT("SRotatorInputBox", "MultipleValues", "Multiple Values") )
+			.ToolTipTextFormat(BuildToolTipTextFormatString(InArgs._YawDisplayName.Get().ToString()))
+			.TypeInterface(InArgs._TypeInterface)
+			.MinFractionalDigits(InArgs._TypeInterface ? InArgs._TypeInterface->GetMinFractionalDigits() : TOptional<int32>())
+			.MaxFractionalDigits(InArgs._TypeInterface ? InArgs._TypeInterface->GetMaxFractionalDigits() : TOptional<int32>())
+			.DisplayToggle(InArgs._DisplayToggle)
+			.ToggleChecked(InArgs._ToggleYawChecked)
+			.OnToggleChanged(InArgs._OnToggleYawChanged)
+			.TogglePadding(InArgs._TogglePadding)
+			.PreventThrottling(InArgs._PreventThrottling);
+		return Widget;
+	};
+	
+public:
+
 	/**
 	 * Construct this widget
 	 *
@@ -156,104 +294,38 @@ public:
 	 */
 	void Construct( const FArguments& InArgs )
 	{
-		this->ChildSlot
-		[
-			SNew(SHorizontalBox)
-			+SHorizontalBox::Slot()
-			[
-				SNew(SNumericEntryBox<NumericType>)
-				.AllowSpin(InArgs._AllowSpin)
-				.Delta(InArgs._SpinDelta)
-				.LinearDeltaSensitivity(InArgs._LinearDeltaSensitivity)
-				.MinValue(InArgs._MinSliderValue)
-				.MaxValue(InArgs._MaxSliderValue)
-				.MinSliderValue(InArgs._MinSliderValue)
-				.MaxSliderValue(InArgs._MaxSliderValue)
-				.LabelPadding(FMargin(3.0f))
-				.LabelLocation(SNumericEntryBox<NumericType>::ELabelLocation::Inside)
-				.Label()
-				[
-					InArgs._bColorAxisLabels ? SNumericEntryBox<NumericType>::BuildNarrowColorLabel(AxisDisplayInfo::GetAxisColor(EAxisList::X)) : SNullWidget::NullWidget
-				]
-				.Font( InArgs._Font )
-				.Value( InArgs._Roll )
-				.OnValueChanged( InArgs._OnRollChanged )
-				.OnValueCommitted( InArgs._OnRollCommitted )
-				.OnBeginSliderMovement( CreatePerComponentSliderMovementEvent(InArgs._OnBeginSliderMovement, InArgs._OnRollBeginSliderMovement) )
-				.OnEndSliderMovement( CreatePerComponentSliderMovementEvent< FOnNumericValueChanged, NumericType >(InArgs._OnEndSliderMovement, InArgs._OnRollEndSliderMovement) )
-				.UndeterminedString( NSLOCTEXT("SRotatorInputBox", "MultipleValues", "Multiple Values") )
-				.ToolTipTextFormat(TAttribute<TOptional<FTextFormat>>(NSLOCTEXT("SRotatorInputBox", "Roll_ToolTipTextFormat", "Roll: {0}")))
-				.TypeInterface(InArgs._TypeInterface)
-				.DisplayToggle(InArgs._DisplayToggle)
-				.ToggleChecked(InArgs._ToggleRollChecked)
-				.OnToggleChanged(InArgs._OnToggleRollChanged)
-				.TogglePadding(InArgs._TogglePadding)
-				.PreventThrottling(InArgs._PreventThrottling)
-			]
-			+SHorizontalBox::Slot()
-			[
-				SNew(SNumericEntryBox<NumericType>)
-				.AllowSpin(InArgs._AllowSpin)
-				.Delta(InArgs._SpinDelta)
-				.LinearDeltaSensitivity(InArgs._LinearDeltaSensitivity)
-				.MinValue(InArgs._MinSliderValue)
-				.MaxValue(InArgs._MaxSliderValue)
-				.MinSliderValue(InArgs._MinSliderValue)
-				.MaxSliderValue(InArgs._MaxSliderValue)
-				.LabelPadding(FMargin(3.0f))
-				.LabelLocation(SNumericEntryBox<NumericType>::ELabelLocation::Inside)
-				.Label()
-				[
-					InArgs._bColorAxisLabels ? SNumericEntryBox<NumericType>::BuildNarrowColorLabel(AxisDisplayInfo::GetAxisColor(EAxisList::Y)) : SNullWidget::NullWidget
-				]
-				.Font( InArgs._Font )
-				.Value( InArgs._Pitch )
-				.OnValueChanged( InArgs._OnPitchChanged )
-				.OnValueCommitted( InArgs._OnPitchCommitted )
-				.OnBeginSliderMovement( CreatePerComponentSliderMovementEvent( InArgs._OnBeginSliderMovement, InArgs._OnPitchBeginSliderMovement ) )
-				.OnEndSliderMovement( CreatePerComponentSliderMovementEvent< FOnNumericValueChanged, NumericType >( InArgs._OnEndSliderMovement, InArgs._OnPitchEndSliderMovement ) )
-				.UndeterminedString( NSLOCTEXT("SRotatorInputBox", "MultipleValues", "Multiple Values") )
-				.ToolTipTextFormat(TAttribute<TOptional<FTextFormat>>(NSLOCTEXT("SRotatorInputBox", "Pitch_ToolTipTextFormat", "Pitch: {0}")))
-				.TypeInterface(InArgs._TypeInterface)
-				.DisplayToggle(InArgs._DisplayToggle)
-				.ToggleChecked(InArgs._TogglePitchChecked)
-				.OnToggleChanged(InArgs._OnTogglePitchChanged)
-				.TogglePadding(InArgs._TogglePadding)
-				.PreventThrottling(InArgs._PreventThrottling)
-			]
-			+SHorizontalBox::Slot()
+		using WidgetConstructorFn = TSharedRef<SWidget> (ThisClass::*)(const FArguments&);
+		constexpr TStaticArray<WidgetConstructorFn, ComponentNum> Constructors =
+		{
+			&ThisClass::ConstructRoll,
+			&ThisClass::ConstructPitch,
+			&ThisClass::ConstructYaw
+		};
+
+		TSharedRef<SHorizontalBox> HorizontalBox = SNew(SHorizontalBox);
+
+		for (int32 SwizzleIndex = 0; SwizzleIndex < ComponentNum; SwizzleIndex++)
+		{
+			TSharedPtr<SWidget> Widget = SNullWidget::NullWidget;
+			const auto& Swizzle = InArgs._Swizzle;
+			int32 SwizzleValue = Swizzle[SwizzleIndex];
+			if (ensureMsgf(SwizzleValue < 3 && SwizzleValue >= 0, TEXT("Invalid Swizzle Value")))
+			{
+				Widget = (this->*Constructors[SwizzleValue])(InArgs);
+			}
+
+			TSharedRef<SWidget> WidgetRef = Widget.ToSharedRef();
+			
+			HorizontalBox->AddSlot()
 			[
-				SNew(SNumericEntryBox<NumericType>)
-				.AllowSpin(InArgs._AllowSpin)
-				.Delta(InArgs._SpinDelta)
-				.LinearDeltaSensitivity(InArgs._LinearDeltaSensitivity)
-				.MinValue(InArgs._MinSliderValue)
-				.MaxValue(InArgs._MaxSliderValue)
-				.MinSliderValue(InArgs._MinSliderValue)
-				.MaxSliderValue(InArgs._MaxSliderValue)
-				.LabelPadding(FMargin(3.0f))
-				.LabelLocation(SNumericEntryBox<NumericType>::ELabelLocation::Inside)
-				.Label()
-				[
-					InArgs._bColorAxisLabels ? SNumericEntryBox<NumericType>::BuildNarrowColorLabel(AxisDisplayInfo::GetAxisColor(EAxisList::Z)) : SNullWidget::NullWidget
-				]
-				.Font( InArgs._Font )
-				.Value( InArgs._Yaw )
-				.OnValueChanged( InArgs._OnYawChanged )
-				.OnValueCommitted( InArgs._OnYawCommitted )
-				.OnBeginSliderMovement( CreatePerComponentSliderMovementEvent(InArgs._OnBeginSliderMovement, InArgs._OnYawBeginSliderMovement) )
-				.OnEndSliderMovement( CreatePerComponentSliderMovementEvent< FOnNumericValueChanged, NumericType >(InArgs._OnEndSliderMovement, InArgs._OnYawEndSliderMovement) )
-				.UndeterminedString( NSLOCTEXT("SRotatorInputBox", "MultipleValues", "Multiple Values") )
-				.ToolTipTextFormat(TAttribute<TOptional<FTextFormat>>(NSLOCTEXT("SRotatorInputBox", "Yaw_ToolTipTextFormat", "Yaw: {0}")))
-				.TypeInterface(InArgs._TypeInterface)
-				.DisplayToggle(InArgs._DisplayToggle)
-				.ToggleChecked(InArgs._ToggleYawChecked)
-				.OnToggleChanged(InArgs._OnToggleYawChanged)
-				.TogglePadding(InArgs._TogglePadding)
-				.PreventThrottling(InArgs._PreventThrottling)
-			]
-		];
+				WidgetRef
+			];
+		}
 
+		ChildSlot
+		[
+			HorizontalBox
+		];
 	}
 
 	/**
diff --git a/Engine/Source/Runtime/Slate/Public/Widgets/Input/SVectorInputBox.h b/Engine/Source/Runtime/Slate/Public/Widgets/Input/SVectorInputBox.h
--- a/Engine/Source/Runtime/Slate/Public/Widgets/Input/SVectorInputBox.h
+++ b/Engine/Source/Runtime/Slate/Public/Widgets/Input/SVectorInputBox.h
@@ -62,6 +62,7 @@ private:
 				{
 					_XDisplayName.Set(AxisDisplayInfo::GetAxisDisplayName(EAxisList::X));
 				}
+				_XColor.Set(AxisDisplayInfo::GetAxisColor(EAxisList::X));
 			}
 		}
 
@@ -91,6 +92,8 @@ private:
 
 		/** The DisplayName of the X component of the vector */
 		SLATE_ATTRIBUTE(FText, XDisplayName)
+
+		SLATE_ATTRIBUTE(FLinearColor, XColor)
 	};
 
 	struct FVectorYArgumentsEmpty {};
@@ -109,6 +112,7 @@ private:
 				{
 					_YDisplayName.Set(AxisDisplayInfo::GetAxisDisplayName(EAxisList::Y));
 				}
+				_YColor.Set(AxisDisplayInfo::GetAxisColor(EAxisList::Y));
 			}
 		}
 
@@ -138,6 +142,8 @@ private:
 
 		/** The DisplayName of the Y component of the vector */
 		SLATE_ATTRIBUTE(FText, YDisplayName)
+
+		SLATE_ATTRIBUTE(FLinearColor, YColor)
 	};
 
 	struct FVectorZArgumentsEmpty {};
@@ -156,6 +162,7 @@ private:
 				{
 					_ZDisplayName.Set(AxisDisplayInfo::GetAxisDisplayName(EAxisList::Z));
 				}
+				_ZColor.Set(AxisDisplayInfo::GetAxisColor(EAxisList::Z));
 			}
 		}
 
@@ -185,6 +192,8 @@ private:
 
 		/** The DisplayName of the Z component of the vector */
 		SLATE_ATTRIBUTE(FText, ZDisplayName)
+
+		SLATE_ATTRIBUTE(FLinearColor, ZColor)
 	};
 
 	struct FVectorWArgumentsEmpty {};
@@ -196,7 +205,12 @@ private:
 		FORCENOINLINE FVectorWArguments()
 			: _ToggleWChecked(ECheckBoxState::Checked)
 			, _WDisplayName(NSLOCTEXT("SVectorInputBox", "W_DisplayName", "W"))
-		{}
+		{
+			if constexpr (NumberOfComponents >= 4)
+			{
+				_WColor.Set(SNumericEntryBox<NumericType>::LilacLabelBackgroundColor);
+			}
+		}
 
 		/** W Component of the vector */
 		SLATE_ATTRIBUTE(TOptional<NumericType>, W)
@@ -224,6 +238,8 @@ private:
 
 		/** The DisplayName of the W component of the vector */
 		SLATE_ATTRIBUTE(FText, WDisplayName)
+
+		SLATE_ATTRIBUTE(FLinearColor, WColor)
 	};
 
 public:
@@ -243,6 +259,7 @@ public:
 			, _SpinDelta(1)
 			, _LinearDeltaSensitivity(1)
 			, _bColorAxisLabels(false)
+			, _Swizzle(AxisDisplayInfo::DefaultAxisComponentDisplaySwizzle()) /* X, Y, Z, W */
 			, _DisplayToggle(false)
 			, _TogglePadding(FMargin(1.f,0.f,1.f,0.f) )
 			, _PreventThrottling(false)
@@ -289,6 +306,9 @@ public:
 
 		/** Should the axis labels be colored */
 		SLATE_ARGUMENT( bool, bColorAxisLabels )
+		
+		/** Controls the display swizzle of the vector */
+		SLATE_ARGUMENT(FIntVector4, Swizzle)
 
 		/** Called right before the slider begins to move for any of the vector components */
 		SLATE_EVENT( FSimpleDelegate, OnBeginSliderMovement )
@@ -310,8 +330,8 @@ public:
 
 		/** If refresh requests for the viewport should happen for all value changes **/
 		SLATE_ARGUMENT(bool, PreventThrottling)
-		
-	SLATE_END_ARGS()
+	}; // SLATE_END_ARGS()
+	
 
 	/**
 	 * Construct this widget
@@ -328,8 +348,7 @@ public:
 		[
 			HorizontalBox
 		];
-
-
+		
 		VectorAttribute = InArgs._Vector;
 		OnVectorValueChanged = InArgs._OnVectorChanged;
 		OnVectorValueCommitted = InArgs._OnVectorCommitted;
@@ -420,21 +439,34 @@ public:
 			});
 		}
 
+		using ComponentConstructorFn = void (ThisClass::*)(const FArguments&, TSharedRef<SHorizontalBox>);
+		TStaticArray<ComponentConstructorFn, NumberOfComponents> ComponentConstructors;
+		
 		if constexpr (NumberOfComponents >= 1)
 		{
-			ConstructX(InArgs, HorizontalBox);
+			ComponentConstructors[0] = ComponentConstructorFn(&ThisClass::ConstructX);
 		}
 		if constexpr (NumberOfComponents >= 2)
 		{
-			ConstructY(InArgs, HorizontalBox);
+			ComponentConstructors[1] = ComponentConstructorFn(&ThisClass::ConstructY);
 		}
 		if constexpr (NumberOfComponents >= 3)
 		{
-			ConstructZ(InArgs, HorizontalBox);
+			ComponentConstructors[2] = ComponentConstructorFn(&ThisClass::ConstructZ);
 		}
 		if constexpr (NumberOfComponents >= 4)
 		{
-			ConstructW(InArgs, HorizontalBox);
+			ComponentConstructors[3] = ComponentConstructorFn(&ThisClass::ConstructW);
+		}
+
+		// Call the widget construction functions in swizzle order
+		for (int32 ComponentIndex = 0; ComponentIndex < NumberOfComponents; ++ComponentIndex)
+		{
+			int32 ConstructorIndex = InArgs._Swizzle[ComponentIndex];
+			if (ensureMsgf(ConstructorIndex < NumberOfComponents, TEXT("Invalid swizzle index (%d >= %d)"), ConstructorIndex, NumberOfComponents))
+			{
+				(this->*ComponentConstructors[ConstructorIndex])(InArgs, HorizontalBox);
+			}
 		}
 	}
 
@@ -444,7 +476,7 @@ private:
 	 */
 	void ConstructComponent(int32 ComponentIndex,
 		const FArguments& InArgs,
-		const FLinearColor& LabelColor,
+		const TAttribute<FLinearColor>& LabelColor,
 		const TAttribute<FText>& ComponentDisplayName,
 		TSharedRef<SHorizontalBox>& HorizontalBox,
 		const TAttribute<TOptional<NumericType>>& Component,
@@ -459,7 +491,7 @@ private:
 		TSharedRef<SWidget> LabelWidget = SNullWidget::NullWidget;
 		if (InArgs._bColorAxisLabels)
 		{
-			LabelWidget = SNumericEntryBox<NumericType>::BuildNarrowColorLabel(LabelColor);
+			LabelWidget = SNumericEntryBox<NumericType>::BuildNarrowColorLabel(LabelColor.Get());
 		}
 
 		TStringBuilder<32> ToolTipTextFormatString;
@@ -517,7 +549,7 @@ private:
 	{
 		ConstructComponent(0,
 			InArgs,
-			AxisDisplayInfo::GetAxisColor(EAxisList::X),
+			InArgs._XColor,
 			InArgs._XDisplayName,
 			HorizontalBox,
 			InArgs._X,
@@ -538,7 +570,7 @@ private:
 	{
 		ConstructComponent(1,
 			InArgs,
-			AxisDisplayInfo::GetAxisColor(EAxisList::Y),
+			InArgs._YColor,
 			InArgs._YDisplayName,
 			HorizontalBox,
 			InArgs._Y,
@@ -559,7 +591,7 @@ private:
 	{
 		ConstructComponent(2,
 			InArgs,
-			AxisDisplayInfo::GetAxisColor(EAxisList::Z),
+			InArgs._ZColor,
 			InArgs._ZDisplayName,
 			HorizontalBox,
 			InArgs._Z,
@@ -580,7 +612,7 @@ private:
 	{
 		ConstructComponent(3,
 			InArgs,
-			SNumericEntryBox<NumericType>::LilacLabelBackgroundColor,
+			InArgs._WColor,
 			InArgs._WDisplayName,
 			HorizontalBox,
 			InArgs._W,
