Only apply swizzling to Transforms rather than all rotators and Vector3s
* Cleaned up the AxisDisplayInfo so that application of swizzle is based on current coordinate system rather than constructed implementation of the EditorAxisDisplayInfo ModularFeature
* Fixed up SRotatorInputBox to not directly use AxisDisplayInfo for swizzle information. Updated customizations to swizzle where required.
* const-ified a getter (despite internal caching, which is really just an implementation detail)

diff --git a/Engine/Source/Editor/DetailCustomizations/Private/ComponentTransformDetails.cpp b/Engine/Source/Editor/DetailCustomizations/Private/ComponentTransformDetails.cpp
--- a/Engine/Source/Editor/DetailCustomizations/Private/ComponentTransformDetails.cpp
+++ b/Engine/Source/Editor/DetailCustomizations/Private/ComponentTransformDetails.cpp
@@ -492,6 +492,7 @@ void FComponentTransformDetails::GenerateChildContent( IDetailChildrenBuilder& C
 			.YDisplayName(AxisDisplayInfo::GetAxisDisplayName(EAxisList::Left))
 			.ZDisplayName(AxisDisplayInfo::GetAxisDisplayName(EAxisList::Up))
 			.bColorAxisLabels(true)
+			.Swizzle(AxisDisplayInfo::GetTransformAxisSwizzle())
 			.IsEnabled(this, &FComponentTransformDetails::GetIsLocationEnabled)
 			.OnXChanged(this, &FComponentTransformDetails::OnSetTransformAxis, ETextCommit::Default, ETransformField::Location, EAxisList::X, false)
 			.OnYChanged(this, &FComponentTransformDetails::OnSetTransformAxis, ETextCommit::Default, ETransformField::Location, EAxisList::Y, false)
@@ -553,7 +554,11 @@ void FComponentTransformDetails::GenerateChildContent( IDetailChildrenBuilder& C
 			.Roll( this, &FComponentTransformDetails::GetRotationX )
 			.Pitch( this, &FComponentTransformDetails::GetRotationY )
 			.Yaw( this, &FComponentTransformDetails::GetRotationZ )
+			.RollDisplayName(AxisDisplayInfo::GetAxisDisplayName(EAxisList::Forward))
+			.PitchDisplayName(AxisDisplayInfo::GetAxisDisplayName(EAxisList::Left))
+			.YawDisplayName(AxisDisplayInfo::GetAxisDisplayName(EAxisList::Up))
 			.bColorAxisLabels( true )
+			.Swizzle(AxisDisplayInfo::GetTransformAxisSwizzle())
 			.IsEnabled( this, &FComponentTransformDetails::GetIsRotationEnabled )
 			.OnPitchBeginSliderMovement( this, &FComponentTransformDetails::OnBeginRotationSlider )
 			.OnYawBeginSliderMovement( this, &FComponentTransformDetails::OnBeginRotationSlider )
@@ -606,6 +611,7 @@ void FComponentTransformDetails::GenerateChildContent( IDetailChildrenBuilder& C
 			.YDisplayName(AxisDisplayInfo::GetAxisDisplayName(EAxisList::Left))
 			.ZDisplayName(AxisDisplayInfo::GetAxisDisplayName(EAxisList::Up))
 			.bColorAxisLabels( true )
+			.Swizzle(AxisDisplayInfo::GetTransformAxisSwizzle())
 			.IsEnabled( this, &FComponentTransformDetails::GetIsScaleEnabled )
 			.OnXChanged( this, &FComponentTransformDetails::OnSetTransformAxis, ETextCommit::Default, ETransformField::Scale, EAxisList::X, false )
 			.OnYChanged( this, &FComponentTransformDetails::OnSetTransformAxis, ETextCommit::Default, ETransformField::Scale, EAxisList::Y, false )
diff --git a/Engine/Source/Editor/DetailCustomizations/Private/MathStructProxyCustomizations.cpp b/Engine/Source/Editor/DetailCustomizations/Private/MathStructProxyCustomizations.cpp
--- a/Engine/Source/Editor/DetailCustomizations/Private/MathStructProxyCustomizations.cpp
+++ b/Engine/Source/Editor/DetailCustomizations/Private/MathStructProxyCustomizations.cpp
@@ -180,24 +180,20 @@ void FMatrixStructCustomization<T>::CustomizeLocation(TSharedRef<class IProperty
 	{
 		Row.OnPasteFromTextDelegate.Pin()->AddSP(this, &FMatrixStructCustomization<T>::OnPasteFromText, FTransformField::Location, WeakHandlePtr);
 	}
-	
-	const EAxisList::Type XAxis = bUseLeftUpForwardAxisDisplayCoordinateSystem ? EAxisList::Forward : EAxisList::X;
-	const EAxisList::Type YAxis = bUseLeftUpForwardAxisDisplayCoordinateSystem ? EAxisList::Left : EAxisList::Y;
-	const EAxisList::Type ZAxis = bUseLeftUpForwardAxisDisplayCoordinateSystem ? EAxisList::Up : EAxisList::Z;
 
 	constexpr int32 NumComponents = 3;
 	TStaticArray<TFunction<TSharedRef<SWidget>()>, NumComponents> ComponentConstructors = {
-		[this, &StructPropertyHandle, XAxis]()->TSharedRef<SWidget>
+		[this, &StructPropertyHandle]()->TSharedRef<SWidget>
 		{
-			return MakeNumericProxyWidget<UE::Math::TVector<T>, T>(StructPropertyHandle, CachedTranslationX, AxisDisplayInfo::GetAxisDisplayName(XAxis), false, AxisDisplayInfo::GetAxisColor(XAxis));
+			return MakeNumericProxyWidget<UE::Math::TVector<T>, T>(StructPropertyHandle, CachedTranslationX, AxisDisplayInfo::GetAxisDisplayName(EAxisList::Forward), false, AxisDisplayInfo::GetAxisColor(EAxisList::Forward));
 		},
-		[this, &StructPropertyHandle, YAxis]()->TSharedRef<SWidget>
+		[this, &StructPropertyHandle]()->TSharedRef<SWidget>
 		{
-			return MakeNumericProxyWidget<UE::Math::TVector<T>, T>(StructPropertyHandle, CachedTranslationY, AxisDisplayInfo::GetAxisDisplayName(YAxis), false, AxisDisplayInfo::GetAxisColor(YAxis));
+			return MakeNumericProxyWidget<UE::Math::TVector<T>, T>(StructPropertyHandle, CachedTranslationY, AxisDisplayInfo::GetAxisDisplayName(EAxisList::Left), false, AxisDisplayInfo::GetAxisColor(EAxisList::Left));
 		},
-		[this, &StructPropertyHandle, ZAxis]()->TSharedRef<SWidget>
+		[this, &StructPropertyHandle]()->TSharedRef<SWidget>
 		{
-			return MakeNumericProxyWidget<UE::Math::TVector<T>, T>(StructPropertyHandle, CachedTranslationZ, AxisDisplayInfo::GetAxisDisplayName(ZAxis), false, AxisDisplayInfo::GetAxisColor(ZAxis));
+			return MakeNumericProxyWidget<UE::Math::TVector<T>, T>(StructPropertyHandle, CachedTranslationZ, AxisDisplayInfo::GetAxisDisplayName(EAxisList::Up), false, AxisDisplayInfo::GetAxisColor(EAxisList::Up));
 		}
 	};
 	const TStaticArray<FMargin, NumComponents> Paddings =
@@ -249,24 +245,20 @@ void FMatrixStructCustomization<T>::CustomizeRotation(TSharedRef<class IProperty
 	{
 		Row.OnPasteFromTextDelegate.Pin()->AddSP(this, &FMatrixStructCustomization<T>::OnPasteFromText, FTransformField::Rotation, WeakHandlePtr);
 	}
-
-	const EAxisList::Type XAxis = bUseLeftUpForwardAxisDisplayCoordinateSystem ? EAxisList::Forward : EAxisList::X;
-	const EAxisList::Type YAxis = bUseLeftUpForwardAxisDisplayCoordinateSystem ? EAxisList::Left : EAxisList::Y;
-	const EAxisList::Type ZAxis = bUseLeftUpForwardAxisDisplayCoordinateSystem ? EAxisList::Up : EAxisList::Z;
 	
 	constexpr int32 NumComponents = 3;
 	TStaticArray<TFunction<TSharedRef<SWidget>()>, NumComponents> ComponentConstructors = {
-		[this, &StructPropertyHandle, XAxis]()->TSharedRef<SWidget>
+		[this, &StructPropertyHandle]()->TSharedRef<SWidget>
 		{
-			return MakeNumericProxyWidget<UE::Math::TRotator<T>, T>(StructPropertyHandle, CachedRotationRoll, AxisDisplayInfo::GetRotationAxisNameShort(XAxis), true, AxisDisplayInfo::GetAxisColor(XAxis));
+			return MakeNumericProxyWidget<UE::Math::TRotator<T>, T>(StructPropertyHandle, CachedRotationRoll, AxisDisplayInfo::GetRotationAxisNameShort(EAxisList::Forward), true, AxisDisplayInfo::GetAxisColor(EAxisList::Forward));
 		},
-		[this, &StructPropertyHandle, YAxis]()->TSharedRef<SWidget>
+		[this, &StructPropertyHandle]()->TSharedRef<SWidget>
 		{
-			return MakeNumericProxyWidget<UE::Math::TRotator<T>, T>(StructPropertyHandle, CachedRotationPitch, AxisDisplayInfo::GetRotationAxisNameShort(YAxis), true, AxisDisplayInfo::GetAxisColor(YAxis));
+			return MakeNumericProxyWidget<UE::Math::TRotator<T>, T>(StructPropertyHandle, CachedRotationPitch, AxisDisplayInfo::GetRotationAxisNameShort(EAxisList::Left), true, AxisDisplayInfo::GetAxisColor(EAxisList::Left));
 		},
-		[this, &StructPropertyHandle, ZAxis]()->TSharedRef<SWidget>
+		[this, &StructPropertyHandle]()->TSharedRef<SWidget>
 		{
-			return MakeNumericProxyWidget<UE::Math::TRotator<T>, T>(StructPropertyHandle, CachedRotationYaw, AxisDisplayInfo::GetRotationAxisNameShort(ZAxis), true, AxisDisplayInfo::GetAxisColor(ZAxis));
+			return MakeNumericProxyWidget<UE::Math::TRotator<T>, T>(StructPropertyHandle, CachedRotationYaw, AxisDisplayInfo::GetRotationAxisNameShort(EAxisList::Up), true, AxisDisplayInfo::GetAxisColor(EAxisList::Up));
 		}
 	};
 	const TStaticArray<FMargin, NumComponents> Paddings =
@@ -319,23 +311,19 @@ void FMatrixStructCustomization<T>::CustomizeScale(TSharedRef<class IPropertyHan
 		Row.OnPasteFromTextDelegate.Pin()->AddSP(this, &FMatrixStructCustomization<T>::OnPasteFromText, FTransformField::Scale, WeakHandlePtr);
 	}
 
-	const EAxisList::Type XAxis = bUseLeftUpForwardAxisDisplayCoordinateSystem ? EAxisList::Forward : EAxisList::X;
-	const EAxisList::Type YAxis = bUseLeftUpForwardAxisDisplayCoordinateSystem ? EAxisList::Left : EAxisList::Y;
-	const EAxisList::Type ZAxis = bUseLeftUpForwardAxisDisplayCoordinateSystem ? EAxisList::Up : EAxisList::Z;
-
 	constexpr int32 NumComponents = 3;
 	TStaticArray<TFunction<TSharedRef<SWidget>()>, NumComponents> ComponentConstructors = {
-		[this, &StructPropertyHandle, XAxis]()->TSharedRef<SWidget>
+		[this, &StructPropertyHandle]()->TSharedRef<SWidget>
 		{
-			return MakeNumericProxyWidget<UE::Math::TVector<T>, T>(StructPropertyHandle, CachedScaleX, AxisDisplayInfo::GetAxisDisplayName(XAxis), false, AxisDisplayInfo::GetAxisColor(XAxis));
+			return MakeNumericProxyWidget<UE::Math::TVector<T>, T>(StructPropertyHandle, CachedScaleX, AxisDisplayInfo::GetAxisDisplayName(EAxisList::Forward), false, AxisDisplayInfo::GetAxisColor(EAxisList::Forward));
 		},
-		[this, &StructPropertyHandle, YAxis]()->TSharedRef<SWidget>
+		[this, &StructPropertyHandle]()->TSharedRef<SWidget>
 		{
-			return MakeNumericProxyWidget<UE::Math::TVector<T>, T>(StructPropertyHandle, CachedScaleY, AxisDisplayInfo::GetAxisDisplayName(YAxis), false, AxisDisplayInfo::GetAxisColor(YAxis));
+			return MakeNumericProxyWidget<UE::Math::TVector<T>, T>(StructPropertyHandle, CachedScaleY, AxisDisplayInfo::GetAxisDisplayName(EAxisList::Left), false, AxisDisplayInfo::GetAxisColor(EAxisList::Left));
 		},
-		[this, &StructPropertyHandle, ZAxis]()->TSharedRef<SWidget>
+		[this, &StructPropertyHandle]()->TSharedRef<SWidget>
 		{
-			return MakeNumericProxyWidget<UE::Math::TVector<T>, T>(StructPropertyHandle, CachedScaleZ, AxisDisplayInfo::GetAxisDisplayName(ZAxis), false, AxisDisplayInfo::GetAxisColor(ZAxis));
+			return MakeNumericProxyWidget<UE::Math::TVector<T>, T>(StructPropertyHandle, CachedScaleZ, AxisDisplayInfo::GetAxisDisplayName(EAxisList::Up), false, AxisDisplayInfo::GetAxisColor(EAxisList::Up));
 		}
 	};
 	const TStaticArray<FMargin, NumComponents> Paddings =
@@ -857,7 +845,7 @@ bool FTransformStructCustomization<T>::FlushValues( TWeakPtr<IPropertyHandle> Pr
 template <typename T>
 FIntVector4 FTransformStructCustomization<T>::GetSwizzle() const
 {
-	return AxisDisplayInfo::DefaultAxisComponentDisplaySwizzle();
+	return AxisDisplayInfo::GetTransformAxisSwizzle();
 }
 
 template<typename T>
diff --git a/Engine/Source/Editor/LevelEditor/Private/DlgDeltaTransform.cpp b/Engine/Source/Editor/LevelEditor/Private/DlgDeltaTransform.cpp
--- a/Engine/Source/Editor/LevelEditor/Private/DlgDeltaTransform.cpp
+++ b/Engine/Source/Editor/LevelEditor/Private/DlgDeltaTransform.cpp
@@ -78,6 +78,7 @@ public:
 							.YDisplayName(AxisDisplayInfo::GetAxisDisplayName(EAxisList::Left))
 							.ZDisplayName(AxisDisplayInfo::GetAxisDisplayName(EAxisList::Up))
 							.bColorAxisLabels( true )
+							.Swizzle(AxisDisplayInfo::GetTransformAxisSwizzle())
 							.AllowSpin( false )
 							.OnXCommitted( this, &SDlgDeltaTransform::OnSetDelta, 0 )
 							.OnYCommitted( this, &SDlgDeltaTransform::OnSetDelta, 1 )
diff --git a/Engine/Source/Editor/UnrealEd/Private/EditorAxisDisplayInfo.cpp b/Engine/Source/Editor/UnrealEd/Private/EditorAxisDisplayInfo.cpp
--- a/Engine/Source/Editor/UnrealEd/Private/EditorAxisDisplayInfo.cpp
+++ b/Engine/Source/Editor/UnrealEd/Private/EditorAxisDisplayInfo.cpp
@@ -24,7 +24,7 @@ FEditorAxisDisplayInfo::FEditorAxisDisplayInfo()
 	FEditorDelegates::OnEditorBoot.AddRaw(this, &FEditorAxisDisplayInfo::InitSettingsInfo);
 }
 
-EAxisList::Type FEditorAxisDisplayInfo::GetAxisDisplayCoordinateSystem()
+EAxisList::Type FEditorAxisDisplayInfo::GetAxisDisplayCoordinateSystem() const
 {
 	if (!AxisDisplayCoordinateSystem.IsSet())
 	{
@@ -199,6 +199,12 @@ FLinearColor FEditorAxisDisplayInfo::GetAxisColor(EAxisList::Type Axis)
 
 FIntVector4 FEditorAxisDisplayInfo::DefaultAxisComponentDisplaySwizzle() const
 {
+	if (GetAxisDisplayCoordinateSystem() == EAxisList::LeftUpForward)
+	{
+		// Unreal:    -Y    Z     X
+		// Semantic: Left, Up, Forward
+		return FIntVector4(1, 2, 0, 3);
+	}
 	return FIntVector4(0, 1, 2, 3);
 }
 
@@ -220,17 +226,31 @@ FText FEditorAxisDisplayInfo::GetRotationAxisName(EAxisList::Type Axis)
 
 FText FEditorAxisDisplayInfo::GetRotationAxisNameShort(EAxisList::Type Axis)
 {
-	if (Axis == EAxisList::X || Axis == EAxisList::Forward)
+	Axis = MapAxis(Axis);
+	
+	if (Axis == EAxisList::X)
 	{
-		return LOCTEXT("Roll", "Roll");
+		return LOCTEXT("Roll_ToolTipTextFormat", "Roll");
 	}
-	else if (Axis == EAxisList::Y || Axis == EAxisList::Left)
+	else if (Axis == EAxisList::Y)
 	{
-		return LOCTEXT("Pitch", "Pitch");
+		return LOCTEXT("Pitch_ToolTipTextFormat", "Pitch");
 	}
-	else if (Axis == EAxisList::Z || Axis == EAxisList::Up)
+	else if (Axis == EAxisList::Z)
+	{
+		return LOCTEXT("Yaw_ToolTipTextFormat", "Yaw");
+	}
+	else if (Axis == EAxisList::Forward)
+	{
+		return LOCTEXT("Forward_ToolTipTextFormat", "Forward");
+	}
+	else if (Axis == EAxisList::Left)
+	{
+		return LOCTEXT("Left_ToolTipTextFormat", "Left");
+	}
+	else if (Axis == EAxisList::Up)
 	{
-		return LOCTEXT("Yaw", "Yaw");
+		return LOCTEXT("Up_ToolTipTextFormat", "Up");
 	}
 
 	ensureMsgf(false, TEXT("Unsupported Axis: %d"), Axis);
diff --git a/Engine/Source/Editor/UnrealEd/Public/EditorAxisDisplayInfo.h b/Engine/Source/Editor/UnrealEd/Public/EditorAxisDisplayInfo.h
--- a/Engine/Source/Editor/UnrealEd/Public/EditorAxisDisplayInfo.h
+++ b/Engine/Source/Editor/UnrealEd/Public/EditorAxisDisplayInfo.h
@@ -11,7 +11,7 @@ public:
 	FEditorAxisDisplayInfo();
 	virtual ~FEditorAxisDisplayInfo() = default;
 
-	virtual EAxisList::Type GetAxisDisplayCoordinateSystem() override;
+	virtual EAxisList::Type GetAxisDisplayCoordinateSystem() const override;
 
 	virtual FText GetAxisDisplayName(EAxisList::Type Axis) override;
 
@@ -34,6 +34,6 @@ private:
 	void InitSettingsInfo(double EditorStartupTime);
 
 	TOptional<bool> bUseForwardRightUpDisplayNames;
-	TOptional<EAxisList::Type> AxisDisplayCoordinateSystem;
+	mutable TOptional<EAxisList::Type> AxisDisplayCoordinateSystem;
 
 };
diff --git a/Engine/Source/Runtime/Core/Private/Misc/AxisDisplayInfo.cpp b/Engine/Source/Runtime/Core/Private/Misc/AxisDisplayInfo.cpp
--- a/Engine/Source/Runtime/Core/Private/Misc/AxisDisplayInfo.cpp
+++ b/Engine/Source/Runtime/Core/Private/Misc/AxisDisplayInfo.cpp
@@ -158,7 +158,7 @@ FText AxisDisplayInfo::GetRotationAxisNameShort(EAxisList::Type Axis)
 	return LOCTEXT("UnsupportedRotationAxisNameShort", "?");
 }
 
-FIntVector4 AxisDisplayInfo::DefaultAxisComponentDisplaySwizzle()
+FIntVector4 AxisDisplayInfo::GetTransformAxisSwizzle()
 {
 	if (IModularFeatures::Get().IsModularFeatureAvailable(IAxisDisplayInfo::GetModularFeatureName()))
 	{
diff --git a/Engine/Source/Runtime/Core/Public/Misc/AxisDisplayInfo.h b/Engine/Source/Runtime/Core/Public/Misc/AxisDisplayInfo.h
--- a/Engine/Source/Runtime/Core/Public/Misc/AxisDisplayInfo.h
+++ b/Engine/Source/Runtime/Core/Public/Misc/AxisDisplayInfo.h
@@ -66,9 +66,11 @@ namespace AxisDisplayInfo
 	
 
 	/**
-	 * Default swizzle for axis displays in areas like the details panel
+	 * Default swizzle for displaying axis in transform displays
+	 * Swizzle maps from the default X, Y, Z, W positions to new position
+	 * { 1, 2, 0, 3 } would map position 0 to fetch Y, position 1 to fetch Z, etc
 	 */
-	CORE_API FIntVector4 DefaultAxisComponentDisplaySwizzle();
+	CORE_API FIntVector4 GetTransformAxisSwizzle();
 };
 
 class IAxisDisplayInfo : public IModularFeature
@@ -90,7 +92,7 @@ public:
 	 *
 	 * @return The display coordinate system, either EAxisList::XYZ or EAxisList::LeftUpForward
 	 */
-	virtual EAxisList::Type GetAxisDisplayCoordinateSystem() = 0;
+	virtual EAxisList::Type GetAxisDisplayCoordinateSystem() const = 0;
 
 	/**
 	 * Gets the display name to use for the given axis.
diff --git a/Engine/Source/Runtime/Slate/Public/Widgets/Input/SRotatorInputBox.h b/Engine/Source/Runtime/Slate/Public/Widgets/Input/SRotatorInputBox.h
--- a/Engine/Source/Runtime/Slate/Public/Widgets/Input/SRotatorInputBox.h
+++ b/Engine/Source/Runtime/Slate/Public/Widgets/Input/SRotatorInputBox.h
@@ -36,7 +36,7 @@ public:
 		, _PitchDisplayName(AxisDisplayInfo::GetRotationAxisNameShort(EAxisList::Y))
 		, _YawDisplayName(AxisDisplayInfo::GetRotationAxisNameShort(EAxisList::Z))
 		, _bColorAxisLabels(false)
-		, _Swizzle(AxisDisplayInfo::DefaultAxisComponentDisplaySwizzle())
+		, _Swizzle(FIntVector4(0, 1, 2, 3))
 		, _Font(FAppStyle::Get().GetFontStyle("NormalFont"))
 		, _AllowSpin(true)
 		, _SpinDelta(1.f)
diff --git a/Engine/Source/Runtime/Slate/Public/Widgets/Input/SVectorInputBox.h b/Engine/Source/Runtime/Slate/Public/Widgets/Input/SVectorInputBox.h
--- a/Engine/Source/Runtime/Slate/Public/Widgets/Input/SVectorInputBox.h
+++ b/Engine/Source/Runtime/Slate/Public/Widgets/Input/SVectorInputBox.h
@@ -259,19 +259,11 @@ public:
 			, _SpinDelta(1)
 			, _LinearDeltaSensitivity(1)
 			, _bColorAxisLabels(false)
+			, _Swizzle(FIntVector4(0, 1, 2, 3))
 			, _DisplayToggle(false)
 			, _TogglePadding(FMargin(1.f,0.f,1.f,0.f) )
 			, _PreventThrottling(false)
-		{
-			if constexpr (NumberOfComponents == 3)
-			{
-				_Swizzle = AxisDisplayInfo::DefaultAxisComponentDisplaySwizzle();
-			}
-			else
-			{
-				_Swizzle = FIntVector4(0, 1, 2, 3);
-			}
-		}
+		{}
 
 		/** Optional Value of the vector */
 		SLATE_ATTRIBUTE(TOptional<VectorType>, Vector)
