FComponentTransformDetails : Set PreventThrottling to true on vector and rotator input boxes.

diff --git a/Engine/Source/Editor/DetailCustomizations/Private/ComponentTransformDetails.cpp b/Engine/Source/Editor/DetailCustomizations/Private/ComponentTransformDetails.cpp
--- a/Engine/Source/Editor/DetailCustomizations/Private/ComponentTransformDetails.cpp
+++ b/Engine/Source/Editor/DetailCustomizations/Private/ComponentTransformDetails.cpp
@@ -506,6 +506,7 @@ void FComponentTransformDetails::GenerateChildContent( IDetailChildrenBuilder& C
 			.SpinDelta(LocationSpinDelta)
 			.OnBeginSliderMovement(this, &FComponentTransformDetails::OnBeginLocationSlider)
 			.OnEndSliderMovement(this, &FComponentTransformDetails::OnEndLocationSlider)
+			.PreventThrottling(true)
 		];
 	}
 	
@@ -574,6 +575,7 @@ void FComponentTransformDetails::GenerateChildContent( IDetailChildrenBuilder& C
 			.OnYawCommitted( this, &FComponentTransformDetails::OnSetTransformAxis, ETransformField::Rotation, EAxisList::Z, true )
 			.TypeInterface( TypeInterface )
 			.Font( FontInfo )
+			.PreventThrottling(true)
 		];
 	}
 	
@@ -627,6 +629,7 @@ void FComponentTransformDetails::GenerateChildContent( IDetailChildrenBuilder& C
 			.SpinDelta( ScaleSpinDelta )
 			.OnBeginSliderMovement( this, &FComponentTransformDetails::OnBeginScaleSlider )
 			.OnEndSliderMovement(this, &FComponentTransformDetails::OnEndScaleSlider)
+			.PreventThrottling(true)
 		];
 	}
 }
