UE: Use same guid settings when determining whether the update button is enabled as used when it is pressed

diff --git a/Engine/Source/Editor/DetailCustomizations/Private/PoseAssetDetails.cpp b/Engine/Source/Editor/DetailCustomizations/Private/PoseAssetDetails.cpp
--- a/Engine/Source/Editor/DetailCustomizations/Private/PoseAssetDetails.cpp
+++ b/Engine/Source/Editor/DetailCustomizations/Private/PoseAssetDetails.cpp
@@ -703,7 +703,9 @@ bool FPoseAssetDetails::IsUpdateSourceEnabled() const
 	const UAnimSequence* AnimSequenceSelected = Cast<UAnimSequence>(ObjectSet);
 	if (AnimSequenceSelected && PoseAsset.IsValid())
 	{
-		return (PoseAsset->SourceAnimation && PoseAsset->SourceAnimation->GetDataModel()->GenerateGuid() != PoseAsset->SourceAnimationRawDataGUID);		
+		IAnimationDataModel::FGuidGenerationSettings Settings;
+		Settings.bIncludeTimingData = false;
+		return (PoseAsset->SourceAnimation && PoseAsset->SourceAnimation->GetDataModel()->GenerateGuid(Settings) != PoseAsset->SourceAnimationRawDataGUID);
 	}
 
 	return false;
