Add AxisDisplayInfo:: convenience functions for getting axis tool tip text
* Apply these to transform customizations

diff --git a/Engine/Source/Editor/DetailCustomizations/Private/MathStructProxyCustomizations.cpp b/Engine/Source/Editor/DetailCustomizations/Private/MathStructProxyCustomizations.cpp
--- a/Engine/Source/Editor/DetailCustomizations/Private/MathStructProxyCustomizations.cpp
+++ b/Engine/Source/Editor/DetailCustomizations/Private/MathStructProxyCustomizations.cpp
@@ -185,15 +185,15 @@ void FMatrixStructCustomization<T>::CustomizeLocation(TSharedRef<class IProperty
 	TStaticArray<TFunction<TSharedRef<SWidget>()>, NumComponents> ComponentConstructors = {
 		[this, &StructPropertyHandle]()->TSharedRef<SWidget>
 		{
-			return MakeNumericProxyWidget<UE::Math::TVector<T>, T>(StructPropertyHandle, CachedTranslationX, AxisDisplayInfo::GetAxisDisplayName(EAxisList::Forward), false, AxisDisplayInfo::GetAxisColor(EAxisList::Forward));
+			return MakeNumericProxyWidget<UE::Math::TVector<T>, T>(StructPropertyHandle, CachedTranslationX, AxisDisplayInfo::GetAxisToolTip(EAxisList::Forward), false, AxisDisplayInfo::GetAxisColor(EAxisList::Forward));
 		},
 		[this, &StructPropertyHandle]()->TSharedRef<SWidget>
 		{
-			return MakeNumericProxyWidget<UE::Math::TVector<T>, T>(StructPropertyHandle, CachedTranslationY, AxisDisplayInfo::GetAxisDisplayName(EAxisList::Left), false, AxisDisplayInfo::GetAxisColor(EAxisList::Left));
+			return MakeNumericProxyWidget<UE::Math::TVector<T>, T>(StructPropertyHandle, CachedTranslationY, AxisDisplayInfo::GetAxisToolTip(EAxisList::Left), false, AxisDisplayInfo::GetAxisColor(EAxisList::Left));
 		},
 		[this, &StructPropertyHandle]()->TSharedRef<SWidget>
 		{
-			return MakeNumericProxyWidget<UE::Math::TVector<T>, T>(StructPropertyHandle, CachedTranslationZ, AxisDisplayInfo::GetAxisDisplayName(EAxisList::Up), false, AxisDisplayInfo::GetAxisColor(EAxisList::Up));
+			return MakeNumericProxyWidget<UE::Math::TVector<T>, T>(StructPropertyHandle, CachedTranslationZ, AxisDisplayInfo::GetAxisToolTip(EAxisList::Up), false, AxisDisplayInfo::GetAxisColor(EAxisList::Up));
 		}
 	};
 	const TStaticArray<FMargin, NumComponents> Paddings =
@@ -250,15 +250,15 @@ void FMatrixStructCustomization<T>::CustomizeRotation(TSharedRef<class IProperty
 	TStaticArray<TFunction<TSharedRef<SWidget>()>, NumComponents> ComponentConstructors = {
 		[this, &StructPropertyHandle]()->TSharedRef<SWidget>
 		{
-			return MakeNumericProxyWidget<UE::Math::TRotator<T>, T>(StructPropertyHandle, CachedRotationRoll, AxisDisplayInfo::GetRotationAxisNameShort(EAxisList::Forward), true, AxisDisplayInfo::GetAxisColor(EAxisList::Forward));
+			return MakeNumericProxyWidget<UE::Math::TRotator<T>, T>(StructPropertyHandle, CachedRotationRoll, AxisDisplayInfo::GetRotationAxisToolTip(EAxisList::Forward), true, AxisDisplayInfo::GetAxisColor(EAxisList::Forward));
 		},
 		[this, &StructPropertyHandle]()->TSharedRef<SWidget>
 		{
-			return MakeNumericProxyWidget<UE::Math::TRotator<T>, T>(StructPropertyHandle, CachedRotationPitch, AxisDisplayInfo::GetRotationAxisNameShort(EAxisList::Left), true, AxisDisplayInfo::GetAxisColor(EAxisList::Left));
+			return MakeNumericProxyWidget<UE::Math::TRotator<T>, T>(StructPropertyHandle, CachedRotationPitch, AxisDisplayInfo::GetRotationAxisToolTip(EAxisList::Left), true, AxisDisplayInfo::GetAxisColor(EAxisList::Left));
 		},
 		[this, &StructPropertyHandle]()->TSharedRef<SWidget>
 		{
-			return MakeNumericProxyWidget<UE::Math::TRotator<T>, T>(StructPropertyHandle, CachedRotationYaw, AxisDisplayInfo::GetRotationAxisNameShort(EAxisList::Up), true, AxisDisplayInfo::GetAxisColor(EAxisList::Up));
+			return MakeNumericProxyWidget<UE::Math::TRotator<T>, T>(StructPropertyHandle, CachedRotationYaw, AxisDisplayInfo::GetRotationAxisToolTip(EAxisList::Up), true, AxisDisplayInfo::GetAxisColor(EAxisList::Up));
 		}
 	};
 	const TStaticArray<FMargin, NumComponents> Paddings =
@@ -315,15 +315,15 @@ void FMatrixStructCustomization<T>::CustomizeScale(TSharedRef<class IPropertyHan
 	TStaticArray<TFunction<TSharedRef<SWidget>()>, NumComponents> ComponentConstructors = {
 		[this, &StructPropertyHandle]()->TSharedRef<SWidget>
 		{
-			return MakeNumericProxyWidget<UE::Math::TVector<T>, T>(StructPropertyHandle, CachedScaleX, AxisDisplayInfo::GetAxisDisplayName(EAxisList::Forward), false, AxisDisplayInfo::GetAxisColor(EAxisList::Forward));
+			return MakeNumericProxyWidget<UE::Math::TVector<T>, T>(StructPropertyHandle, CachedScaleX, AxisDisplayInfo::GetAxisToolTip(EAxisList::Forward), false, AxisDisplayInfo::GetAxisColor(EAxisList::Forward));
 		},
 		[this, &StructPropertyHandle]()->TSharedRef<SWidget>
 		{
-			return MakeNumericProxyWidget<UE::Math::TVector<T>, T>(StructPropertyHandle, CachedScaleY, AxisDisplayInfo::GetAxisDisplayName(EAxisList::Left), false, AxisDisplayInfo::GetAxisColor(EAxisList::Left));
+			return MakeNumericProxyWidget<UE::Math::TVector<T>, T>(StructPropertyHandle, CachedScaleY, AxisDisplayInfo::GetAxisToolTip(EAxisList::Left), false, AxisDisplayInfo::GetAxisColor(EAxisList::Left));
 		},
 		[this, &StructPropertyHandle]()->TSharedRef<SWidget>
 		{
-			return MakeNumericProxyWidget<UE::Math::TVector<T>, T>(StructPropertyHandle, CachedScaleZ, AxisDisplayInfo::GetAxisDisplayName(EAxisList::Up), false, AxisDisplayInfo::GetAxisColor(EAxisList::Up));
+			return MakeNumericProxyWidget<UE::Math::TVector<T>, T>(StructPropertyHandle, CachedScaleZ, AxisDisplayInfo::GetAxisToolTip(EAxisList::Up), false, AxisDisplayInfo::GetAxisColor(EAxisList::Up));
 		}
 	};
 	const TStaticArray<FMargin, NumComponents> Paddings =
diff --git a/Engine/Source/Runtime/Core/Private/Misc/AxisDisplayInfo.cpp b/Engine/Source/Runtime/Core/Private/Misc/AxisDisplayInfo.cpp
--- a/Engine/Source/Runtime/Core/Private/Misc/AxisDisplayInfo.cpp
+++ b/Engine/Source/Runtime/Core/Private/Misc/AxisDisplayInfo.cpp
@@ -83,6 +83,16 @@ FText AxisDisplayInfo::GetAxisDisplayNameShort(EAxisList::Type Axis)
 	return LOCTEXT("UnsupportedGetAxisDisplayNameShort", "?");
 }
 
+FText AxisDisplayInfo::GetAxisToolTip(EAxisList::Type Axis)
+{
+	if (IModularFeatures::Get().IsModularFeatureAvailable(IAxisDisplayInfo::GetModularFeatureName()))
+	{
+		IAxisDisplayInfo& DisplayInfo = IModularFeatures::Get().GetModularFeature<IAxisDisplayInfo>(IAxisDisplayInfo::GetModularFeatureName());
+		return DisplayInfo.GetAxisToolTip(Axis);
+	}
+	return LOCTEXT("UnsupportedDisplayName", "Unsupported");
+}
+
 FLinearColor AxisDisplayInfo::GetAxisColor(EAxisList::Type Axis)
 {
 	if (IModularFeatures::Get().IsModularFeatureAvailable(IAxisDisplayInfo::GetModularFeatureName()))
@@ -158,6 +168,16 @@ FText AxisDisplayInfo::GetRotationAxisNameShort(EAxisList::Type Axis)
 	return LOCTEXT("UnsupportedRotationAxisNameShort", "?");
 }
 
+FText AxisDisplayInfo::GetRotationAxisToolTip(EAxisList::Type Axis)
+{
+	if (IModularFeatures::Get().IsModularFeatureAvailable(IAxisDisplayInfo::GetModularFeatureName()))
+	{
+		IAxisDisplayInfo& DisplayInfo = IModularFeatures::Get().GetModularFeature<IAxisDisplayInfo>(IAxisDisplayInfo::GetModularFeatureName());
+		return DisplayInfo.GetRotationAxisNameShort(Axis);
+	}
+	return LOCTEXT("UnsupportedRotationAxisToolTip", "");
+}
+
 FIntVector4 AxisDisplayInfo::GetTransformAxisSwizzle()
 {
 	if (IModularFeatures::Get().IsModularFeatureAvailable(IAxisDisplayInfo::GetModularFeatureName()))
diff --git a/Engine/Source/Runtime/Core/Public/Misc/AxisDisplayInfo.h b/Engine/Source/Runtime/Core/Public/Misc/AxisDisplayInfo.h
--- a/Engine/Source/Runtime/Core/Public/Misc/AxisDisplayInfo.h
+++ b/Engine/Source/Runtime/Core/Public/Misc/AxisDisplayInfo.h
@@ -42,6 +42,8 @@ namespace AxisDisplayInfo
 	 */
 	CORE_API FText GetAxisDisplayNameShort(EAxisList::Type Axis);
 
+	CORE_API FText GetAxisToolTip(EAxisList::Type Axis);
+
 	/**
 	 * Gets the color to use for the given axis.
 	 * 
@@ -63,6 +65,7 @@ namespace AxisDisplayInfo
 	
 	CORE_API FText GetRotationAxisName(EAxisList::Type Axis);
 	CORE_API FText GetRotationAxisNameShort(EAxisList::Type Axis);
+	CORE_API FText GetRotationAxisToolTip(EAxisList::Type Axis);
 	
 
 	/**
