Fix several incorrect tool tips for LUF
* Fallback to legacy GetPropertyDisplayName if tool tip text is empty
* Fix typo in AxisDisplayInfo forwarding GetRotationAxisToolTip to incorrect function
* Update EditorAxisDisplay::GetRotationAxisToolTip with hint about previous X,Y,Z axis mapping
* MathStructCustomizations to preferentially use GetToolTipText() to generate tool tips
* ActorComponent tranform customization to display (was X,-Y,Z) when in LUF mode
* verse::rotation and verse::transform rotation components will contain updated hint

diff --git a/Engine/Source/Editor/DetailCustomizations/Private/ComponentTransformDetails.cpp b/Engine/Source/Editor/DetailCustomizations/Private/ComponentTransformDetails.cpp
--- a/Engine/Source/Editor/DetailCustomizations/Private/ComponentTransformDetails.cpp
+++ b/Engine/Source/Editor/DetailCustomizations/Private/ComponentTransformDetails.cpp
@@ -488,9 +488,9 @@ void FComponentTransformDetails::GenerateChildContent( IDetailChildrenBuilder& C
 			.X(this, &FComponentTransformDetails::GetLocationX)
 			.Y(this, &FComponentTransformDetails::GetLocationY)
 			.Z(this, &FComponentTransformDetails::GetLocationZ)
-			.XDisplayName(AxisDisplayInfo::GetAxisDisplayName(EAxisList::Forward))
-			.YDisplayName(AxisDisplayInfo::GetAxisDisplayName(EAxisList::Left))
-			.ZDisplayName(AxisDisplayInfo::GetAxisDisplayName(EAxisList::Up))
+			.XDisplayName(AxisDisplayInfo::GetAxisToolTip(EAxisList::Forward))
+			.YDisplayName(AxisDisplayInfo::GetAxisToolTip(EAxisList::Left))
+			.ZDisplayName(AxisDisplayInfo::GetAxisToolTip(EAxisList::Up))
 			.bColorAxisLabels(true)
 			.Swizzle(AxisDisplayInfo::GetTransformAxisSwizzle())
 			.IsEnabled(this, &FComponentTransformDetails::GetIsLocationEnabled)
@@ -555,9 +555,9 @@ void FComponentTransformDetails::GenerateChildContent( IDetailChildrenBuilder& C
 			.Roll( this, &FComponentTransformDetails::GetRotationX )
 			.Pitch( this, &FComponentTransformDetails::GetRotationY )
 			.Yaw( this, &FComponentTransformDetails::GetRotationZ )
-			.RollDisplayName(AxisDisplayInfo::GetAxisDisplayName(EAxisList::Forward))
-			.PitchDisplayName(AxisDisplayInfo::GetAxisDisplayName(EAxisList::Left))
-			.YawDisplayName(AxisDisplayInfo::GetAxisDisplayName(EAxisList::Up))
+			.RollDisplayName(AxisDisplayInfo::GetRotationAxisToolTip(EAxisList::Forward))
+			.PitchDisplayName(AxisDisplayInfo::GetRotationAxisToolTip(EAxisList::Left))
+			.YawDisplayName(AxisDisplayInfo::GetRotationAxisToolTip(EAxisList::Up))
 			.bColorAxisLabels( true )
 			.Swizzle(AxisDisplayInfo::GetTransformAxisSwizzle())
 			.IsEnabled( this, &FComponentTransformDetails::GetIsRotationEnabled )
@@ -609,9 +609,9 @@ void FComponentTransformDetails::GenerateChildContent( IDetailChildrenBuilder& C
 			.X( this, &FComponentTransformDetails::GetScaleX )
 			.Y( this, &FComponentTransformDetails::GetScaleY )
 			.Z( this, &FComponentTransformDetails::GetScaleZ )
-			.XDisplayName(AxisDisplayInfo::GetAxisDisplayName(EAxisList::Forward))
-			.YDisplayName(AxisDisplayInfo::GetAxisDisplayName(EAxisList::Left))
-			.ZDisplayName(AxisDisplayInfo::GetAxisDisplayName(EAxisList::Up))
+			.XDisplayName(AxisDisplayInfo::GetAxisToolTip(EAxisList::Forward))
+			.YDisplayName(AxisDisplayInfo::GetAxisToolTip(EAxisList::Left))
+			.ZDisplayName(AxisDisplayInfo::GetAxisToolTip(EAxisList::Up))
 			.bColorAxisLabels( true )
 			.Swizzle(AxisDisplayInfo::GetTransformAxisSwizzle())
 			.IsEnabled( this, &FComponentTransformDetails::GetIsScaleEnabled )
diff --git a/Engine/Source/Editor/DetailCustomizations/Private/MathStructCustomizations.cpp b/Engine/Source/Editor/DetailCustomizations/Private/MathStructCustomizations.cpp
--- a/Engine/Source/Editor/DetailCustomizations/Private/MathStructCustomizations.cpp
+++ b/Engine/Source/Editor/DetailCustomizations/Private/MathStructCustomizations.cpp
@@ -702,7 +702,12 @@ TOptional<FTextFormat> FMathStructCustomization::OnGetValueToolTipTextFormat(TWe
 	if (TSharedPtr<IPropertyHandle> PropertyHandle = WeakHandlePtr.Pin())
 	{
 		TStringBuilder<32> ToolTipTextFormatString;
-		ToolTipTextFormatString.Append(PropertyHandle->GetPropertyDisplayName().ToString());
+		FText TooltipText = PropertyHandle->GetToolTipText();
+		if (TooltipText.IsEmpty())
+		{
+			TooltipText = PropertyHandle->GetPropertyDisplayName();
+		}
+		ToolTipTextFormatString.Append(TooltipText.ToString());
 		ToolTipTextFormatString.Append(TEXT(": {0}"));
 
 		return TOptional<FTextFormat>(FText::FromString(ToolTipTextFormatString.ToString()));
diff --git a/Engine/Source/Editor/LevelEditor/Private/DlgDeltaTransform.cpp b/Engine/Source/Editor/LevelEditor/Private/DlgDeltaTransform.cpp
--- a/Engine/Source/Editor/LevelEditor/Private/DlgDeltaTransform.cpp
+++ b/Engine/Source/Editor/LevelEditor/Private/DlgDeltaTransform.cpp
@@ -74,9 +74,9 @@ public:
 							.X( this, &SDlgDeltaTransform::GetDeltaX )
 							.Y( this, &SDlgDeltaTransform::GetDeltaY )
 							.Z( this, &SDlgDeltaTransform::GetDeltaZ )
-							.XDisplayName(AxisDisplayInfo::GetAxisDisplayName(EAxisList::Forward))
-							.YDisplayName(AxisDisplayInfo::GetAxisDisplayName(EAxisList::Left))
-							.ZDisplayName(AxisDisplayInfo::GetAxisDisplayName(EAxisList::Up))
+							.XDisplayName(AxisDisplayInfo::GetAxisToolTip(EAxisList::Forward))
+							.YDisplayName(AxisDisplayInfo::GetAxisToolTip(EAxisList::Left))
+							.ZDisplayName(AxisDisplayInfo::GetAxisToolTip(EAxisList::Up))
 							.bColorAxisLabels( true )
 							.Swizzle(AxisDisplayInfo::GetTransformAxisSwizzle())
 							.AllowSpin( false )
diff --git a/Engine/Source/Editor/UnrealEd/Private/EditorAxisDisplayInfo.cpp b/Engine/Source/Editor/UnrealEd/Private/EditorAxisDisplayInfo.cpp
--- a/Engine/Source/Editor/UnrealEd/Private/EditorAxisDisplayInfo.cpp
+++ b/Engine/Source/Editor/UnrealEd/Private/EditorAxisDisplayInfo.cpp
@@ -219,7 +219,36 @@ bool FEditorAxisDisplayInfo::UseForwardRightUpDisplayNames()
 
 FText FEditorAxisDisplayInfo::GetRotationAxisToolTip(EAxisList::Type Axis) const
 {
-	return const_cast<FEditorAxisDisplayInfo*>(this)->GetRotationAxisNameShort(Axis);
+	Axis = MapAxis(Axis);
+	
+	if (Axis == EAxisList::X)
+	{
+		return LOCTEXT("GetRotationAxisToolTip_Roll", "Roll");
+	}
+	else if (Axis == EAxisList::Y)
+	{
+		return LOCTEXT("GetRotationAxisToolTip_Pitch", "Pitch");
+	}
+	else if (Axis == EAxisList::Z)
+	{
+		return LOCTEXT("GetRotationAxisToolTip_Yaw", "Yaw");
+	}
+	else if (Axis == EAxisList::Forward)
+	{
+		return LOCTEXT("GetRotationAxisToolTip_Forward", "Forward (was X)");
+	}
+	else if (Axis == EAxisList::Left)
+	{
+		return LOCTEXT("GetRotationAxisToolTip_Left", "Left (was -Y)");
+	}
+	else if (Axis == EAxisList::Up)
+	{
+		return LOCTEXT("GetRotationAxisToolTip_Up", "Up (was Z)");
+	}
+
+	ensureMsgf(false, TEXT("Unsupported Axis: %d"), Axis);
+
+	return LOCTEXT("GetRotationAxisToolTip_Unsupported", "?");
 }
 
 FText FEditorAxisDisplayInfo::GetRotationAxisName(EAxisList::Type Axis)
@@ -233,32 +262,32 @@ FText FEditorAxisDisplayInfo::GetRotationAxisNameShort(EAxisList::Type Axis)
 	
 	if (Axis == EAxisList::X)
 	{
-		return LOCTEXT("Roll_ToolTipTextFormat", "Roll");
+		return LOCTEXT("GetRotationAxisNameShort_X", "Roll");
 	}
 	else if (Axis == EAxisList::Y)
 	{
-		return LOCTEXT("Pitch_ToolTipTextFormat", "Pitch");
+		return LOCTEXT("GetRotationAxisNameShort_Y", "Pitch");
 	}
 	else if (Axis == EAxisList::Z)
 	{
-		return LOCTEXT("Yaw_ToolTipTextFormat", "Yaw");
+		return LOCTEXT("GetRotationAxisNameShort_Z", "Yaw");
 	}
 	else if (Axis == EAxisList::Forward)
 	{
-		return LOCTEXT("Forward_ToolTipTextFormat", "Forward");
+		return LOCTEXT("GetRotationAxisNameShort_Forward", "Forward");
 	}
 	else if (Axis == EAxisList::Left)
 	{
-		return LOCTEXT("Left_ToolTipTextFormat", "Left");
+		return LOCTEXT("GetRotationAxisNameShort_Left", "Left");
 	}
 	else if (Axis == EAxisList::Up)
 	{
-		return LOCTEXT("Up_ToolTipTextFormat", "Up");
+		return LOCTEXT("GetRotationAxisNameShort_Up", "Up");
 	}
 
 	ensureMsgf(false, TEXT("Unsupported Axis: %d"), Axis);
 
-	return LOCTEXT("UnsupportedDisplayNameShort", "?");
+	return LOCTEXT("GetRotationAxisNameShort_Unsupported", "?");
 }
 
 EAxisList::Type FEditorAxisDisplayInfo::MapAxis(EAxisList::Type Axis) const
diff --git a/Engine/Source/Runtime/Core/Private/Misc/AxisDisplayInfo.cpp b/Engine/Source/Runtime/Core/Private/Misc/AxisDisplayInfo.cpp
--- a/Engine/Source/Runtime/Core/Private/Misc/AxisDisplayInfo.cpp
+++ b/Engine/Source/Runtime/Core/Private/Misc/AxisDisplayInfo.cpp
@@ -173,7 +173,7 @@ FText AxisDisplayInfo::GetRotationAxisToolTip(EAxisList::Type Axis)
 	if (IModularFeatures::Get().IsModularFeatureAvailable(IAxisDisplayInfo::GetModularFeatureName()))
 	{
 		IAxisDisplayInfo& DisplayInfo = IModularFeatures::Get().GetModularFeature<IAxisDisplayInfo>(IAxisDisplayInfo::GetModularFeatureName());
-		return DisplayInfo.GetRotationAxisNameShort(Axis);
+		return DisplayInfo.GetRotationAxisToolTip(Axis);
 	}
 	return LOCTEXT("UnsupportedRotationAxisToolTip", "");
 }
