[PropertyEditor] Fix vector customization to support missing properties. It's probably a user's mistake, and the user should use the proper customization.

diff --git a/Engine/Source/Editor/DetailCustomizations/Private/VectorStructCustomization.cpp b/Engine/Source/Editor/DetailCustomizations/Private/VectorStructCustomization.cpp
--- a/Engine/Source/Editor/DetailCustomizations/Private/VectorStructCustomization.cpp
+++ b/Engine/Source/Editor/DetailCustomizations/Private/VectorStructCustomization.cpp
@@ -40,7 +40,14 @@ void FVectorStructCustomization::GetSortedChildren(TSharedRef<IPropertyHandle> S
 	static const FName Y("Y");
 	static const FName Z("Z");
 
-	TSharedPtr<IPropertyHandle> VectorChildren[3];
+	FStructProperty* StructProperty = CastField<FStructProperty>(StructPropertyHandle->GetProperty());
+	if (!ensureMsgf(StructProperty, TEXT("Vector customization only supports UScriptStruct.")))
+	{
+		return;
+	}
+
+	constexpr int32 NumberOfElements = 3;
+	TSharedPtr<IPropertyHandle> VectorChildren[NumberOfElements];
 
 	uint32 NumChildren;
 	StructPropertyHandle->GetNumChildren(NumChildren);
@@ -64,18 +71,24 @@ void FVectorStructCustomization::GetSortedChildren(TSharedRef<IPropertyHandle> S
 
 			VectorChildren[1] = ChildHandle;
 		}
-		else
+		else if(PropertyName == Z)
 		{
-			check(PropertyName == Z);
-
 			ChildHandle->SetPropertyDisplayName(AxisDisplayInfo::GetAxisDisplayName(EAxisList::Z));
 			ChildHandle->SetToolTipText(AxisDisplayInfo::GetAxisDisplayName(EAxisList::Z));
 
 			VectorChildren[2] = ChildHandle;
 		}
+		else
+		{
+			ensureMsgf(false, TEXT("The property doesn't exist. Vector customization supports X,Y,Z properties. %s"), *StructProperty->Struct->GetFullName());
+		}
 	}
 
-	OutChildren.Add(VectorChildren[0].ToSharedRef());
-	OutChildren.Add(VectorChildren[1].ToSharedRef());
-	OutChildren.Add(VectorChildren[2].ToSharedRef());
+	for (int32 Index = 0; Index < NumberOfElements; ++Index)
+	{
+		if (ensureMsgf(VectorChildren[Index], TEXT("Missing a property. Vector customization supports X,Y,Z properties. %s"), *StructProperty->Struct->GetFullName()))
+		{
+			OutChildren.Add(VectorChildren[Index].ToSharedRef());
+		}
+	}
 }
