Expose template functions to extend MathStructProxyCustomization

diff --git a/Engine/Source/Editor/DetailCustomizations/DetailCustomizations.Build.cs b/Engine/Source/Editor/DetailCustomizations/DetailCustomizations.Build.cs
--- a/Engine/Source/Editor/DetailCustomizations/DetailCustomizations.Build.cs
+++ b/Engine/Source/Editor/DetailCustomizations/DetailCustomizations.Build.cs
@@ -6,6 +6,7 @@ public class DetailCustomizations : ModuleRules
 {
 	public DetailCustomizations(ReadOnlyTargetRules Target) : base(Target)
 	{
+		PublicDependencyModuleNames.AddRange(new string[] { "Slate", "PropertyEditor" });
 		// Enable truncation warnings in this module
 		CppCompileWarningSettings.UnsafeTypeCastWarningLevel = WarningLevel.Warning;
 
diff --git a/Engine/Source/Editor/DetailCustomizations/Private/MathStructProxyCustomizations.cpp b/Engine/Source/Editor/DetailCustomizations/Private/MathStructProxyCustomizations.cpp
--- a/Engine/Source/Editor/DetailCustomizations/Private/MathStructProxyCustomizations.cpp
+++ b/Engine/Source/Editor/DetailCustomizations/Private/MathStructProxyCustomizations.cpp
@@ -50,87 +50,11 @@ void FMathStructProxyCustomization::MakeHeaderRow( TSharedRef<class IPropertyHan
 
 }
 
-template<typename ProxyType, typename NumericType>
-TSharedRef<SWidget> FMathStructProxyCustomization::MakeNumericProxyWidget(
-	TSharedRef<IPropertyHandle>& StructPropertyHandle,
-	TSharedRef< TProxyProperty<ProxyType, NumericType> >& ProxyValue,
-	const FText& Label,
-	bool bRotationInDegrees,
-	const FLinearColor& LabelBackgroundColor)
-{
-	TWeakPtr<IPropertyHandle> WeakHandlePtr = StructPropertyHandle;
-
-	return
-		SNew(SNumericEntryBox<NumericType>)
-		.IsEnabled(this, &FMathStructProxyCustomization::IsValueEnabled, WeakHandlePtr)
-		.Value(this, &FMathStructProxyCustomization::OnGetValue<ProxyType, NumericType>, WeakHandlePtr, ProxyValue)
-		.Font(IDetailLayoutBuilder::GetDetailFont())
-		.UndeterminedString(NSLOCTEXT("PropertyEditor", "MultipleValues", "Multiple Values"))
-		.OnValueCommitted(this, &FMathStructProxyCustomization::OnValueCommitted<ProxyType, NumericType>, WeakHandlePtr, ProxyValue)
-		.OnValueChanged(this, &FMathStructProxyCustomization::OnValueChanged<ProxyType, NumericType>, WeakHandlePtr, ProxyValue)
-		.OnBeginSliderMovement(this, &FMathStructProxyCustomization::OnBeginSliderMovement)
-		.OnEndSliderMovement(this, &FMathStructProxyCustomization::OnEndSliderMovement<ProxyType, NumericType>, WeakHandlePtr, ProxyValue)
-		// Only allow spin on handles with one object.  Otherwise it is not clear what value to spin
-		.AllowSpin(StructPropertyHandle->GetNumOuterObjects() == 1)
-		.MinValue(TOptional<NumericType>())
-		.MaxValue(TOptional<NumericType>())
-		.MaxSliderValue(bRotationInDegrees ? 360.0f : TOptional<NumericType>())
-		.MinSliderValue(bRotationInDegrees ? 0.0f : TOptional<NumericType>())
-		.LabelPadding(FMargin(3))
-		.ToolTipTextFormat(this, &FMathStructProxyCustomization::OnGetValueToolTipTextFormat, Label)
-		.LabelLocation(SNumericEntryBox<NumericType>::ELabelLocation::Inside)
-		.Label()
-		[
-			SNumericEntryBox<NumericType>::BuildNarrowColorLabel(LabelBackgroundColor)
-		];
-}
-
-
-template<typename ProxyType, typename NumericType>
-TOptional<NumericType> FMathStructProxyCustomization::OnGetValue( TWeakPtr<IPropertyHandle> WeakHandlePtr, TSharedRef< TProxyProperty<ProxyType, NumericType> > ProxyValue ) const
-{
-	if(CacheValues(WeakHandlePtr))
-	{
-		return ProxyValue->Get();
-	}
-	return TOptional<NumericType>();
-}
-
-template<typename ProxyType, typename NumericType>
-void FMathStructProxyCustomization::OnValueCommitted( NumericType NewValue, ETextCommit::Type CommitType, TWeakPtr<IPropertyHandle> WeakHandlePtr, TSharedRef< TProxyProperty<ProxyType, NumericType> > ProxyValue )
-{
-	if (!bIsUsingSlider && !GIsTransacting)
-	{
-		ProxyValue->Set(NewValue);
-		FlushValues(WeakHandlePtr);
-	}
-}	
-
-template<typename ProxyType, typename NumericType>
-void FMathStructProxyCustomization::OnValueChanged( NumericType NewValue, TWeakPtr<IPropertyHandle> WeakHandlePtr, TSharedRef< TProxyProperty<ProxyType, NumericType> > ProxyValue )
-{
-	if( bIsUsingSlider )
-	{
-		ProxyValue->Set(NewValue);
-		FlushValues(WeakHandlePtr);
-	}
-}
-
 void FMathStructProxyCustomization::OnBeginSliderMovement()
 {
 	bIsUsingSlider = true;
 }
 
-template<typename ProxyType, typename NumericType>
-void FMathStructProxyCustomization::OnEndSliderMovement( NumericType NewValue, TWeakPtr<IPropertyHandle> WeakHandlePtr, TSharedRef< TProxyProperty<ProxyType, NumericType> > ProxyValue )
-{
-	bIsUsingSlider = false;
-
-	ProxyValue->Set(NewValue);
-	FlushValues(WeakHandlePtr);
-}
-
-
 template <typename ProxyType, typename NumericType>
 FText FMathStructProxyCustomization::OnGetValueToolTip(TWeakPtr<IPropertyHandle> WeakHandlePtr, TSharedRef<TProxyProperty<ProxyType, NumericType>> ProxyValue, FText Label) const
 {
diff --git a/Engine/Source/Editor/DetailCustomizations/Public/Customizations/MathStructProxyCustomizations.h b/Engine/Source/Editor/DetailCustomizations/Public/Customizations/MathStructProxyCustomizations.h
--- a/Engine/Source/Editor/DetailCustomizations/Public/Customizations/MathStructProxyCustomizations.h
+++ b/Engine/Source/Editor/DetailCustomizations/Public/Customizations/MathStructProxyCustomizations.h
@@ -4,6 +4,7 @@
 
 #include "Containers/Array.h"
 #include "CoreMinimal.h"
+#include "DetailLayoutBuilder.h"
 #include "Customizations/MathStructCustomizations.h"
 #include "Delegates/Delegate.h"
 #include "IPropertyTypeCustomization.h"
@@ -18,6 +19,7 @@
 #include "Templates/SharedPointer.h"
 #include "Types/SlateEnums.h"
 #include "Widgets/SWidget.h"
+#include "Widgets/Input/SNumericEntryBox.h"
 
 class FDetailWidgetRow;
 class IPropertyHandle;
@@ -156,14 +158,14 @@ private:
  * Helper class to aid representing math structs to the user in an editable form
  * e.g. representing a quaternion as a set of euler angles
  */
-class DETAILCUSTOMIZATIONS_API FMathStructProxyCustomization : public FMathStructCustomization
+class FMathStructProxyCustomization : public FMathStructCustomization
 {
 public:
 	/** IPropertyTypeCustomization interface */
-	virtual void CustomizeChildren( TSharedRef<class IPropertyHandle> StructPropertyHandle, class IDetailChildrenBuilder& StructBuilder, IPropertyTypeCustomizationUtils& StructCustomizationUtils ) override;
+	DETAILCUSTOMIZATIONS_API virtual void CustomizeChildren( TSharedRef<class IPropertyHandle> StructPropertyHandle, class IDetailChildrenBuilder& StructBuilder, IPropertyTypeCustomizationUtils& StructCustomizationUtils ) override;
 
 	/** FMathStructCustomization interface */
-	virtual void MakeHeaderRow( TSharedRef<class IPropertyHandle>& InStructPropertyHandle, FDetailWidgetRow& Row ) override;
+	DETAILCUSTOMIZATIONS_API virtual void MakeHeaderRow( TSharedRef<class IPropertyHandle>& InStructPropertyHandle, FDetailWidgetRow& Row ) override;
 
 protected:
 
@@ -196,7 +198,7 @@ protected:
 	UE_DEPRECATED(5.6, "OnGetValueToolTip is deprecated and no longer called. Please use the Label in MakeNumericProxyWidget to customize the display name portion of the ToolTip.")
 	FText OnGetValueToolTip(TWeakPtr<IPropertyHandle> WeakHandlePtr, TSharedRef<TProxyProperty<ProxyType, NumericType>> ProxyValue, FText Label) const;
 
-	TOptional<FTextFormat> OnGetValueToolTipTextFormat(FText Label) const;
+	DETAILCUSTOMIZATIONS_API TOptional<FTextFormat> OnGetValueToolTipTextFormat(FText Label) const;
 
 private:
 
@@ -230,7 +232,7 @@ private:
 	void OnValueChanged( NumericType NewValue, TWeakPtr<IPropertyHandle> WeakHandlePtr, TSharedRef< TProxyProperty<ProxyType, NumericType> > ProxyValue );
 
 	/** Called when a value starts to be changed by a slider */
-	void OnBeginSliderMovement();
+	DETAILCUSTOMIZATIONS_API void OnBeginSliderMovement();
 
 	/** Called when a value stops being changed by a slider */
 	template<typename ProxyType, typename NumericType>
@@ -241,6 +243,79 @@ protected:
 	TSharedPtr<IPropertyUtilities> PropertyUtilities;
 };
 
+template<typename ProxyType, typename NumericType>
+inline TSharedRef<SWidget> FMathStructProxyCustomization::MakeNumericProxyWidget(
+	TSharedRef<IPropertyHandle>& StructPropertyHandle,
+	TSharedRef< TProxyProperty<ProxyType, NumericType> >& ProxyValue,
+	const FText& Label,
+	bool bRotationInDegrees,
+	const FLinearColor& LabelBackgroundColor)
+{
+	TWeakPtr<IPropertyHandle> WeakHandlePtr = StructPropertyHandle;
+
+	return
+		SNew(SNumericEntryBox<NumericType>)
+		.IsEnabled(this, &FMathStructProxyCustomization::IsValueEnabled, WeakHandlePtr)
+		.Value(this, &FMathStructProxyCustomization::OnGetValue<ProxyType, NumericType>, WeakHandlePtr, ProxyValue)
+		.Font(IDetailLayoutBuilder::GetDetailFont())
+		.UndeterminedString(NSLOCTEXT("PropertyEditor", "MultipleValues", "Multiple Values"))
+		.OnValueCommitted(this, &FMathStructProxyCustomization::OnValueCommitted<ProxyType, NumericType>, WeakHandlePtr, ProxyValue)
+		.OnValueChanged(this, &FMathStructProxyCustomization::OnValueChanged<ProxyType, NumericType>, WeakHandlePtr, ProxyValue)
+		.OnBeginSliderMovement(this, &FMathStructProxyCustomization::OnBeginSliderMovement)
+		.OnEndSliderMovement(this, &FMathStructProxyCustomization::OnEndSliderMovement<ProxyType, NumericType>, WeakHandlePtr, ProxyValue)
+		// Only allow spin on handles with one object.  Otherwise it is not clear what value to spin
+		.AllowSpin(StructPropertyHandle->GetNumOuterObjects() == 1)
+		.MinValue(TOptional<NumericType>())
+		.MaxValue(TOptional<NumericType>())
+		.MaxSliderValue(bRotationInDegrees ? 360.0f : TOptional<NumericType>())
+		.MinSliderValue(bRotationInDegrees ? 0.0f : TOptional<NumericType>())
+		.LabelPadding(FMargin(3))
+		.ToolTipTextFormat(this, &FMathStructProxyCustomization::OnGetValueToolTipTextFormat, Label)
+		.LabelLocation(SNumericEntryBox<NumericType>::ELabelLocation::Inside)
+		.Label()
+		[
+			SNumericEntryBox<NumericType>::BuildNarrowColorLabel(LabelBackgroundColor)
+		];
+}
+
+template<typename ProxyType, typename NumericType>
+inline void FMathStructProxyCustomization::OnEndSliderMovement( NumericType NewValue, TWeakPtr<IPropertyHandle> WeakHandlePtr, TSharedRef< TProxyProperty<ProxyType, NumericType> > ProxyValue )
+{
+	bIsUsingSlider = false;
+
+	ProxyValue->Set(NewValue);
+	FlushValues(WeakHandlePtr);
+}
+
+template<typename ProxyType, typename NumericType>
+inline TOptional<NumericType> FMathStructProxyCustomization::OnGetValue( TWeakPtr<IPropertyHandle> WeakHandlePtr, TSharedRef< TProxyProperty<ProxyType, NumericType> > ProxyValue ) const
+{
+	if(CacheValues(WeakHandlePtr))
+	{
+		return ProxyValue->Get();
+	}
+	return TOptional<NumericType>();
+}
+
+template<typename ProxyType, typename NumericType>
+inline void FMathStructProxyCustomization::OnValueCommitted( NumericType NewValue, ETextCommit::Type CommitType, TWeakPtr<IPropertyHandle> WeakHandlePtr, TSharedRef< TProxyProperty<ProxyType, NumericType> > ProxyValue )
+{
+	if (!bIsUsingSlider && !GIsTransacting)
+	{
+		ProxyValue->Set(NewValue);
+		FlushValues(WeakHandlePtr);
+	}
+}	
+
+template<typename ProxyType, typename NumericType>
+inline void FMathStructProxyCustomization::OnValueChanged( NumericType NewValue, TWeakPtr<IPropertyHandle> WeakHandlePtr, TSharedRef< TProxyProperty<ProxyType, NumericType> > ProxyValue )
+{
+	if( bIsUsingSlider )
+	{
+		ProxyValue->Set(NewValue);
+		FlushValues(WeakHandlePtr);
+	}
+}
 
 struct FTransformField
 {
