MaterialParameterCollection SanitizeParameters fixes.

diff --git a/Engine/Source/Editor/DetailCustomizations/Private/MaterialParameterCollectionDetails.cpp b/Engine/Source/Editor/DetailCustomizations/Private/MaterialParameterCollectionDetails.cpp
--- a/Engine/Source/Editor/DetailCustomizations/Private/MaterialParameterCollectionDetails.cpp
+++ b/Engine/Source/Editor/DetailCustomizations/Private/MaterialParameterCollectionDetails.cpp
@@ -17,7 +17,7 @@ TSharedRef<IDetailCustomization> FMaterialParameterCollectionDetails::MakeInstan
 
 FMaterialParameterCollectionDetails::~FMaterialParameterCollectionDetails()
 {
-	FEditorDelegates::PostUndoRedo.RemoveAll(DetailLayout);
+	UnbindDelegates();
 }
 
 class FMaterialParameterCollectionDetails::FBaseOverridesSelectorConst
@@ -260,10 +260,15 @@ void FMaterialParameterCollectionDetails::CustomizeCollectionDetails(UMaterialPa
 
 void FMaterialParameterCollectionDetails::CustomizeDetails(IDetailLayoutBuilder& InDetailLayout)
 {
+	UnbindDelegates();
+
 	DetailLayout = &InDetailLayout;
 
+	// Bind delegates for detecting undo/redo and property changes in other objects.
 	FEditorDelegates::PostUndoRedo.AddRaw(DetailLayout, &IDetailLayoutBuilder::ForceRefreshDetails);
+	FCoreUObjectDelegates::OnObjectPropertyChanged.AddRaw(this, &FMaterialParameterCollectionDetails::OnObjectPropertyChanged);
 
+	// Customize the details of each collection being customized.
 	TArray<TWeakObjectPtr<UMaterialParameterCollection>> Collections = DetailLayout->GetObjectsOfTypeBeingCustomized<UMaterialParameterCollection>();
 	for (TWeakObjectPtr<UMaterialParameterCollection>& CollectionWeak : Collections)
 	{
@@ -274,4 +279,39 @@ void FMaterialParameterCollectionDetails::CustomizeDetails(IDetailLayoutBuilder&
 	}
 }
 
+void FMaterialParameterCollectionDetails::OnObjectPropertyChanged(UObject* Object, FPropertyChangedEvent& PropertyChangedEvent)
+{
+	if (!DetailLayout || !Object->IsA<UMaterialParameterCollection>())
+	{
+		return;
+	}
+
+	// Refresh the layout details if the modified object is the base of any of the collections being customized.
+	TArray<TWeakObjectPtr<UMaterialParameterCollection>> Collections = DetailLayout->GetObjectsOfTypeBeingCustomized<UMaterialParameterCollection>();
+	for (TWeakObjectPtr<UMaterialParameterCollection>& CollectionWeak : Collections)
+	{
+		if (UMaterialParameterCollection* Collection = CollectionWeak.Get())
+		{
+			for (UMaterialParameterCollection* BaseCollection = Collection->Base; BaseCollection; BaseCollection = BaseCollection->Base)
+			{
+				if (BaseCollection == Object)
+				{
+					DetailLayout->ForceRefreshDetails();
+					return;
+				}
+			}
+		}
+	}
+}
+
+void FMaterialParameterCollectionDetails::UnbindDelegates()
+{
+	FCoreUObjectDelegates::OnObjectPropertyChanged.RemoveAll(this);
+
+	if (DetailLayout)
+	{
+		FEditorDelegates::PostUndoRedo.RemoveAll(DetailLayout);
+	}
+}
+
 #undef LOCTEXT_NAMESPACE
diff --git a/Engine/Source/Editor/DetailCustomizations/Private/MaterialParameterCollectionDetails.h b/Engine/Source/Editor/DetailCustomizations/Private/MaterialParameterCollectionDetails.h
--- a/Engine/Source/Editor/DetailCustomizations/Private/MaterialParameterCollectionDetails.h
+++ b/Engine/Source/Editor/DetailCustomizations/Private/MaterialParameterCollectionDetails.h
@@ -14,6 +14,7 @@ struct FCollectionScalarParameter;
 struct FCollectionVectorParameter;
 struct FGuid;
 struct FLinearColor;
+struct FPropertyChangedEvent;
 
 //-----------------------------------------------------------------------------
 //   FMaterialParameterCollectionDetails
@@ -53,5 +54,9 @@ private:
 
 	void CustomizeCollectionDetails(UMaterialParameterCollection* Collection);
 
+	void OnObjectPropertyChanged(UObject* Object, FPropertyChangedEvent& PropertyChangedEvent);
+
+	void UnbindDelegates();
+
 	IDetailLayoutBuilder* DetailLayout = nullptr;
 };
diff --git a/Engine/Source/Runtime/Engine/Private/Materials/ParameterCollection.cpp b/Engine/Source/Runtime/Engine/Private/Materials/ParameterCollection.cpp
--- a/Engine/Source/Runtime/Engine/Private/Materials/ParameterCollection.cpp
+++ b/Engine/Source/Runtime/Engine/Private/Materials/ParameterCollection.cpp
@@ -345,6 +345,16 @@ void UMaterialParameterCollection::PostEditChangeProperty(FPropertyChangedEvent&
 	}
 	else
 	{
+		// If this collection is the base of another collection, the other collection must be updated.
+		for (TObjectIterator<UMaterialParameterCollection> It(RF_ClassDefaultObject, true, EInternalObjectFlags::Garbage); It; ++It)
+		{
+			UMaterialParameterCollection* Collection = *It;
+			if (Collection != this)
+			{
+				Collection->UpdateOverrides(this);
+			}
+		}
+
 		// We didn't need to recreate the uniform buffer, just update its contents
 		for (TObjectIterator<UWorld> It(/*AdditionalExclusionFlags = */RF_ClassDefaultObject, /*bIncludeDerivedClasses = */true, /*InInternalExclusionFlags = */EInternalObjectFlags::Garbage); It; ++It)
 		{
@@ -367,13 +377,19 @@ TTuple<FString, int32> SplitParameterName(FName ParameterName)
 
 	int32 Number = 0;
 	int32 NumberStartIndex = NameString.FindLastCharByPredicate([](TCHAR Letter){ return !FChar::IsDigit(Letter); }) + 1;
-	if (NumberStartIndex < NameString.Len() - 1)
+	if (NumberStartIndex < NameString.Len())
 	{
 		ensure(TCString<TCHAR>::IsNumeric(&NameString[NumberStartIndex]));
 		TTypeFromString<int32>::FromString(Number, &NameString[NumberStartIndex]);
+		NameString.LeftInline(NumberStartIndex);
+	}
+
+	if (NameString.IsEmpty())
+	{
+		NameString = TEXT("Param");
 	}
 	
-	return { NameString.Left(NumberStartIndex), Number };
+	return { MoveTemp(NameString), Number };
 }
 
 template<typename FCollectionParameterType>
@@ -442,6 +458,10 @@ bool UMaterialParameterCollection::UpdateOverrides(UMaterialParameterCollection*
 		return false;
 	}
 
+	// Ensure the collection's parameter names account for the base.
+	SanitizeParameters(&UMaterialParameterCollection::ScalarParameters);
+	SanitizeParameters(&UMaterialParameterCollection::VectorParameters);
+
 	// If the BaseStateId is up-to-date, the base overrides should be as well.
 	if (BaseStateId == Base->StateId)
 	{
@@ -478,10 +498,6 @@ bool UMaterialParameterCollection::UpdateOverrides(UMaterialParameterCollection*
 		}
 	}
 
-	// Ensure the collection's parameter names account for the new base.
-	SanitizeParameters(&UMaterialParameterCollection::ScalarParameters);
-	SanitizeParameters(&UMaterialParameterCollection::VectorParameters);
-
 	// Indicate that this collection is based on the BaseCollection.
 	return true;
 }
