Fix issue where MaterialParameterCollection base parameter overrides were not updating live in the editor.

#tests Editor - overridden MPC in MI, verified editing child MPC value would update MI editor
#rnx

diff --git a/Engine/Source/Editor/DetailCustomizations/Private/MaterialParameterCollectionDetails.cpp b/Engine/Source/Editor/DetailCustomizations/Private/MaterialParameterCollectionDetails.cpp
--- a/Engine/Source/Editor/DetailCustomizations/Private/MaterialParameterCollectionDetails.cpp
+++ b/Engine/Source/Editor/DetailCustomizations/Private/MaterialParameterCollectionDetails.cpp
@@ -4,9 +4,11 @@
 
 #include "DetailLayoutBuilder.h"
 #include "Editor.h"
+#include "EditorSupportDelegates.h"
 #include "IDetailGroup.h"
 #include "IPropertyUtilities.h"
 #include "Materials/MaterialParameterCollection.h"
+#include "ScopedTransaction.h"
 
 #define LOCTEXT_NAMESPACE "MaterialParameterCollectionDetails"
 
@@ -106,6 +108,16 @@ auto FMaterialParameterCollectionDetails::GetParameterValue(const UMaterialParam
 	return CollectionParameter.DefaultValue;
 }
 
+FProperty* FMaterialParameterCollectionDetails::GetBaseOverridesMapProperty(const FCollectionScalarParameter&)
+{
+	return FindFProperty<FProperty>(UMaterialParameterCollection::StaticClass(), GET_MEMBER_NAME_CHECKED(UMaterialParameterCollection, ScalarParameterBaseOverrides));
+}
+
+FProperty* FMaterialParameterCollectionDetails::GetBaseOverridesMapProperty(const FCollectionVectorParameter&)
+{
+	return FindFProperty<FProperty>(UMaterialParameterCollection::StaticClass(), GET_MEMBER_NAME_CHECKED(UMaterialParameterCollection, VectorParameterBaseOverrides));
+}
+
 template<class FCollectionParameterType>
 void FMaterialParameterCollectionDetails::AddParameter(IDetailGroup& DetailGroup, TSharedPtr<IPropertyHandle> PropertyHandle, UMaterialParameterCollection* Collection, const FCollectionParameterType& CollectionParameter, const UMaterialParameterCollection* BaseCollection)
 {
@@ -135,9 +147,16 @@ void FMaterialParameterCollectionDetails::AddParameter(IDetailGroup& DetailGroup
 			void* ValueData = nullptr;
 			if (PropertyHandle->GetValueData(ValueData) == FPropertyAccess::Success && ValueData)
 			{
+				FProperty* CollectionPropertyToChange = GetBaseOverridesMapProperty(CollectionParameter);
+				Collection->PreEditChange(CollectionPropertyToChange);
+
 				FValueType& Value = *static_cast<FValueType*>(ValueData);
 				FValueType& ParameterValue = GetBaseOverridesMap(Collection).FindOrAdd(CollectionParameter);
 				ParameterValue = Value;
+
+				FPropertyChangedEvent CollectionChangedEvent(CollectionPropertyToChange, EPropertyChangeType::ValueSet);
+				Collection->PostEditChangeProperty(CollectionChangedEvent);
+
 			}
 		});
 
@@ -156,6 +175,12 @@ void FMaterialParameterCollectionDetails::AddParameter(IDetailGroup& DetailGroup
 			void* ValueData = nullptr;
 			if (PropertyHandle->GetValueData(ValueData) == FPropertyAccess::Success && ValueData)
 			{
+				// Create a new transaction for undo/redo support - normally handled by IPropertyHandle
+				FScopedTransaction EditConditionChangedTransaction(NSLOCTEXT("PropertyEditor", "UpdatedEditConditionFmt", "Edit Condition Changed"));
+
+				FProperty* CollectionPropertyToChange = GetBaseOverridesMapProperty(CollectionParameter);
+				Collection->PreEditChange(CollectionPropertyToChange);
+
 				FValueType& Value = *static_cast<FValueType*>(ValueData);
 
 				// Remove the parameter from the collection instance.
@@ -170,6 +195,12 @@ void FMaterialParameterCollectionDetails::AddParameter(IDetailGroup& DetailGroup
 					// Add the parameter to the collection instance, with the value specified by its parent.
 					BaseOverridesMap.Add(CollectionParameter, Value);
 				}
+
+				FPropertyChangedEvent CollectionChangedEvent(CollectionPropertyToChange, (bOverride ? EPropertyChangeType::ArrayAdd : EPropertyChangeType::ArrayRemove));
+				Collection->PostEditChangeProperty(CollectionChangedEvent);
+
+				// Update all viewports after we modify the collection - normally handled by IPropertyHandle
+				FEditorSupportDelegates::RedrawAllViewports.Broadcast();
 			}
 		};
 
diff --git a/Engine/Source/Editor/DetailCustomizations/Private/MaterialParameterCollectionDetails.h b/Engine/Source/Editor/DetailCustomizations/Private/MaterialParameterCollectionDetails.h
--- a/Engine/Source/Editor/DetailCustomizations/Private/MaterialParameterCollectionDetails.h
+++ b/Engine/Source/Editor/DetailCustomizations/Private/MaterialParameterCollectionDetails.h
@@ -4,6 +4,7 @@
 
 #include "IDetailCustomization.h"
 
+class FProperty;
 class IDetailGroup;
 class IDetailCategoryBuilder;
 class IDetailLayoutBuilder;
@@ -43,6 +44,9 @@ private:
 	static FBaseOverridesSelector GetBaseOverridesMap(UMaterialParameterCollection* OverrideCollection);
 	static FBaseOverridesSelectorConst GetBaseOverridesMap(const UMaterialParameterCollection* OverrideCollection);
 
+	static FProperty* GetBaseOverridesMapProperty(const FCollectionScalarParameter&);
+	static FProperty* GetBaseOverridesMapProperty(const FCollectionVectorParameter&);
+
 	template<class FCollectionParameterType>
 	static auto GetParameterValue(const UMaterialParameterCollection* Collection, const FCollectionParameterType& CollectionParameter, const UMaterialParameterCollection* BaseCollection = nullptr);
 
