Fix inverted logic for component transform ResetToDefault button visibility

diff --git a/Engine/Source/Editor/DetailCustomizations/Private/ComponentTransformDetails.cpp b/Engine/Source/Editor/DetailCustomizations/Private/ComponentTransformDetails.cpp
--- a/Engine/Source/Editor/DetailCustomizations/Private/ComponentTransformDetails.cpp
+++ b/Engine/Source/Editor/DetailCustomizations/Private/ComponentTransformDetails.cpp
@@ -1060,7 +1060,7 @@ bool FComponentTransformDetails::GetRotationResetVisibility() const
 	else
 	{
 		const FVector Data = Archetype ? ConvertFromUnrealSpace_EulerDeg(Archetype->GetRelativeRotation()) : FVector::ZeroVector;
-		return ComponentTransformDetails::Private::AreRotationsEqual(Data, CachedRotation.ToVector());
+		return !ComponentTransformDetails::Private::AreRotationsEqual(Data, CachedRotation.ToVector());
 	}
 }
 
