Sort sound categories in asset editor details pane based on a predefined order.

diff --git a/Engine/Source/Editor/DetailCustomizations/Private/AttenuationSettingsCustomizations.cpp b/Engine/Source/Editor/DetailCustomizations/Private/AttenuationSettingsCustomizations.cpp
--- a/Engine/Source/Editor/DetailCustomizations/Private/AttenuationSettingsCustomizations.cpp
+++ b/Engine/Source/Editor/DetailCustomizations/Private/AttenuationSettingsCustomizations.cpp
@@ -99,14 +99,22 @@ namespace AttenuationSettingsUtils
 			{
 				SortOrder += 9;
 			}
+			else if (CategoryName == "AttenuationSourceDataOverride")
+			{
+				SortOrder += 10;
+			}
+			else if (CategoryName == "AttenuationAudioLink")
+			{
+				SortOrder += 11;
+			}
 
 			else
 			{
 				// Add space to any other categories interfering with space for attenuation-related categories
 				const int32 ValueSortOrder = Pair.Value->GetSortOrder();
-				if (ValueSortOrder>= SortOrder && ValueSortOrder <SortOrder + 10)
+				if (ValueSortOrder>= SortOrder)
 				{
-					SortOrder += 10;
+					SortOrder = ValueSortOrder + 12;
 				}
 				// Else, leave current sort order alone
 				else
diff --git a/Engine/Source/Editor/DetailCustomizations/Private/DetailCustomizations.cpp b/Engine/Source/Editor/DetailCustomizations/Private/DetailCustomizations.cpp
--- a/Engine/Source/Editor/DetailCustomizations/Private/DetailCustomizations.cpp
+++ b/Engine/Source/Editor/DetailCustomizations/Private/DetailCustomizations.cpp
@@ -144,6 +144,7 @@
 #include "SoftClassPathCustomization.h"
 #include "SoftObjectPathCustomization.h"
 #include "Sound/SoundNodeDistanceCrossFade.h"
+#include "SoundCueDetails.h"
 #include "SoundSourceBusDetails.h"
 #include "SoundWaveDetails.h"
 #include "SourceCodeAccessSettingsDetails.h"
@@ -437,6 +438,7 @@ void FDetailCustomizationsModule::RegisterObjectCustomizations()
 	RegisterCustomClassLayout("SoundSourceBus", FOnGetDetailCustomizationInstance::CreateStatic(&FSoundSourceBusDetails::MakeInstance));
 	RegisterCustomClassLayout("DialogueWave", FOnGetDetailCustomizationInstance::CreateStatic(&FDialogueWaveDetails::MakeInstance));
 	RegisterCustomClassLayout("SoundWave", FOnGetDetailCustomizationInstance::CreateStatic(&FSoundWaveDetails::MakeInstance));
+	RegisterCustomClassLayout("SoundCue", FOnGetDetailCustomizationInstance::CreateStatic(&FSoundCueDetails::MakeInstance));
 
 	RegisterCustomClassLayout("BodySetup", FOnGetDetailCustomizationInstance::CreateStatic(&FBodySetupDetails::MakeInstance));
 	RegisterCustomClassLayout("SkeletalBodySetup", FOnGetDetailCustomizationInstance::CreateStatic(&FSkeletalBodySetupDetails::MakeInstance));
diff --git a/Engine/Source/Editor/DetailCustomizations/Private/SoundBaseDetailsUtils.cpp b/Engine/Source/Editor/DetailCustomizations/Private/SoundBaseDetailsUtils.cpp
new file mode 100644
--- /dev/null
+++ b/Engine/Source/Editor/DetailCustomizations/Private/SoundBaseDetailsUtils.cpp
@@ -0,0 +1,76 @@
+// Copyright Epic Games, Inc. All Rights Reserved.
+
+#include "SoundBaseDetailsUtils.h"
+
+namespace SoundBaseDetailsUtils
+{
+	static const int32 DefaultTopSortOrder = 4000;
+
+	static const TArray<FName> DefaultCategoryOrder =
+	{
+		"SoundWave",
+		"Format",
+		"Sound",
+		"Analysis",
+		"Modulation",
+		"Subtitles",
+		"Loading",
+		"Info",
+		"File Path",
+		"Waveform Processing",
+		"Developer",
+		"Voice Management",
+		"Effects",
+		"Attenuation",
+		"Curves",
+		"Advanced",
+		"Memory",
+	};
+} // namespace SoundBaseDetailsUtils
+
+void FSoundBaseDetailsUtils::SortSoundCategories(const TMap<FName, IDetailCategoryBuilder*>& AllCategoryMap)
+{
+	// if no category order is passed in, use the default
+	FSoundBaseDetailsUtils::SortSoundCategories(AllCategoryMap, SoundBaseDetailsUtils::DefaultCategoryOrder);
+}
+
+void FSoundBaseDetailsUtils::SortSoundCategories(const TMap<FName, IDetailCategoryBuilder*>& AllCategoryMap, const TArray<FName> CategoryOrder)
+{
+	// Start at the top of the list, or at a default value if you can't find that
+	int32 TopSortOrder = SoundBaseDetailsUtils::DefaultTopSortOrder;
+	const IDetailCategoryBuilder* FirstCategory = AllCategoryMap.FindRef(CategoryOrder[0]);
+	if (FirstCategory)
+	{
+		TopSortOrder = FirstCategory->GetSortOrder();
+	}
+
+	// Categories not listed in the CategoryOrder will be appended at the end with incoming order
+	int32 NumCategoriesNotFound = 0;
+	for (const TPair<FName, IDetailCategoryBuilder*>& Pair : AllCategoryMap)
+	{
+		const FName CategoryName = Pair.Key;
+		int32 SortOffset = 1;
+
+		bool CategoryFound = false;
+		for (const FName& OrderedName : CategoryOrder)
+		{
+			if (OrderedName == CategoryName)
+			{
+				CategoryFound = true;
+				break;
+			}
+			SortOffset++;
+		}
+
+		int32 SortOrder = TopSortOrder + SortOffset;
+		if (!CategoryFound)
+		{
+			NumCategoriesNotFound++;
+			SortOrder += NumCategoriesNotFound;
+		}
+
+		Pair.Value->SetSortOrder(SortOrder);
+	}
+
+	return;
+}
diff --git a/Engine/Source/Editor/DetailCustomizations/Private/SoundBaseDetailsUtils.h b/Engine/Source/Editor/DetailCustomizations/Private/SoundBaseDetailsUtils.h
new file mode 100644
--- /dev/null
+++ b/Engine/Source/Editor/DetailCustomizations/Private/SoundBaseDetailsUtils.h
@@ -0,0 +1,12 @@
+// Copyright Epic Games, Inc. All Rights Reserved.
+
+#pragma once
+
+#include "DetailCategoryBuilder.h"
+
+class FSoundBaseDetailsUtils
+{
+public:
+	static void SortSoundCategories(const TMap<FName, IDetailCategoryBuilder*>& AllCategoryMap);
+	static void SortSoundCategories(const TMap<FName, IDetailCategoryBuilder*>& AllCategoryMap, const TArray<FName> CategoryOrder);
+};
diff --git a/Engine/Source/Editor/DetailCustomizations/Private/SoundCueDetails.cpp b/Engine/Source/Editor/DetailCustomizations/Private/SoundCueDetails.cpp
new file mode 100644
--- /dev/null
+++ b/Engine/Source/Editor/DetailCustomizations/Private/SoundCueDetails.cpp
@@ -0,0 +1,41 @@
+// Copyright Epic Games, Inc. All Rights Reserved.
+
+#include "SoundCueDetails.h"
+
+#include "DetailLayoutBuilder.h"
+#include "DetailCategoryBuilder.h"
+#include "SoundBaseDetailsUtils.h"
+
+#define LOCTEXT_NAMESPACE "FSoundCueDetails"
+
+namespace SoundCueDetailsUtils
+{
+	static const TArray<FName> CategoryOrder =
+	{
+		"Sound",
+		"Attenuation",
+		"Effects",
+		"Voice Management",
+		"AudioProperties",
+		"Memory",
+		"Advanced",
+		"Developer",
+	};
+
+	void SortCategories(const TMap<FName, IDetailCategoryBuilder*>& AllCategoryMap)
+	{
+		FSoundBaseDetailsUtils::SortSoundCategories(AllCategoryMap, SoundCueDetailsUtils::CategoryOrder);
+	}
+} // namespace SoundCueDetailsUtils
+
+TSharedRef<IDetailCustomization> FSoundCueDetails::MakeInstance()
+{
+	return MakeShared<FSoundCueDetails>();
+}
+
+void FSoundCueDetails::CustomizeDetails(IDetailLayoutBuilder& DetailBuilder)
+{
+	DetailBuilder.SortCategories(SoundCueDetailsUtils::SortCategories);
+}
+
+#undef LOCTEXT_NAMESPACE
diff --git a/Engine/Source/Editor/DetailCustomizations/Private/SoundCueDetails.h b/Engine/Source/Editor/DetailCustomizations/Private/SoundCueDetails.h
new file mode 100644
--- /dev/null
+++ b/Engine/Source/Editor/DetailCustomizations/Private/SoundCueDetails.h
@@ -0,0 +1,20 @@
+// Copyright Epic Games, Inc. All Rights Reserved.
+
+#pragma once
+
+#include "IDetailCustomization.h"
+#include "Templates/SharedPointer.h"
+
+class IDetailLayoutBuilder;
+
+class FSoundCueDetails : public IDetailCustomization
+{
+public:
+	/** Makes a new instance of this detail layout class for a specific detail view requesting it */
+	static TSharedRef<IDetailCustomization> MakeInstance();
+
+private:
+	//~ Begin IDetailCustomization interface */
+	virtual void CustomizeDetails(IDetailLayoutBuilder& DetailBuilder) override;
+	//~ End IDetailCustomization interface */
+};
diff --git a/Engine/Source/Editor/DetailCustomizations/Private/SoundWaveDetails.cpp b/Engine/Source/Editor/DetailCustomizations/Private/SoundWaveDetails.cpp
--- a/Engine/Source/Editor/DetailCustomizations/Private/SoundWaveDetails.cpp
+++ b/Engine/Source/Editor/DetailCustomizations/Private/SoundWaveDetails.cpp
@@ -18,6 +18,7 @@
 #include "ScopedTransaction.h"
 #include "SlotBase.h"
 #include "Sound/SoundWave.h"
+#include "SoundBaseDetailsUtils.h"
 #include "ISoundWaveCloudStreaming.h"
 #include "Templates/Casts.h"
 #include "Types/SlateEnums.h"
@@ -38,13 +39,45 @@ class SWidget;
 
 static const FName InternalCurveTableName("InternalCurveTable");
 
+namespace SoundWaveDetailsUtils
+{
+	static const TArray<FName> CategoryOrder =
+	{
+		"SoundWave",
+		"Info",
+		"Format",
+		"Loading",
+		"Subtitles",
+		"Waveform Processing",
+		"Sound",
+		"Attenuation",
+		"Effects",
+		"Modulation",
+		"Voice Management",
+		"Audio Properties",
+		"Curves",
+		"Analysis",
+		"Memory",
+		"Advanced",
+		"Developer",
+		"File Path",
+	};
+
+	void SortCategories(const TMap<FName, IDetailCategoryBuilder*>& AllCategoryMap)
+	{
+		FSoundBaseDetailsUtils::SortSoundCategories(AllCategoryMap, SoundWaveDetailsUtils::CategoryOrder);
+	}
+} // namespace SoundWaveDetailsUtils
+
 TSharedRef<IDetailCustomization> FSoundWaveDetails::MakeInstance()
 {
-	return MakeShareable(new FSoundWaveDetails);
+	return MakeShared<FSoundWaveDetails>();
 }
 
 void FSoundWaveDetails::CustomizeDetails(IDetailLayoutBuilder& DetailBuilder)
 {
+	DetailBuilder.SortCategories(SoundWaveDetailsUtils::SortCategories);
+
 	CustomizeCurveDetails(DetailBuilder);
 	CustomizeCloudStreamingPlatformDetails(DetailBuilder);
 
