Splines - Fix visualizer bugs related to tangent editing/visualization.

diff --git a/Engine/Source/Editor/ComponentVisualizers/Private/SplineComponentVisualizer.cpp b/Engine/Source/Editor/ComponentVisualizers/Private/SplineComponentVisualizer.cpp
--- a/Engine/Source/Editor/ComponentVisualizers/Private/SplineComponentVisualizer.cpp
+++ b/Engine/Source/Editor/ComponentVisualizers/Private/SplineComponentVisualizer.cpp
@@ -680,8 +680,7 @@ void FSplineComponentVisualizer::DrawVisualization(const UActorComponent* Compon
 
 						const FVector Location = SplineComp->GetLocationAtSplinePoint(SelectedKey, ESplineCoordinateSpace::World);
 						const FVector LeaveTangent = SplineComp->GetLeaveTangentAtSplinePoint(SelectedKey, ESplineCoordinateSpace::World) * TangentScale;
-						const FVector ArriveTangent = SplineComp->bAllowDiscontinuousSpline ?
-							SplineComp->GetArriveTangentAtSplinePoint(SelectedKey, ESplineCoordinateSpace::World) * TangentScale : LeaveTangent;
+						const FVector ArriveTangent = SplineComp->GetArriveTangentAtSplinePoint(SelectedKey, ESplineCoordinateSpace::World) * TangentScale;
 
 						PDI->SetHitProxy(NULL);
 
diff --git a/Engine/Source/Editor/DetailCustomizations/Private/SplineComponentDetails.cpp b/Engine/Source/Editor/DetailCustomizations/Private/SplineComponentDetails.cpp
--- a/Engine/Source/Editor/DetailCustomizations/Private/SplineComponentDetails.cpp
+++ b/Engine/Source/Editor/DetailCustomizations/Private/SplineComponentDetails.cpp
@@ -1059,9 +1059,8 @@ void FSplinePointDetails::OnSetArriveTangent(FVector::FReal NewValue, ETextCommi
 			}
 
 			FVector LocalArriveTangent = SplineComp->GetArriveTangentAtSplinePoint(Index, ESplineCoordinateSpace::Local);
-			FVector LocalLeaveTangent = SplineComp->GetLeaveTangentAtSplinePoint(Index, ESplineCoordinateSpace::Local);
-
 			LocalArriveTangent.SetComponentForAxis(Axis, NewValue);
+			FVector LocalLeaveTangent = SplineComp->bAllowDiscontinuousSpline ? SplineComp->GetLeaveTangentAtSplinePoint(Index, ESplineCoordinateSpace::Local) : LocalArriveTangent;
 			
 			SplineComp->SetTangentsAtSplinePoint(Index, LocalArriveTangent, LocalLeaveTangent, ESplineCoordinateSpace::Local, false);
 
@@ -1106,10 +1105,9 @@ void FSplinePointDetails::OnSetLeaveTangent(FVector::FReal NewValue, ETextCommit
 				continue;
 			}
 
-			FVector LocalArriveTangent = SplineComp->GetArriveTangentAtSplinePoint(Index, ESplineCoordinateSpace::Local);
 			FVector LocalLeaveTangent = SplineComp->GetLeaveTangentAtSplinePoint(Index, ESplineCoordinateSpace::Local);
-
 			LocalLeaveTangent.SetComponentForAxis(Axis, NewValue);
+			FVector LocalArriveTangent = SplineComp->bAllowDiscontinuousSpline ? SplineComp->GetArriveTangentAtSplinePoint(Index, ESplineCoordinateSpace::Local) : LocalLeaveTangent;
 			
 			SplineComp->SetTangentsAtSplinePoint(Index, LocalArriveTangent, LocalLeaveTangent, ESplineCoordinateSpace::Local, false);
 		}
