SplineComponentDetails: Remove checks if the spline visualizer cannot be found.

Without the visualizer, the customizaton isn't extremely useful but it will eventually be hooked into the new component visualizer. We also want the ability to turn the current visualizer without unregistering the details customization which requires that the customization can 'work' without the visualizer.

diff --git a/Engine/Source/Editor/DetailCustomizations/Private/SplineComponentDetails.cpp b/Engine/Source/Editor/DetailCustomizations/Private/SplineComponentDetails.cpp
--- a/Engine/Source/Editor/DetailCustomizations/Private/SplineComponentDetails.cpp
+++ b/Engine/Source/Editor/DetailCustomizations/Private/SplineComponentDetails.cpp
@@ -337,9 +337,10 @@ bool FSplinePointDetails::bAlreadyWarnedInvalidIndex = false;
 FSplinePointDetails::FSplinePointDetails(USplineComponent* InOwningSplineComponent)
 	: SplineComp(nullptr)
 {
-	TSharedPtr<FComponentVisualizer> Visualizer = GUnrealEd->FindComponentVisualizer(InOwningSplineComponent->GetClass());
-	SplineVisualizer = StaticCastSharedPtr<FSplineComponentVisualizer>(Visualizer);
-	check(SplineVisualizer.IsValid());
+	if (TSharedPtr<FComponentVisualizer> Visualizer = GUnrealEd->FindComponentVisualizer(InOwningSplineComponent->GetClass()))
+	{
+		SplineVisualizer = StaticCastSharedPtr<FSplineComponentVisualizer>(Visualizer);
+	}
 
 	for (const FName& Property : USplineComponent::GetSplinePropertyNames())
 	{
@@ -1189,7 +1190,10 @@ void FSplinePointDetails::OnSetRotation(FRotator::FReal NewValue, ETextCommit::T
 		}
 	}
 
-	SplineVisualizer->SetCachedRotation(NewRotationRelative);
+	if (SplineVisualizer.IsValid())
+	{
+		SplineVisualizer->SetCachedRotation(NewRotationRelative);
+	}
 
 	if (CommitInfo == ETextCommit::OnEnter || CommitInfo == ETextCommit::OnUserMovedFocus)
 	{
