Details Panel: Allow customizations to override the copy behavior for categories to match how pasting categories is implemented - if an override is not provided it will fallback to copying via the property handle
NOTE: Not every customization using the paste override has been updated to also use the copy override since copy has a fallback to the default implementation and does not need an override unless there is some custom behavior it needs to handle

ComponentTransformDetails: override copy on the category to make sure the copied values are swizzled correctly

diff --git a/Engine/Source/Editor/DetailCustomizations/Private/ComponentTransformDetails.cpp b/Engine/Source/Editor/DetailCustomizations/Private/ComponentTransformDetails.cpp
--- a/Engine/Source/Editor/DetailCustomizations/Private/ComponentTransformDetails.cpp
+++ b/Engine/Source/Editor/DetailCustomizations/Private/ComponentTransformDetails.cpp
@@ -306,7 +306,32 @@ void FComponentTransformDetails::OnCopy( ETransformField::Type TransformField )
 	CacheDetails();
 
 	FString CopyStr;
-	switch (TransformField)
+
+	CopyToText(CopyStr, TransformField);
+
+	if( !CopyStr.IsEmpty() )
+	{
+		FPlatformApplicationMisc::ClipboardCopy( *CopyStr );
+	}
+}
+
+void FComponentTransformDetails::OnCopyToText(TSharedPtr<IPropertyHandle> PropertyHandle, FString& Text, const TOptional<FGuid>& InOperationId,
+	ETransformField::Type InTransformField)
+{
+	const FString PropertyPath = UE::PropertyEditor::GetPropertyPath(PropertyHandle);
+
+	// ensure that if tag is specified, that it matches the subscriber
+	if (!PropertyPath.Equals(UE::DetailsCustomizations::Internal::TransformFieldToPropertyNameString[InTransformField]))
+	{
+		return;
+	}
+
+	CopyToText(Text, InTransformField);
+}
+
+void FComponentTransformDetails::CopyToText(FString& CopyStr, ETransformField::Type InTransformField)
+{
+	switch (InTransformField)
 	{
 	case ETransformField::Location:
 		CopyStr = FString::Printf(TEXT("(X=%f,Y=%f,Z=%f)"), GetLocationX().GetValue(), GetLocationY().GetValue(), GetLocationZ().GetValue());
@@ -321,10 +346,6 @@ void FComponentTransformDetails::OnCopy( ETransformField::Type TransformField )
 		break;
 	}
 
-	if( !CopyStr.IsEmpty() )
-	{
-		FPlatformApplicationMisc::ClipboardCopy( *CopyStr );
-	}
 }
 
 void FComponentTransformDetails::OnPaste( ETransformField::Type TransformField )
@@ -462,6 +483,7 @@ void FComponentTransformDetails::GenerateChildContent( IDetailChildrenBuilder& C
 		}
 
 		ParentCategory.OnPasteFromText()->AddSP(this, &FComponentTransformDetails::OnPasteFromText, ETransformField::Location);
+		ParentCategory.OnCopyToText()->AddSP(this, &FComponentTransformDetails::OnCopyToText, ETransformField::Location);
 
 		FindOrCreatePropertyHandle(USceneComponent::GetAbsoluteLocationPropertyName(), ChildrenBuilder);
 
@@ -525,6 +547,7 @@ void FComponentTransformDetails::GenerateChildContent( IDetailChildrenBuilder& C
 		}
 
 		ParentCategory.OnPasteFromText()->AddSP(this, &FComponentTransformDetails::OnPasteFromText, ETransformField::Rotation);
+		ParentCategory.OnCopyToText()->AddSP(this, &FComponentTransformDetails::OnCopyToText, ETransformField::Rotation);
 
 		FindOrCreatePropertyHandle(USceneComponent::GetAbsoluteRotationPropertyName(), ChildrenBuilder);
 
@@ -583,7 +606,8 @@ void FComponentTransformDetails::GenerateChildContent( IDetailChildrenBuilder& C
 	if(!bHideScaleField)
 	{
 		ParentCategory.OnPasteFromText()->AddSP(this, &FComponentTransformDetails::OnPasteFromText, ETransformField::Scale);
-		
+		ParentCategory.OnCopyToText()->AddSP(this, &FComponentTransformDetails::OnCopyToText, ETransformField::Scale);
+
 		FindOrCreatePropertyHandle(USceneComponent::GetAbsoluteScalePropertyName(), ChildrenBuilder);
 
 		TSharedPtr<IPropertyHandle> PropertyHandle = FindOrCreatePropertyHandle(USceneComponent::GetRelativeScale3DPropertyName(), ChildrenBuilder);
diff --git a/Engine/Source/Editor/DetailCustomizations/Private/ComponentTransformDetails.h b/Engine/Source/Editor/DetailCustomizations/Private/ComponentTransformDetails.h
--- a/Engine/Source/Editor/DetailCustomizations/Private/ComponentTransformDetails.h
+++ b/Engine/Source/Editor/DetailCustomizations/Private/ComponentTransformDetails.h
@@ -187,6 +187,11 @@ private:
 	 * Copies the specified transform field to the clipboard
 	 */
 	void OnCopy( ETransformField::Type TransformField );
+	
+	void OnCopyToText(TSharedPtr<IPropertyHandle> PropertyHandle, FString& Text, const TOptional<FGuid>& InOperationId, ETransformField::Type InTransformField);
+
+	void CopyToText(FString& Text, ETransformField::Type InTransformField);
+
 
 	/**
 	 * Pastes the specified transform field from the clipboard
diff --git a/Engine/Source/Editor/PropertyEditor/Private/CustomChildBuilder.cpp b/Engine/Source/Editor/PropertyEditor/Private/CustomChildBuilder.cpp
--- a/Engine/Source/Editor/PropertyEditor/Private/CustomChildBuilder.cpp
+++ b/Engine/Source/Editor/PropertyEditor/Private/CustomChildBuilder.cpp
@@ -39,6 +39,11 @@ FDetailWidgetRow& FCustomChildrenBuilder::AddCustomRow( const FText& SearchStrin
 	{
 		NewRow->OnPasteFromTextDelegate = PasteFromTextDelegate;
 	}
+
+	if (const TSharedPtr<FOnCopyToText> CopyToTextDelegate = GetParentCategory().OnCopyToText())
+	{
+		NewRow->OnCopyToTextDelegate = CopyToTextDelegate;
+	}
 	
 	NewCustomization.WidgetDecl = NewRow;
 
diff --git a/Engine/Source/Editor/PropertyEditor/Private/DetailCategoryBuilderImpl.cpp b/Engine/Source/Editor/PropertyEditor/Private/DetailCategoryBuilderImpl.cpp
--- a/Engine/Source/Editor/PropertyEditor/Private/DetailCategoryBuilderImpl.cpp
+++ b/Engine/Source/Editor/PropertyEditor/Private/DetailCategoryBuilderImpl.cpp
@@ -209,6 +209,7 @@ FDetailCategoryImpl::FDetailCategoryImpl(FName InCategoryName, TSharedRef<FDetai
 	: HeaderContentWidget(nullptr)
 	, DetailLayoutBuilder(InDetailLayout)
 	, PasteFromTextDelegate(MakeShared<FOnPasteFromText>())
+	, CopyToTextDelegate(MakeShared<FOnCopyToText>())
 	, CategoryName(InCategoryName)
 	, SortOrder(0)
 	, bRestoreExpansionState(!ContainsOnlyAdvanced())
@@ -928,6 +929,7 @@ TSharedRef<ITableRow> FDetailCategoryImpl::GenerateWidgetForTableView(const TSha
 	TSharedRef<SDetailCategoryTableRow> RowWidget =
 		SNew(SDetailCategoryTableRow, AsShared(), OwnerTable)
 		.PasteFromText(OnPasteFromText())
+		.CopyToText(OnCopyToText())
 		.Objects(&Objects)
 		.IsEmpty( bIsEmpty )
 		.InnerCategory(ParentLayout.IsValid() ? ParentLayout->IsLayoutForExternalRoot() : false)
diff --git a/Engine/Source/Editor/PropertyEditor/Private/DetailCategoryBuilderImpl.h b/Engine/Source/Editor/PropertyEditor/Private/DetailCategoryBuilderImpl.h
--- a/Engine/Source/Editor/PropertyEditor/Private/DetailCategoryBuilderImpl.h
+++ b/Engine/Source/Editor/PropertyEditor/Private/DetailCategoryBuilderImpl.h
@@ -195,6 +195,10 @@ public:
 	/** Delegate handling pasting an optionally tagged text snippet */
 	virtual TSharedPtr<FOnPasteFromText> OnPasteFromText() const override { return PasteFromTextDelegate; }
 
+	/** Delegate handling copying a property handle to text. The default property handle impl for copying will be used if the delegate is not bound
+	 * or returns false for a property */
+	virtual TSharedPtr<FOnCopyToText> OnCopyToText() const override { return CopyToTextDelegate; }
+
 	/** FDetailTreeNode interface */
 	virtual TSharedPtr<IDetailsView> GetNodeDetailsViewSharedPtr() const override { return GetDetailsViewSharedPtr(); }
 	virtual TSharedPtr<IDetailsViewPrivate> GetDetailsViewSharedPtr() const override;
@@ -441,6 +445,8 @@ private:
 	TWeakPtr<FDetailLayoutBuilderImpl> DetailLayoutBuilder;
 	/** Delegate handling pasting an optionally tagged text snippet */
 	TSharedPtr<FOnPasteFromText> PasteFromTextDelegate;
+	/** Delegate handling copying a property to text. The default property handle copy will be used if not provided */
+	TSharedPtr<FOnCopyToText> CopyToTextDelegate;
 	
 	/** The category identifier */
 	FName CategoryName;
diff --git a/Engine/Source/Editor/PropertyEditor/Private/DetailGroup.cpp b/Engine/Source/Editor/PropertyEditor/Private/DetailGroup.cpp
--- a/Engine/Source/Editor/PropertyEditor/Private/DetailGroup.cpp
+++ b/Engine/Source/Editor/PropertyEditor/Private/DetailGroup.cpp
@@ -16,6 +16,7 @@ FDetailGroup::FDetailGroup( const FName InGroupName, TSharedRef<FDetailCategoryI
 	, bResetEnabled( false )
 	, DisplayMode( EDetailGroupDisplayMode::Group )
 	, PasteFromTextDelegate(MakeShared<FOnPasteFromText>())
+	, CopyToTextDelegate(MakeShared<FOnCopyToText>())
 {
 }
 
diff --git a/Engine/Source/Editor/PropertyEditor/Private/DetailGroup.h b/Engine/Source/Editor/PropertyEditor/Private/DetailGroup.h
--- a/Engine/Source/Editor/PropertyEditor/Private/DetailGroup.h
+++ b/Engine/Source/Editor/PropertyEditor/Private/DetailGroup.h
@@ -103,6 +103,12 @@ public:
 	 */
 	virtual TSharedPtr<FOnPasteFromText> OnPasteFromText() const override { return PasteFromTextDelegate; }
 
+	/**
+	* Return an optional PasteFromText delegate
+	*/
+	virtual TSharedPtr<FOnCopyToText> OnCopyToText() const override { return CopyToTextDelegate; }
+
+
 private:
 	/**
 	 * Called when the name of the group is clicked to expand the group
@@ -147,4 +153,7 @@ private:
 	FDetailGroupReset OnDetailGroupReset;
 	/** Delegate handling pasting an optionally tagged text snippet */
 	TSharedPtr<FOnPasteFromText> PasteFromTextDelegate;
+	/** Optional Delegate handling copying to a text snippet */
+	TSharedPtr<FOnCopyToText> CopyToTextDelegate;
+
 };
diff --git a/Engine/Source/Editor/PropertyEditor/Private/DetailPropertyRow.cpp b/Engine/Source/Editor/PropertyEditor/Private/DetailPropertyRow.cpp
--- a/Engine/Source/Editor/PropertyEditor/Private/DetailPropertyRow.cpp
+++ b/Engine/Source/Editor/PropertyEditor/Private/DetailPropertyRow.cpp
@@ -994,6 +994,7 @@ void FDetailPropertyRow::SetWidgetRowProperties(FDetailWidgetRow& Row) const
 		Row.CopyMenuAction = CustomPropertyWidget->CopyMenuAction;
 		Row.PasteMenuAction = CustomPropertyWidget->PasteMenuAction;
 		Row.CustomMenuItems = CustomPropertyWidget->CustomMenuItems;
+		Row.OnCopyToTextDelegate = CustomPropertyWidget->OnCopyToTextDelegate;
         Row.OnPasteFromTextDelegate = CustomPropertyWidget->OnPasteFromTextDelegate;
 		Row.FilterTextString = CustomPropertyWidget->FilterTextString;
 
diff --git a/Engine/Source/Editor/PropertyEditor/Private/SDetailCategoryTableRow.cpp b/Engine/Source/Editor/PropertyEditor/Private/SDetailCategoryTableRow.cpp
--- a/Engine/Source/Editor/PropertyEditor/Private/SDetailCategoryTableRow.cpp
+++ b/Engine/Source/Editor/PropertyEditor/Private/SDetailCategoryTableRow.cpp
@@ -52,6 +52,10 @@ void SDetailCategoryTableRow::Construct(const FArguments& InArgs, TSharedRef<FDe
 	{
 		OnPasteFromTextDelegate = InArgs._PasteFromText;
 	}
+	if (InArgs._CopyToText.IsValid())
+	{
+		OnCopyToTextDelegate = InArgs._CopyToText;
+	}
 	
 	PasteAction.ExecuteAction = FExecuteAction::CreateSP(this, &SDetailCategoryTableRow::OnPasteCategory);
 	PasteAction.CanExecuteAction = FCanExecuteAction::CreateSP(this, &SDetailCategoryTableRow::CanPasteCategory);
@@ -372,7 +376,9 @@ void SDetailCategoryTableRow::OnCopyCategory()
 		
 		TMap<FString, FString> PropertyValues;
 		PropertyValues.Reserve(CategoryProperties.Num());
-		
+
+		const FGuid OperationGuid = FGuid::NewGuid();			
+
 		for (TSharedPtr<IPropertyHandle> PropertyHandle : CategoryProperties)
 		{
 			if (PropertyHandle.IsValid() && PropertyHandle->IsValidHandle())
@@ -380,7 +386,18 @@ void SDetailCategoryTableRow::OnCopyCategory()
 				FString PropertyPath = UE::PropertyEditor::GetPropertyPath(PropertyHandle);
 				
 				FString PropertyValueStr;
-				if (PropertyHandle->GetValueAsFormattedString(PropertyValueStr, PPF_Copy) == FPropertyAccess::Success)
+
+				// First check if there is some custom copy paste that handles this property
+				if (OnCopyToTextDelegate.IsValid())
+				{
+					OnCopyToTextDelegate->Broadcast(PropertyHandle, PropertyValueStr, OperationGuid);
+				}
+
+				if (!PropertyValueStr.IsEmpty())
+				{
+					PropertyValues.Add(PropertyPath, PropertyValueStr);
+				}
+				else if (PropertyHandle->GetValueAsFormattedString(PropertyValueStr, PPF_Copy) == FPropertyAccess::Success)
 				{
 					PropertyValues.Add(PropertyPath, PropertyValueStr);
 				}
diff --git a/Engine/Source/Editor/PropertyEditor/Private/SDetailCategoryTableRow.h b/Engine/Source/Editor/PropertyEditor/Private/SDetailCategoryTableRow.h
--- a/Engine/Source/Editor/PropertyEditor/Private/SDetailCategoryTableRow.h
+++ b/Engine/Source/Editor/PropertyEditor/Private/SDetailCategoryTableRow.h
@@ -34,6 +34,7 @@ public:
 		* empty stub with no expansion arrow */
 		SLATE_ARGUMENT(bool, IsEmpty)
 		SLATE_ARGUMENT(TSharedPtr<FOnPasteFromText>, PasteFromText)
+		SLATE_ARGUMENT(TSharedPtr<FOnCopyToText>, CopyToText)
 		SLATE_ARGUMENT(TArray<TWeakObjectPtr<UObject>>*, Objects)
 	SLATE_END_ARGS()
 
@@ -100,6 +101,9 @@ private:
 	/** Delegate handling pasting an optionally tagged text snippet */
 	TSharedPtr<FOnPasteFromText> OnPasteFromTextDelegate;
 
+	/** Delegate handling pasting an optionally tagged text snippet */
+	TSharedPtr<FOnCopyToText> OnCopyToTextDelegate;
+
 	/** Animation curve for displaying pulse */
 	FCurveSequence PulseAnimation;
 };
\ No newline at end of file
diff --git a/Engine/Source/Editor/PropertyEditor/Private/SDetailSingleItemRow.cpp b/Engine/Source/Editor/PropertyEditor/Private/SDetailSingleItemRow.cpp
--- a/Engine/Source/Editor/PropertyEditor/Private/SDetailSingleItemRow.cpp
+++ b/Engine/Source/Editor/PropertyEditor/Private/SDetailSingleItemRow.cpp
@@ -369,6 +369,23 @@ void SDetailSingleItemRow::Construct( const FArguments& InArgs, FDetailLayoutCus
 					}
 				}
 
+				if (WidgetRow.IsCopyToTextBound())
+				{
+					OnCopyToTextDelegate = WidgetRow.OnCopyToTextDelegate.Pin();					
+				}
+				else if (Category.IsValid())
+				{
+					OnCopyToTextDelegate = Category->OnCopyToText();
+				}
+				// If still not set, but this is a group, initialize
+				else if (Group.IsValid())
+				{
+					if (TSharedPtr<FOnCopyToText> GroupOnCopyToTextDelegate = Group->OnCopyToText())
+					{
+						OnCopyToTextDelegate = GroupOnCopyToTextDelegate;
+					}
+				}
+				
 				if (WidgetRow.IsPasteFromTextBound())
 				{
 					OnPasteFromTextDelegate = WidgetRow.OnPasteFromTextDelegate.Pin();					
@@ -386,6 +403,8 @@ void SDetailSingleItemRow::Construct( const FArguments& InArgs, FDetailLayoutCus
 					}
 				}
 
+				// Note: We currently only animate the row on paste but if there is other custom behavior in OnPasteFromText we should
+				// also bind to OnCopyToText to match it
 				if (OnPasteFromTextDelegate.IsValid())
 				{
 					OnPasteFromTextDelegate->AddSP(this, &SDetailSingleItemRow::OnPasteFromText);
@@ -969,6 +988,8 @@ void SDetailSingleItemRow::OnCopyGroup()
 		
 		TMap<FString, FString> PropertyValues;
 		PropertyValues.Reserve(GroupProperties.Num());
+
+		const FGuid OperationGuid = FGuid::NewGuid();			
 		
 		for (TSharedPtr<IPropertyHandle> PropertyHandle : GroupProperties)
 		{
@@ -977,7 +998,18 @@ void SDetailSingleItemRow::OnCopyGroup()
 				FString PropertyPath = UE::PropertyEditor::GetPropertyPath(PropertyHandle);
 				
 				FString PropertyValueStr;
-				if (PropertyHandle->GetValueAsFormattedString(PropertyValueStr, PPF_Copy) == FPropertyAccess::Success)
+
+				// First check if there is some custom copy paste that handles this property
+				if (OnCopyToTextDelegate.IsValid())
+				{
+					OnCopyToTextDelegate->Broadcast(PropertyHandle, PropertyValueStr, OperationGuid);
+				}
+				
+				if (!PropertyValueStr.IsEmpty())
+				{
+					PropertyValues.Add(PropertyPath, PropertyValueStr);
+				}
+				else if (PropertyHandle->GetValueAsFormattedString(PropertyValueStr, PPF_Copy) == FPropertyAccess::Success)
 				{
 					PropertyValues.Add(PropertyPath, PropertyValueStr);
 				}
diff --git a/Engine/Source/Editor/PropertyEditor/Private/SDetailSingleItemRow.h b/Engine/Source/Editor/PropertyEditor/Private/SDetailSingleItemRow.h
--- a/Engine/Source/Editor/PropertyEditor/Private/SDetailSingleItemRow.h
+++ b/Engine/Source/Editor/PropertyEditor/Private/SDetailSingleItemRow.h
@@ -172,6 +172,9 @@ private:
 	/** Delegate handling pasting an optionally tagged text snippet */
 	TSharedPtr<FOnPasteFromText> OnPasteFromTextDelegate;
 
+	/** Delegate handling copying to a text snippet */
+	TSharedPtr<FOnCopyToText> OnCopyToTextDelegate;
+
 	/** Animation curve for displaying pulse */
 	FCurveSequence PulseAnimation;
 };
diff --git a/Engine/Source/Editor/PropertyEditor/Public/DetailCategoryBuilder.h b/Engine/Source/Editor/PropertyEditor/Public/DetailCategoryBuilder.h
--- a/Engine/Source/Editor/PropertyEditor/Public/DetailCategoryBuilder.h
+++ b/Engine/Source/Editor/PropertyEditor/Public/DetailCategoryBuilder.h
@@ -229,6 +229,9 @@ public:
 	/** Optional PasteFromText delegate for this category */
 	virtual TSharedPtr<FOnPasteFromText> OnPasteFromText() const { return nullptr; }
 
+	/** Optional CopyToText delegate for this category */
+	virtual TSharedPtr<FOnCopyToText> OnCopyToText() const { return nullptr; }
+
 	/**
 	 * If true, this Category should have no UProperty data associated with it, and will be shown as an empty stub
 	 * with no expansion arrow
diff --git a/Engine/Source/Editor/PropertyEditor/Public/DetailWidgetRow.h b/Engine/Source/Editor/PropertyEditor/Public/DetailWidgetRow.h
--- a/Engine/Source/Editor/PropertyEditor/Public/DetailWidgetRow.h
+++ b/Engine/Source/Editor/PropertyEditor/Public/DetailWidgetRow.h
@@ -310,6 +310,12 @@ public:
 	{
 		return OnPasteFromTextDelegate.IsValid() && OnPasteFromTextDelegate.Pin()->IsBoundToObject(this);
 	}
+	
+	/** @return true if a custom copy text is bound on this row */
+	bool IsCopyToTextBound() const
+	{
+		return OnCopyToTextDelegate.IsValid() && OnCopyToTextDelegate.Pin()->IsBoundToObject(this);
+	}
 
 	/**
 	* Sets a tag which can be used to identify this row 
@@ -398,6 +404,8 @@ public:
 	FUIAction PasteMenuAction;	
 	/** Delegate for pasting (optionally tagged) text on this row */
 	TWeakPtr<FOnPasteFromText> OnPasteFromTextDelegate;
+	/** Optional Delegate for copying text on this row (default property handle copy used if unbound) */
+	TWeakPtr<FOnCopyToText> OnCopyToTextDelegate;
 
 	struct FCustomMenuData
 	{
diff --git a/Engine/Source/Editor/PropertyEditor/Public/IDetailGroup.h b/Engine/Source/Editor/PropertyEditor/Public/IDetailGroup.h
--- a/Engine/Source/Editor/PropertyEditor/Public/IDetailGroup.h
+++ b/Engine/Source/Editor/PropertyEditor/Public/IDetailGroup.h
@@ -121,4 +121,8 @@ public:
 
 	/** Return an optional PasteFromText delegate */
 	virtual TSharedPtr<FOnPasteFromText> OnPasteFromText() const { return nullptr; }
+	
+	/** Return an optional CopyToText delegate */
+	virtual TSharedPtr<FOnCopyToText> OnCopyToText() const { return nullptr; }
+
 };
diff --git a/Engine/Source/Editor/PropertyEditor/Public/PropertyEditorCopyPaste.h b/Engine/Source/Editor/PropertyEditor/Public/PropertyEditorCopyPaste.h
--- a/Engine/Source/Editor/PropertyEditor/Public/PropertyEditorCopyPaste.h
+++ b/Engine/Source/Editor/PropertyEditor/Public/PropertyEditorCopyPaste.h
@@ -18,3 +18,11 @@ namespace UE::PropertyEditor
 /** Subscribe to and parse the string contents. Note that the caller is responsible for providing valid UE text,
  * and the subscriber is responsible for applying it if applicable.  */
 typedef TMulticastDelegate<void(const FString& InTag, const FString& InText, const TOptional<FGuid>& InOperationId)> FOnPasteFromText;
+
+/** Subscribe to and copy the property value to text. Optionally required when there is custom behavior needed on copy to match FOnPasteFromText,
+ * otherwise the default property handle copy will be used. Note that the caller is responsible for providing a valid property handle, and the
+ * subscriber is responsible for copying it to text if applicable. OutPropertyText should be left empty if the subscriber does not handle
+ * the given property.
+ */
+typedef TMulticastDelegate<void(TSharedPtr<IPropertyHandle> InPropertyHandle, FString& OutPropertyText, const TOptional<FGuid>& InOperationId)> FOnCopyToText;
+
