Allow a UFont to be a wrapper around FCoreStyle::GetDefaultFont(), for use with editor UI

Deprecation: Public access to UFont::CompositeFont is now deprecated.
* Use UFont::GetCompositeFont() if you want to resolve the FCompositeFont to use at runtime (if any).
* Use UFont::GetInternalCompositeFont() or UFont::GetMutableInternalCompositeFont() if you need to specifically access/modify UFont::CompositeFont.

#virtualized

diff --git a/Engine/Plugins/Experimental/Text3D/Source/Text3D/Private/Text3DComponent.cpp b/Engine/Plugins/Experimental/Text3D/Source/Text3D/Private/Text3DComponent.cpp
--- a/Engine/Plugins/Experimental/Text3D/Source/Text3D/Private/Text3DComponent.cpp
+++ b/Engine/Plugins/Experimental/Text3D/Source/Text3D/Private/Text3DComponent.cpp
@@ -115,7 +115,7 @@ int32 UText3DComponent::GetTypeFaceIndex() const
 
 	if (Font)
 	{
-		TArray<FTypefaceEntry>& Fonts = Font->CompositeFont.DefaultTypeface.Fonts;
+		TArrayView<const FTypefaceEntry> Fonts = GetAvailableTypefaces();
 
 		for (int32 Index = 0; Index < Fonts.Num(); Index++)
 		{
@@ -134,13 +134,13 @@ const FTypefaceEntry* UText3DComponent::GetTypeFaceEntry() const
 {
 	if (Font)
 	{
-		TArray<FTypefaceEntry>& TypefaceEntries = Font->CompositeFont.DefaultTypeface.Fonts;
+		TArrayView<const FTypefaceEntry> Fonts = GetAvailableTypefaces();
 
-		for (int32 Index = 0; Index < TypefaceEntries.Num(); Index++)
+		for (int32 Index = 0; Index < Fonts.Num(); Index++)
 		{
-			if (Typeface == TypefaceEntries[Index].Name)
+			if (Typeface == Fonts[Index].Name)
 			{
-				return &TypefaceEntries[Index];
+				return &Fonts[Index];
 			}
 		}
 	}
@@ -161,11 +161,14 @@ bool UText3DComponent::IsTypefaceAvailable(FName InTypeface) const
 	return false;
 }
 
-TArray<FTypefaceEntry> UText3DComponent::GetAvailableTypefaces() const
+TArrayView<const FTypefaceEntry> UText3DComponent::GetAvailableTypefaces() const
 {
 	if (Font)
 	{
-		return Font->CompositeFont.DefaultTypeface.Fonts;
+		if (const FCompositeFont* CompositeFont = Font->GetCompositeFont())
+		{
+			return CompositeFont->DefaultTypeface.Fonts;
+		}
 	}
 
 	return {};
@@ -175,7 +178,7 @@ void UText3DComponent::RefreshTypeface()
 {
 	if (Font)
 	{
-		TArray<FTypefaceEntry>& Fonts = Font->CompositeFont.DefaultTypeface.Fonts;
+		TArrayView<const FTypefaceEntry> Fonts = GetAvailableTypefaces();
 		for (int32 Index = 0; Index < Fonts.Num(); Index++)
 		{
 			if (Typeface == Fonts[Index].Name)
@@ -1089,7 +1092,7 @@ TArray<FName> UText3DComponent::GetTypefaceNames() const
 
 	if (Font)
 	{
-		for (const FTypefaceEntry& TypeFaceFont : Font->CompositeFont.DefaultTypeface.Fonts)
+		for (const FTypefaceEntry& TypeFaceFont : GetAvailableTypefaces())
 		{
 			TypefaceNames.Add(TypeFaceFont.Name);
 		}
diff --git a/Engine/Plugins/Experimental/Text3D/Source/Text3D/Public/Text3DComponent.h b/Engine/Plugins/Experimental/Text3D/Source/Text3D/Public/Text3DComponent.h
--- a/Engine/Plugins/Experimental/Text3D/Source/Text3D/Public/Text3DComponent.h
+++ b/Engine/Plugins/Experimental/Text3D/Source/Text3D/Public/Text3DComponent.h
@@ -497,7 +497,7 @@ private:
 	void Rebuild(EText3DRendererFlags InUpdateFlags, bool bCleanCache);
 
 	bool IsTypefaceAvailable(FName InTypeface) const;
-	TArray<FTypefaceEntry> GetAvailableTypefaces() const;
+	TArrayView<const FTypefaceEntry> GetAvailableTypefaces() const;
 	void RefreshTypeface();
 
 	void UpdateStatistics();
diff --git a/Engine/Source/Editor/DetailCustomizations/Private/SlateFontInfoCustomization.cpp b/Engine/Source/Editor/DetailCustomizations/Private/SlateFontInfoCustomization.cpp
--- a/Engine/Source/Editor/DetailCustomizations/Private/SlateFontInfoCustomization.cpp
+++ b/Engine/Source/Editor/DetailCustomizations/Private/SlateFontInfoCustomization.cpp
@@ -323,7 +323,8 @@ bool FSlateFontInfoStructCustomization::OnFilterFontAsset(const FAssetData& InAs
 void FSlateFontInfoStructCustomization::OnFontChanged(const FAssetData& InAssetData)
 {
 	const UFont* const FontAsset = Cast<const UFont>(InAssetData.GetAsset());
-	const FName FirstFontName = (FontAsset && FontAsset->CompositeFont.DefaultTypeface.Fonts.Num()) ? FontAsset->CompositeFont.DefaultTypeface.Fonts[0].Name : NAME_None;
+	const FCompositeFont* CompositeFont = FontAsset ? FontAsset->GetCompositeFont() : nullptr;
+	const FName FirstFontName = (CompositeFont && CompositeFont->DefaultTypeface.Fonts.Num()) ? CompositeFont->DefaultTypeface.Fonts[0].Name : NAME_None;
 
 	TArray<FSlateFontInfo*> SlateFontInfoStructs = GetFontInfoBeingEdited();
 	for(FSlateFontInfo* FontInfo : SlateFontInfoStructs)
@@ -396,14 +397,17 @@ void FSlateFontInfoStructCustomization::OnFontEntryComboOpening()
 		const FName ActiveFontEntry = GetActiveFontEntry();
 		TSharedPtr<FName> SelectedNamePtr;
 
-		for(const FTypefaceEntry& TypefaceEntry : FontObject->CompositeFont.DefaultTypeface.Fonts)
+		if (const FCompositeFont* CompositeFont = FontObject->GetCompositeFont())
 		{
-			TSharedPtr<FName> NameEntryPtr = MakeShareable(new FName(TypefaceEntry.Name));
-			FontEntryComboData.Add(NameEntryPtr);
-
-			if(!TypefaceEntry.Name.IsNone() && TypefaceEntry.Name == ActiveFontEntry)
+			for (const FTypefaceEntry& TypefaceEntry : CompositeFont->DefaultTypeface.Fonts)
 			{
-				SelectedNamePtr = NameEntryPtr;
+				TSharedPtr<FName> NameEntryPtr = MakeShareable(new FName(TypefaceEntry.Name));
+				FontEntryComboData.Add(NameEntryPtr);
+
+				if (!TypefaceEntry.Name.IsNone() && TypefaceEntry.Name == ActiveFontEntry)
+				{
+					SelectedNamePtr = NameEntryPtr;
+				}
 			}
 		}
 
@@ -461,9 +465,12 @@ FName FSlateFontInfoStructCustomization::GetActiveFontEntry() const
 		const UFont* const FontObject = Cast<const UFont>(FirstSlateFontInfo->FontObject);
 		if(FontObject)
 		{
-			return (FirstSlateFontInfo->TypefaceFontName.IsNone() && FontObject->CompositeFont.DefaultTypeface.Fonts.Num())
-				? FontObject->CompositeFont.DefaultTypeface.Fonts[0].Name
-				: FirstSlateFontInfo->TypefaceFontName;
+			if (const FCompositeFont* CompositeFont = FontObject->GetCompositeFont())
+			{
+				return (FirstSlateFontInfo->TypefaceFontName.IsNone() && CompositeFont->DefaultTypeface.Fonts.Num())
+					? CompositeFont->DefaultTypeface.Fonts[0].Name
+					: FirstSlateFontInfo->TypefaceFontName;
+			}
 		}
 	}
 
diff --git a/Engine/Source/Editor/FontEditor/Private/FontEditor.cpp b/Engine/Source/Editor/FontEditor/Private/FontEditor.cpp
--- a/Engine/Source/Editor/FontEditor/Private/FontEditor.cpp
+++ b/Engine/Source/Editor/FontEditor/Private/FontEditor.cpp
@@ -465,7 +465,7 @@ void FFontEditor::PostUndo(bool bSuccess)
 void FFontEditor::NotifyPostChange( const FPropertyChangedEvent& PropertyChangedEvent, class FEditPropertyChain* PropertyThatChanged)
 {
 	static const FName FontCacheTypePropertyName = GET_MEMBER_NAME_CHECKED(UFont, FontCacheType);
-	static const FName CompositeFontPropertyName = GET_MEMBER_NAME_CHECKED(UFont, CompositeFont);
+	static const FName CompositeFontPropertyName = UFont::GetInternalCompositeFontPropertyName();
 	static const FName TexturePageWidthName = GET_MEMBER_NAME_CHECKED(FFontImportOptionsData, TexturePageWidth);
 	static const FName TexturePageMaxHeightName = GET_MEMBER_NAME_CHECKED(FFontImportOptionsData, TexturePageMaxHeight);
 	static const FName DistanceFieldScaleFactorName = GET_MEMBER_NAME_CHECKED(FFontImportOptionsData, DistanceFieldScaleFactor);
diff --git a/Engine/Source/Editor/FontEditor/Private/FontFaceEditor.cpp b/Engine/Source/Editor/FontEditor/Private/FontFaceEditor.cpp
--- a/Engine/Source/Editor/FontEditor/Private/FontFaceEditor.cpp
+++ b/Engine/Source/Editor/FontEditor/Private/FontFaceEditor.cpp
@@ -601,17 +601,18 @@ void FFontFaceEditor::MakePreviewFont(TObjectPtr<UObject>& TargetObject, UFontFa
 	}
 	if (UFont* TargetFont = CastChecked<UFont>(TargetObject))
 	{
-		if (TargetFont->CompositeFont.DefaultTypeface.Fonts.IsEmpty())
+		FCompositeFont& MutableInternalCompositeFont = TargetFont->GetMutableInternalCompositeFont();
+		if (MutableInternalCompositeFont.DefaultTypeface.Fonts.IsEmpty())
 		{
 			FTypefaceEntry FontTypeface;
 			FontTypeface.Name = TEXT("Regular");
 			FontTypeface.Font = FFontData(Face);
 			TargetFont->FontCacheType = EFontCacheType::Runtime;
-			TargetFont->CompositeFont.DefaultTypeface.Fonts.Add(MoveTemp(FontTypeface));
+			MutableInternalCompositeFont.DefaultTypeface.Fonts.Add(MoveTemp(FontTypeface));
 		}
 		else
 		{
-			TargetFont->CompositeFont.DefaultTypeface.Fonts[0].Font = FFontData(Face);
+			MutableInternalCompositeFont.DefaultTypeface.Fonts[0].Font = FFontData(Face);
 		}
 		TargetFont->PostEditChange();
 	}
diff --git a/Engine/Source/Editor/FontEditor/Private/SCompositeFontEditor.cpp b/Engine/Source/Editor/FontEditor/Private/SCompositeFontEditor.cpp
--- a/Engine/Source/Editor/FontEditor/Private/SCompositeFontEditor.cpp
+++ b/Engine/Source/Editor/FontEditor/Private/SCompositeFontEditor.cpp
@@ -298,13 +298,13 @@ UFont* SCompositeFontEditor::GetFontObject() const
 FCompositeFont* SCompositeFontEditor::GetCompositeFont() const
 {
 	UFont* const FontObject = GetFontObject();
-	return (FontObject) ? &FontObject->CompositeFont : nullptr;
+	return (FontObject) ? &FontObject->GetMutableInternalCompositeFont() : nullptr;
 }
 
 FTypeface* SCompositeFontEditor::GetDefaultTypeface() const
 {
 	UFont* const FontObject = GetFontObject();
-	return (FontObject) ? &FontObject->CompositeFont.DefaultTypeface : nullptr;
+	return (FontObject) ? &FontObject->GetMutableInternalCompositeFont().DefaultTypeface : nullptr;
 }
 
 const FTypeface* SCompositeFontEditor::GetConstDefaultTypeface() const
@@ -315,7 +315,7 @@ const FTypeface* SCompositeFontEditor::GetConstDefaultTypeface() const
 FTypeface* SCompositeFontEditor::GetFallbackTypeface() const
 {
 	UFont* const FontObject = GetFontObject();
-	return (FontObject) ? &FontObject->CompositeFont.FallbackTypeface.Typeface : nullptr;
+	return (FontObject) ? &FontObject->GetMutableInternalCompositeFont().FallbackTypeface.Typeface : nullptr;
 }
 
 const FTypeface* SCompositeFontEditor::GetConstFallbackTypeface() const
@@ -326,7 +326,7 @@ const FTypeface* SCompositeFontEditor::GetConstFallbackTypeface() const
 FCompositeFallbackFont* SCompositeFontEditor::GetFallbackFont() const
 {
 	UFont* const FontObject = GetFontObject();
-	return (FontObject) ? &FontObject->CompositeFont.FallbackTypeface : nullptr;
+	return (FontObject) ? &FontObject->GetMutableInternalCompositeFont().FallbackTypeface : nullptr;
 }
 
 void SCompositeFontEditor::UpdateSubTypefaceList()
diff --git a/Engine/Source/Editor/FontEditor/Private/SFontEditorViewport.cpp b/Engine/Source/Editor/FontEditor/Private/SFontEditorViewport.cpp
--- a/Engine/Source/Editor/FontEditor/Private/SFontEditorViewport.cpp
+++ b/Engine/Source/Editor/FontEditor/Private/SFontEditorViewport.cpp
@@ -263,8 +263,14 @@ void FFontEditorViewportClient::Draw(FViewport* Viewport, FCanvas* Canvas)
 			FVector2D CurPos = StartPos * FontScale;
 			int32 WidestName = 0;
 
+			TArrayView<const FTypefaceEntry> TypefaceEntries;
+			if (const FCompositeFont* CompositeFont = Font->GetCompositeFont())
+			{
+				TypefaceEntries = CompositeFont->DefaultTypeface.Fonts;
+			}
+
 			// Draw and measure each name so we can work out where to start drawing the preview text column
-			for (const FTypefaceEntry& TypefaceEntry : Font->CompositeFont.DefaultTypeface.Fonts)
+			for (const FTypefaceEntry& TypefaceEntry : TypefaceEntries)
 			{
 				FSlateFontInfo FontInfo(Font, Font->LegacyFontSize, TypefaceEntry.Name);
 				FontInfo.FontFallback = EFontFallback::FF_NoFallback;
@@ -283,7 +289,7 @@ void FFontEditorViewportClient::Draw(FViewport* Viewport, FCanvas* Canvas)
 			CurPos = FVector2D(WidestName + (12.0f * FontScale), StartPos.Y * FontScale);
 
 			// Draw the preview text using each of the default fonts
-			for (const FTypefaceEntry& TypefaceEntry : Font->CompositeFont.DefaultTypeface.Fonts)
+			for (const FTypefaceEntry& TypefaceEntry : TypefaceEntries)
 			{
 				FSlateFontInfo FontInfo(Font, Font->LegacyFontSize, TypefaceEntry.Name);
 				FontInfo.FontFallback = EFontFallback::FF_NoFallback;
diff --git a/Engine/Source/Editor/UnrealEd/Private/Factories/EditorFactories.cpp b/Engine/Source/Editor/UnrealEd/Private/Factories/EditorFactories.cpp
--- a/Engine/Source/Editor/UnrealEd/Private/Factories/EditorFactories.cpp
+++ b/Engine/Source/Editor/UnrealEd/Private/Factories/EditorFactories.cpp
@@ -5151,7 +5151,7 @@ UObject* UFontFileImportFactory::FactoryCreateBinary(UClass* InClass, UObject* I
 			Font->FontCacheType = EFontCacheType::Runtime;
 
 			// Add a default typeface referencing the newly created font face
-			FTypefaceEntry& DefaultTypefaceEntry = Font->CompositeFont.DefaultTypeface.Fonts[Font->CompositeFont.DefaultTypeface.Fonts.AddDefaulted()];
+			FTypefaceEntry& DefaultTypefaceEntry = Font->GetMutableInternalCompositeFont().DefaultTypeface.Fonts.AddDefaulted_GetRef();
 			DefaultTypefaceEntry.Name = "Default";
 			DefaultTypefaceEntry.Font = FFontData(FontFace);
 
diff --git a/Engine/Source/Editor/UnrealEd/Private/ThumbnailRendering/FontThumbnailRenderer.cpp b/Engine/Source/Editor/UnrealEd/Private/ThumbnailRendering/FontThumbnailRenderer.cpp
--- a/Engine/Source/Editor/UnrealEd/Private/ThumbnailRendering/FontThumbnailRenderer.cpp
+++ b/Engine/Source/Editor/UnrealEd/Private/ThumbnailRendering/FontThumbnailRenderer.cpp
@@ -36,7 +36,8 @@ void UFontThumbnailRenderer::GetThumbnailSize(UObject* Object, float Zoom, uint3
 			break;
 
 		case EFontCacheType::Runtime:
-			if(Font->CompositeFont.DefaultTypeface.Fonts.Num() > 0)
+			if (const FCompositeFont* CompositeFont = Font->GetCompositeFont();
+				CompositeFont && CompositeFont->DefaultTypeface.Fonts.Num() > 0)
 			{
 				OutWidth = OutHeight = static_cast<uint32>(Zoom * 256.0f);
 			}
@@ -68,7 +69,8 @@ void UFontThumbnailRenderer::Draw(UObject* Object, int32 X, int32 Y, uint32 Widt
 			break;
 
 		case EFontCacheType::Runtime:
-			if(Font->CompositeFont.DefaultTypeface.Fonts.Num() > 0)
+			if (const FCompositeFont* CompositeFont = Font->GetCompositeFont();
+				CompositeFont && CompositeFont->DefaultTypeface.Fonts.Num() > 0)
 			{
 				static const int32 FontSize = 28;
 				static const float FontScale = 1.0f;
@@ -80,7 +82,7 @@ void UFontThumbnailRenderer::Draw(UObject* Object, int32 X, int32 Y, uint32 Widt
 
 				// Draw the object name using each of the default fonts
 				FVector2D CurPos(X, Y);
-				for(const FTypefaceEntry& TypefaceEntry : Font->CompositeFont.DefaultTypeface.Fonts)
+				for(const FTypefaceEntry& TypefaceEntry : CompositeFont->DefaultTypeface.Fonts)
 				{
 					FSlateFontInfo FontInfo(Font, FontSize, TypefaceEntry.Name);
 					FShapedGlyphSequenceRef FontNameShapedText = FontCache->ShapeBidirectionalText(FontName, FontInfo, FontScale, FontNameBaseTextDirection, ETextShapingMethod::Auto);
diff --git a/Engine/Source/Runtime/Engine/Classes/Engine/Font.h b/Engine/Source/Runtime/Engine/Classes/Engine/Font.h
--- a/Engine/Source/Runtime/Engine/Classes/Engine/Font.h
+++ b/Engine/Source/Runtime/Engine/Classes/Engine/Font.h
@@ -23,6 +23,15 @@ enum class EFontCacheType : uint8
 	Runtime,
 };
 
+UENUM()
+enum class ERuntimeFontSource : uint8
+{
+	/** The runtime composite font data is embedded into this Font asset */
+	Asset,
+
+	/** The runtime composite font is taken from FCoreStyle::GetDefaultFont() */
+	CoreStyleDefault,
+};
 
 /** This struct is serialized using native serialization so any changes to it require a package version bump. */
 USTRUCT()
@@ -106,6 +115,10 @@ class UFont : public UObject, public IFontProviderInterface
 	UPROPERTY(EditAnywhere, Category=Font)
 	EFontCacheType FontCacheType;
 
+	/** Where is the runtime font data for this Font stored? */
+	UPROPERTY(EditAnywhere, AdvancedDisplay, Category=RuntimeFont)
+	ERuntimeFontSource RuntimeFontSource = ERuntimeFontSource::Asset;
+
 	/** List of characters in the font.  For a MultiFont, this will include all characters in all sub-fonts!  Thus,
 		the number of characters in this array isn't necessary the number of characters available in the font */
 	UPROPERTY(EditAnywhere, Category=OfflineFont)
@@ -169,6 +182,7 @@ class UFont : public UObject, public IFontProviderInterface
 	FName LegacyFontName;
 
 	/** Embedded composite font data */
+	UE_DEPRECATED(5.7, "Public access to CompositeFont is deprecated. Use GetCompositeFont(), GetInternalCompositeFont(), or GetMutableInternalCompositeFont() instead.")
 	UPROPERTY(EditAnywhere, Category = RuntimeFont)
 	FCompositeFont CompositeFont;
 
@@ -181,9 +195,27 @@ public:
 	TMap<uint16,uint16> CharRemap;
 
 	/** IFontProviderInterface */
-	virtual const FCompositeFont* GetCompositeFont() const override
+	ENGINE_API virtual const FCompositeFont* GetCompositeFont() const override;
+
+	inline const FCompositeFont& GetInternalCompositeFont() const
+	{
+		PRAGMA_DISABLE_DEPRECATION_WARNINGS
+		return CompositeFont;
+		PRAGMA_ENABLE_DEPRECATION_WARNINGS
+	}
+
+	inline FCompositeFont& GetMutableInternalCompositeFont()
+	{
+		PRAGMA_DISABLE_DEPRECATION_WARNINGS
+		return CompositeFont;
+		PRAGMA_ENABLE_DEPRECATION_WARNINGS
+	}
+
+	static FName GetInternalCompositeFontPropertyName()
 	{
-		return (FontCacheType == EFontCacheType::Runtime) ? &CompositeFont : nullptr;
+		PRAGMA_DISABLE_DEPRECATION_WARNINGS
+		return GET_MEMBER_NAME_CHECKED(UFont, CompositeFont);
+		PRAGMA_ENABLE_DEPRECATION_WARNINGS
 	}
 
 	/** Get the info needed to use this UFont with Slate, using the fallback data for legacy Canvas APIs */
diff --git a/Engine/Source/Runtime/Engine/Private/Font.cpp b/Engine/Source/Runtime/Engine/Private/Font.cpp
--- a/Engine/Source/Runtime/Engine/Private/Font.cpp
+++ b/Engine/Source/Runtime/Engine/Private/Font.cpp
@@ -35,7 +35,7 @@ void UFont::BeginDestroy()
 		if (SlateRenderer)
 		{
 			TSharedRef<FSlateFontCache> FontCache = SlateRenderer->GetFontCache();
-			FontCache->FlushCompositeFont(CompositeFont);
+			FontCache->FlushCompositeFont(GetInternalCompositeFont());
 			FontCache->FlushObject(this);
 		}
 	}
@@ -97,12 +97,13 @@ void UFont::PostLoad()
 		};
 
 		int32 TypefaceIndex = 0;
-		for (FTypefaceEntry& TypefaceEntry : CompositeFont.DefaultTypeface.Fonts)
+		FCompositeFont& MutableInternalCompositeFont = GetMutableInternalCompositeFont();
+		for (FTypefaceEntry& TypefaceEntry : MutableInternalCompositeFont.DefaultTypeface.Fonts)
 		{
 			UpgradeLegacyData(TypefaceIndex, TypefaceEntry.Name, TypefaceEntry.Font);
 		}
 
-		for (FCompositeSubFont& SubFont : CompositeFont.SubTypefaces)
+		for (FCompositeSubFont& SubFont : MutableInternalCompositeFont.SubTypefaces)
 		{
 			++TypefaceIndex;
 			for (FTypefaceEntry& TypefaceEntry : SubFont.Typeface.Fonts)
@@ -380,10 +381,32 @@ void UFont::GetResourceSizeEx(FResourceSizeEx& CumulativeResourceSize)
 		};
 
 		// Sum the contained font data sizes
-		GetTypefaceResourceSize(CompositeFont.DefaultTypeface);
-		for (const FCompositeSubFont& SubTypeface : CompositeFont.SubTypefaces)
+		const FCompositeFont& InternalCompositeFont = GetInternalCompositeFont();
+		GetTypefaceResourceSize(InternalCompositeFont.DefaultTypeface);
+		for (const FCompositeSubFont& SubTypeface : InternalCompositeFont.SubTypefaces)
 		{
 			GetTypefaceResourceSize(SubTypeface.Typeface);
 		}
 	}
 }
+
+const FCompositeFont* UFont::GetCompositeFont() const
+{
+	if (FontCacheType == EFontCacheType::Runtime)
+	{
+		switch (RuntimeFontSource)
+		{
+			case ERuntimeFontSource::Asset:
+				return &GetInternalCompositeFont();
+
+			case ERuntimeFontSource::CoreStyleDefault:
+				// Note: Returning this as a raw pointer is safe (for use on the stack), as FLegacySlateFontInfoCache keeps the FCompositeFont alive as long as Slate is initialized
+				return &FCoreStyle::GetDefaultFont().Get();
+
+			default:
+				break;
+		}
+	}
+
+	return nullptr;
+}
diff --git a/Engine/Source/Runtime/Engine/Private/UnrealEngine.cpp b/Engine/Source/Runtime/Engine/Private/UnrealEngine.cpp
--- a/Engine/Source/Runtime/Engine/Private/UnrealEngine.cpp
+++ b/Engine/Source/Runtime/Engine/Private/UnrealEngine.cpp
@@ -3396,9 +3396,9 @@ void UEngine::InitializeObjectReferences()
 		{
 			UFont* FontPtr = NewObject<UFont>(GetTransientPackage(), InObjectName, RF_Transient);
 			FontPtr->FontCacheType = EFontCacheType::Runtime;
+			FontPtr->RuntimeFontSource = ERuntimeFontSource::CoreStyleDefault;
 			FontPtr->LegacyFontName = InLegacyFontName;
 			FontPtr->LegacyFontSize = InLegacyFontSize;
-			FontPtr->CompositeFont = *FCoreStyle::GetDefaultFont();
 			return FontPtr;
 		};
 
