- Making use of new directly setting of values in a config layer, updated the config hierarchy editor (IOS project settings has a bunch of properties opting in to use it, as an example). Note that this does not need the ini.UseNewPropertySaving cvar, that is only for the normal ProjectSettings type stuff

diff --git a/Engine/Source/Editor/ConfigEditor/Private/ConfigPropertyHelper.cpp b/Engine/Source/Editor/ConfigEditor/Private/ConfigPropertyHelper.cpp
--- a/Engine/Source/Editor/ConfigEditor/Private/ConfigPropertyHelper.cpp
+++ b/Engine/Source/Editor/ConfigEditor/Private/ConfigPropertyHelper.cpp
@@ -15,6 +15,8 @@ class ISourceControlProvider;
 
 void UPropertyConfigFileDisplayRow::InitWithConfigAndProperty(const FString& InConfigFileName, FProperty* InEditProperty)
 {
+	// cache both relative (../../../Proj/Config/Foo.ini) and absolute paths (D:\UE\Proj\Config\Foo.ini) 
+	NormalizedConfigFileName = InConfigFileName;
 	ConfigFileName = FPaths::ConvertRelativePathToFull(InConfigFileName);
 	ExternalProperty = InEditProperty;
 
@@ -44,7 +46,8 @@ void UPropertyConfigFileDisplayRow::InitWithConfigAndProperty(const FString& InC
 			}
 			else
 			{
-				bIsFileWritable = false;
+				// if the file is not in source control, it is "writeable" ie createable
+				bIsFileWritable = true;
 			}
 		}
 	}
diff --git a/Engine/Source/Editor/ConfigEditor/Private/PropertyVisualization/ConfigPropertyCellPresenter.cpp b/Engine/Source/Editor/ConfigEditor/Private/PropertyVisualization/ConfigPropertyCellPresenter.cpp
--- a/Engine/Source/Editor/ConfigEditor/Private/PropertyVisualization/ConfigPropertyCellPresenter.cpp
+++ b/Engine/Source/Editor/ConfigEditor/Private/PropertyVisualization/ConfigPropertyCellPresenter.cpp
@@ -24,6 +24,7 @@
 #include "Templates/Casts.h"
 #include "Types/SlateEnums.h"
 #include "UObject/Class.h"
+#include "UObject/Package.h"
 #include "UObject/FieldPath.h"
 #include "UObject/NameTypes.h"
 #include "UObject/UnrealType.h"
@@ -80,13 +81,14 @@ private:
 	EVisibility GetRemovePropertyFromConfigVisibility() const;
 
 	// Check if the config file has an entry for this property
-	bool DoesConfigFileHaveEntryForProperty();
+	bool DoesConfigFileHaveEntryForProperty() const;
 	// Cached value of the above, as we use this for widget visibility callbacks.
 	bool bCachedConfigHasPropertyValue;
 
 	// The object which holds the property we are editing
 	UPropertyConfigFileDisplayRow* PropertyOwnerDisplayObject;
 
+	FConfigCacheIni* ConfigSystem;
 };
 
 
@@ -208,9 +210,12 @@ bool SConfigPropertyCell::HandleDisplayedValueWidgetEnabled() const
 
 FReply SConfigPropertyCell::HandleAddPropertyToConfigClicked()
 {
-	PropertyOwnerDisplayObject->TryUpdateDefaultConfigFile(PropertyOwnerDisplayObject->ConfigFileName);
 	bCachedConfigHasPropertyValue = true;
 
+	// force a "OnValueChanged" in the customization here by triggering a PropertyChanged event
+	FPropertyChangedEvent GeneratedOutputChangedEvent(PropertyOwnerDisplayObject->ExternalProperty.Get(), EPropertyChangeType::ValueSet);
+	FCoreUObjectDelegates::OnObjectPropertyChanged.Broadcast(PropertyOwnerDisplayObject->TempConfigObject.Get(), GeneratedOutputChangedEvent);
+
 	return FReply::Handled();
 }
 
@@ -223,16 +228,16 @@ EVisibility SConfigPropertyCell::GetAddPropertyToConfigButtonVisibility() const
 
 FReply SConfigPropertyCell::HandleRemovePropertyFromConfigClicked()
 {
-	// This needs work...
+	// remove the property from the layer
 	UClass* OwnerClass = PropertyOwnerDisplayObject->ExternalProperty->GetOwnerClass();
 	FString SectionName = OwnerClass->GetPathName();
 	FString PropertyName = PropertyOwnerDisplayObject->ExternalProperty->GetName();
+	GConfig->ResetKeyInSectionOfStaticLayer(*SectionName, *PropertyName, OwnerClass->GetConfigName(), PropertyOwnerDisplayObject->NormalizedConfigFileName, true);
 
-	// This...
-	GConfig->RemoveKey(*SectionName, *PropertyName, PropertyOwnerDisplayObject->ConfigFileName);
+	// Update the CDO, as this change might have had an impact on it's value.
+	OwnerClass->GetDefaultObject()->ReloadConfig();
 
-	// ... really oughta affect this...
-	PropertyOwnerDisplayObject->TryUpdateDefaultConfigFile(PropertyOwnerDisplayObject->ConfigFileName);
+	bCachedConfigHasPropertyValue = false;
 
 	return FReply::Handled();
 }
@@ -240,19 +245,27 @@ FReply SConfigPropertyCell::HandleRemovePropertyFromConfigClicked()
 
 EVisibility SConfigPropertyCell::GetRemovePropertyFromConfigVisibility() const
 {
-	/* Technical limitations of UObjects and Config make this difficult. Hopefully we can add it at some point. */
-	return EVisibility::Collapsed;
+	return DoesConfigFileHaveEntryForProperty() ? EVisibility::Visible : EVisibility::Collapsed;
 }
 
 
-bool SConfigPropertyCell::DoesConfigFileHaveEntryForProperty()
+bool SConfigPropertyCell::DoesConfigFileHaveEntryForProperty() const
 {
 	UClass* OwnerClass = PropertyOwnerDisplayObject->ExternalProperty->GetOwnerClass();
 	FString SectionName = OwnerClass->GetPathName();
-	FString PropertyName = PropertyOwnerDisplayObject->ExternalProperty->GetName();
-	
-	FString TempValue;
-	return GConfig->GetString(*SectionName, *PropertyName, TempValue, PropertyOwnerDisplayObject->ConfigFileName);
+	FName PropertyName = PropertyOwnerDisplayObject->ExternalProperty->GetFName();
+
+	if (const FConfigBranch* Branch = GConfig->FindBranch(*OwnerClass->GetConfigName(), OwnerClass->GetConfigName()))
+	{
+		if (const FConfigCommandStream* Layer = Branch->GetStaticLayer(PropertyOwnerDisplayObject->NormalizedConfigFileName))
+		{
+			if (const FConfigCommandStreamSection* Sec = Layer->Find(SectionName))
+			{
+				return Sec->Contains(PropertyName);
+			}
+		}
+	}
+	return false;
 }
 
 
diff --git a/Engine/Source/Editor/ConfigEditor/Public/ConfigPropertyHelper.h b/Engine/Source/Editor/ConfigEditor/Public/ConfigPropertyHelper.h
--- a/Engine/Source/Editor/ConfigEditor/Public/ConfigPropertyHelper.h
+++ b/Engine/Source/Editor/ConfigEditor/Public/ConfigPropertyHelper.h
@@ -38,6 +38,9 @@ public:
 	UPROPERTY(Transient, Category = Helper, VisibleAnywhere)
 	FString ConfigFileName;
 
+	FString NormalizedConfigFileName;
+	TObjectPtr<UObject> TempConfigObject;
+
 	UPROPERTY(Transient, Category = Helper, EditAnywhere, meta=(EditCondition="bIsFileWritable"))
 	TFieldPath<FProperty> ExternalProperty;
 
diff --git a/Engine/Source/Editor/DetailCustomizations/Private/ConfigEditorPropertyDetails.cpp b/Engine/Source/Editor/DetailCustomizations/Private/ConfigEditorPropertyDetails.cpp
--- a/Engine/Source/Editor/DetailCustomizations/Private/ConfigEditorPropertyDetails.cpp
+++ b/Engine/Source/Editor/DetailCustomizations/Private/ConfigEditorPropertyDetails.cpp
@@ -100,38 +100,15 @@ void FConfigPropertyHelperDetails::OnPropertyValueChanged(UObject* Object, FProp
 		{
 			const FString& ConfigIniName = *FileName;
 
-			// We should set this up to work with the UObject Config system, its difficult as the outer object isnt of the same type
-			// create a sandbox FConfigCache
-			FConfigCacheIni Config(EConfigCacheType::Temporary);
-
-			// add an empty file to the config so it doesn't read in the original file (see FConfigCacheIni.Find())
-			FConfigFile& NewFile = Config.Add(ConfigIniName, FConfigFile());
+			// @todo everywhere in this code - check platform and swap Config system!
 
 			// save the object properties to this file
-			OriginalProperty->GetOwnerClass()->GetDefaultObject()->SaveConfig(CPF_Config, *ConfigIniName, &Config);
-
-			// Take the saved section for this object and have the config system process and write out the one property we care about.
-			TArray<FString> Keys;
-			NewFile.GetKeys(Keys);
-
-			const FString SectionName = Keys[0];
-			const FString PropertyName = ConfigEditorCopyOfEditProperty->GetName();
-			FString	Value;
-			ConfigEditorCopyOfEditProperty->ExportText_InContainer(0, Value, Object, Object, Object, 0);
-			NewFile.SetString(*SectionName, *PropertyName, *Value);
-			GConfig->SetString(*SectionName, *PropertyName, *Value, ConfigIniName);
-
-			NewFile.UpdateSinglePropertyInSection(*ConfigIniName, *PropertyName, *SectionName);
-
-			// reload the file, so that it refresh the cache internally.
-			if (GConfig->FindConfigFile(ConfigIniName))
-			{
-				FConfigContext::ForceReloadIntoGConfig().Load(*OriginalProperty->GetOwnerClass()->ClassConfigName.ToString());
-			}
-			else
-			{
-				FConfigCacheIni::ClearOtherPlatformConfigs();
-			}
+			UObject::FSaveConfigContext Context;
+			Context.Filename = ConfigIniName;
+			Context.bSaveToLayer = true;
+			Context.SinglePropertyName = OriginalProperty->GetFName();
+			Context.OverrideConfigSectionName = OriginalProperty->GetOwnerClass()->GetPathName();
+			Object->SaveConfig(Context);
 
 			// Update the CDO, as this change might have had an impact on it's value.
 			OriginalProperty->GetOwnerClass()->GetDefaultObject()->ReloadConfig();
@@ -140,9 +117,9 @@ void FConfigPropertyHelperDetails::OnPropertyValueChanged(UObject* Object, FProp
 }
 
 
-void FConfigPropertyHelperDetails::AddEditablePropertyForConfig(IDetailLayoutBuilder& DetailBuilder, const UPropertyConfigFileDisplayRow* ConfigFilePropertyRowObj)
+void FConfigPropertyHelperDetails::AddEditablePropertyForConfig(IDetailLayoutBuilder& DetailBuilder, UPropertyConfigFileDisplayRow* ConfigFilePropertyRowObj)
 {
-	AssociatedConfigFileAndObjectPairings.Add(ConfigFilePropertyRowObj->ConfigFileName, (UObject*)ConfigFilePropertyRowObj);
+	AssociatedConfigFileAndObjectPairings.Add(ConfigFilePropertyRowObj->NormalizedConfigFileName, (UObject*)ConfigFilePropertyRowObj);
 
 	// Add the properties to a property table so we can edit these.
 	IDetailCategoryBuilder& TempCategory = DetailBuilder.EditCategory("TempCategory");
@@ -150,16 +127,31 @@ void FConfigPropertyHelperDetails::AddEditablePropertyForConfig(IDetailLayoutBui
 	UObject* ConfigEntryObject = StaticDuplicateObject(ConfigEditorPropertyViewCDO, GetTransientPackage(), *(ConfigFilePropertyRowObj->ConfigFileName + TEXT("_cdoDupe_") + ConfigEditorPropertyViewCDO->GetName()));
 	ConfigEntryObject->AddToRoot();
 
+	ConfigFilePropertyRowObj->TempConfigObject = ConfigEntryObject;
+
+
 	FString ExistingConfigEntryValue;
 	FString SectionName = OriginalProperty->GetOwnerClass()->GetPathName();
-	FString PropertyName = ConfigEditorCopyOfEditProperty->GetName();
-	if (GConfig->GetString(*SectionName, *PropertyName, ExistingConfigEntryValue, ConfigFilePropertyRowObj->ConfigFileName))
+	FString ConfigName = OriginalProperty->GetOwnerClass()->GetConfigName();
+	FName PropertyName = ConfigEditorCopyOfEditProperty->GetFName();
+
+
+	if (const FConfigBranch* Branch = GConfig->FindBranch(*ConfigName, ConfigName))
 	{
-		ConfigEditorCopyOfEditProperty->ImportText_InContainer(*ExistingConfigEntryValue, ConfigEntryObject, nullptr, 0);
+		if (const FConfigCommandStream* Layer = Branch->GetStaticLayer(ConfigFilePropertyRowObj->NormalizedConfigFileName))
+		{
+			if (const FConfigCommandStreamSection* Sec = Layer->Find(SectionName))
+			{
+				if (const FConfigValue* Value = Sec->Find(PropertyName))
+				{
+					ConfigEditorCopyOfEditProperty->ImportText_InContainer(*Value->GetValueForWriting(), ConfigEntryObject, nullptr, 0);
+				}
+			}
+		}
 	}
 
 	// Cache a reference for future usage.
-	ConfigFileAndPropertySourcePairings.Add(ConfigFilePropertyRowObj->ConfigFileName, (UObject*)ConfigEntryObject);
+	ConfigFileAndPropertySourcePairings.Add(ConfigFilePropertyRowObj->NormalizedConfigFileName, (UObject*)ConfigEntryObject);
 
 	// We need to add a property row for each config file entry. 
 	// This allows us to have an editable widget for each config file.
diff --git a/Engine/Source/Editor/DetailCustomizations/Private/ConfigEditorPropertyDetails.h b/Engine/Source/Editor/DetailCustomizations/Private/ConfigEditorPropertyDetails.h
--- a/Engine/Source/Editor/DetailCustomizations/Private/ConfigEditorPropertyDetails.h
+++ b/Engine/Source/Editor/DetailCustomizations/Private/ConfigEditorPropertyDetails.h
@@ -73,7 +73,7 @@ private:
 	//////////////////////////////////////////////////////////////////////////
 	/// Misc operations...
 private:
-	void AddEditablePropertyForConfig(IDetailLayoutBuilder& DetailBuilder, const class UPropertyConfigFileDisplayRow* ConfigFileProperty);
+	void AddEditablePropertyForConfig(IDetailLayoutBuilder& DetailBuilder, class UPropertyConfigFileDisplayRow* ConfigFileProperty);
 
 
 	//////////////////////////////////////////////////////////////////////////
