Add LOCTEXT for Light Profiles category

#rnx

diff --git a/Engine/Source/Editor/DetailCustomizations/Private/LightComponentDetails.cpp b/Engine/Source/Editor/DetailCustomizations/Private/LightComponentDetails.cpp
--- a/Engine/Source/Editor/DetailCustomizations/Private/LightComponentDetails.cpp
+++ b/Engine/Source/Editor/DetailCustomizations/Private/LightComponentDetails.cpp
@@ -189,7 +189,7 @@ void FLightComponentDetails::CustomizeDetails( IDetailLayoutBuilder& DetailBuild
 
 		LightCategory.AddProperty( DetailBuilder.GetProperty("LightColor", ULightComponentBase::StaticClass() ) );
 
-		IDetailCategoryBuilder& LightProfilesCategory = DetailBuilder.EditCategory( "Light Profiles", FText::GetEmpty(), ECategoryPriority::Default );
+		IDetailCategoryBuilder& LightProfilesCategory = DetailBuilder.EditCategory( "Light Profiles", LOCTEXT("LightProfiles", "Light Profiles"), ECategoryPriority::Default );
 		LightProfilesCategory.AddProperty( IESBrightnessTextureProperty );
 
 		LightProfilesCategory.AddProperty( IESBrightnessEnabledProperty )
