FDateTimeStructCustomization now has a tooltip, so that hovering over a DateTime property's text box in editor provides a hint at the format.

diff --git a/Engine/Source/Editor/DetailCustomizations/Private/DateTimeStructCustomization.cpp b/Engine/Source/Editor/DetailCustomizations/Private/DateTimeStructCustomization.cpp
--- a/Engine/Source/Editor/DetailCustomizations/Private/DateTimeStructCustomization.cpp
+++ b/Engine/Source/Editor/DetailCustomizations/Private/DateTimeStructCustomization.cpp
@@ -38,6 +38,8 @@ void FDateTimeStructCustomization::CustomizeHeader(TSharedRef<IPropertyHandle> S
 {
 	PropertyHandle = StructPropertyHandle;
 
+	const FText DateTimeTooltipText = LOCTEXT("Tooltip", "A date-time string in the following format: YYYY.MM.DD-hh.mm.ss +hhmm \n\nThe time can be entered as UTC (timezone-agnostic) or with a timezone specific UTC offset (+hhmm). The value is displayed here with the local timezone's UTC offset.");
+
 	HeaderRow
 		.NameContent()
 		[
@@ -53,6 +55,7 @@ void FDateTimeStructCustomization::CustomizeHeader(TSharedRef<IPropertyHandle> S
 				.ForegroundColor(this, &FDateTimeStructCustomization::HandleTextBoxForegroundColor)
 				.OnTextChanged(this, &FDateTimeStructCustomization::HandleTextBoxTextChanged)
 				.OnTextCommitted(this, &FDateTimeStructCustomization::HandleTextBoxTextCommited)
+				.ToolTipText(DateTimeTooltipText)
 				.SelectAllTextOnCommit(true)
 				.Font(IPropertyTypeCustomizationUtils::GetRegularFont())
 				.Text(this, &FDateTimeStructCustomization::HandleTextBoxText)
