[Input] Add a way to deprecate specific "legacy" axis and action mapping names. If this flag is set to true on the project settings, then any BP references to it will have a deprecation warning. Any C++ Bind calls using these names will also emit a warning.

You can supress the deprecation warnings with the bEnabledLegacyMappingDeprecationWarnings option in the project's input settings.

#pf 688a922a557f76840a860c8b
[RN] Add a set of deprecated input action/axis names in the project settings.

diff --git a/Engine/Source/Editor/BlueprintGraph/Classes/K2Node_InputAction.h b/Engine/Source/Editor/BlueprintGraph/Classes/K2Node_InputAction.h
--- a/Engine/Source/Editor/BlueprintGraph/Classes/K2Node_InputAction.h
+++ b/Engine/Source/Editor/BlueprintGraph/Classes/K2Node_InputAction.h
@@ -61,6 +61,9 @@ class UK2Node_InputAction : public UK2Node, public IK2Node_EventNodeInterface
 	UE_API virtual FText GetTooltipText() const override;
 	UE_API virtual FSlateIcon GetIconAndTint(FLinearColor& OutColor) const override;
 	UE_API virtual bool IsCompatibleWithGraph(UEdGraph const* Graph) const override;
+	UE_API virtual bool IsDeprecated() const override;
+	UE_API virtual bool HasDeprecatedReference() const override;
+	UE_API virtual FEdGraphNodeDeprecationResponse GetDeprecationResponse(EEdGraphNodeDeprecationType DeprecationType) const override;
 	//~ End UEdGraphNode Interface.
 
 	//~ Begin UK2Node Interface
diff --git a/Engine/Source/Editor/BlueprintGraph/Classes/K2Node_InputAxisEvent.h b/Engine/Source/Editor/BlueprintGraph/Classes/K2Node_InputAxisEvent.h
--- a/Engine/Source/Editor/BlueprintGraph/Classes/K2Node_InputAxisEvent.h
+++ b/Engine/Source/Editor/BlueprintGraph/Classes/K2Node_InputAxisEvent.h
@@ -53,6 +53,9 @@ class UK2Node_InputAxisEvent : public UK2Node_Event
 	UE_API virtual FText GetNodeTitle(ENodeTitleType::Type TitleType) const override;
 	UE_API virtual FText GetTooltipText() const override;
 	UE_API virtual bool IsCompatibleWithGraph(const UEdGraph* TargetGraph) const override;
+	UE_API virtual bool IsDeprecated() const override;
+	UE_API virtual bool HasDeprecatedReference() const override;
+	UE_API virtual FEdGraphNodeDeprecationResponse GetDeprecationResponse(EEdGraphNodeDeprecationType DeprecationType) const override;
 	//~ End EdGraphNode Interface
 
 	//~ Begin UK2Node Interface
diff --git a/Engine/Source/Editor/BlueprintGraph/Private/K2Node_InputAction.cpp b/Engine/Source/Editor/BlueprintGraph/Private/K2Node_InputAction.cpp
--- a/Engine/Source/Editor/BlueprintGraph/Private/K2Node_InputAction.cpp
+++ b/Engine/Source/Editor/BlueprintGraph/Private/K2Node_InputAction.cpp
@@ -129,6 +129,40 @@ bool UK2Node_InputAction::IsCompatibleWithGraph(UEdGraph const* Graph) const
 	return bIsCompatible;
 }
 
+bool UK2Node_InputAction::IsDeprecated() const
+{
+	const UInputSettings* Settings = GetDefault<UInputSettings>();
+	return Super::IsDeprecated() || (Settings->bEnabledLegacyMappingDeprecationWarnings && Settings->IsActionMappingDeprecated(InputActionName));
+}
+
+bool UK2Node_InputAction::HasDeprecatedReference() const
+{
+	const bool bDeprecationWarningEnabled = GetDefault<UInputSettings>()->bEnabledLegacyMappingDeprecationWarnings;
+
+	if (!bDeprecationWarningEnabled)
+	{
+		return false;
+	}
+	
+	return GetDefault<UInputSettings>()->IsActionMappingDeprecated(InputActionName);	
+}
+
+FEdGraphNodeDeprecationResponse UK2Node_InputAction::GetDeprecationResponse(EEdGraphNodeDeprecationType DeprecationType) const
+{
+	FEdGraphNodeDeprecationResponse Response = Super::GetDeprecationResponse(DeprecationType);
+	
+	if (DeprecationType == EEdGraphNodeDeprecationType::NodeHasDeprecatedReference)
+	{
+		Response.MessageType = EEdGraphNodeDeprecationMessageType::Warning;
+
+		FFormatNamedArguments Args;
+		Args.Add(TEXT("ActionName"), FText::FromName(InputActionName));
+		Response.MessageText = FText::Format(LOCTEXT("ActionMappingDeprecationWarning", "The Action Mapping '{ActionName}' has been deprecated."), Args);
+	}
+
+	return Response;
+}
+
 UEdGraphPin* UK2Node_InputAction::GetPressedPin() const
 {
 	return FindPin(TEXT("Pressed"));
@@ -285,8 +319,16 @@ void UK2Node_InputAction::GetMenuActions(FBlueprintActionDatabaseRegistrar& Acti
 			FEditorDelegates::OnActionAxisMappingsChanged.AddStatic(RefreshClassActions);
 		}
 
+		const UInputSettings* Settings = GetDefault<UInputSettings>();
+
 		for (const FName& ActionName : ActionNames)
 		{
+			// Don't show any deprecated axis mappings in the BP action menu
+			if (Settings->bEnabledLegacyMappingDeprecationWarnings && Settings->IsActionMappingDeprecated(ActionName))
+			{
+				continue;
+			}
+			
 			UBlueprintNodeSpawner* NodeSpawner = UBlueprintNodeSpawner::Create(GetClass());
 			check(NodeSpawner != nullptr);
 
diff --git a/Engine/Source/Editor/BlueprintGraph/Private/K2Node_InputAxisEvent.cpp b/Engine/Source/Editor/BlueprintGraph/Private/K2Node_InputAxisEvent.cpp
--- a/Engine/Source/Editor/BlueprintGraph/Private/K2Node_InputAxisEvent.cpp
+++ b/Engine/Source/Editor/BlueprintGraph/Private/K2Node_InputAxisEvent.cpp
@@ -152,6 +152,40 @@ bool UK2Node_InputAxisEvent::IsCompatibleWithGraph(const UEdGraph* TargetGraph)
 	return bIsCompatible && Super::IsCompatibleWithGraph(TargetGraph);
 }
 
+bool UK2Node_InputAxisEvent::IsDeprecated() const
+{
+	const UInputSettings* Settings = GetDefault<UInputSettings>();
+	return Super::IsDeprecated() || (Settings->bEnabledLegacyMappingDeprecationWarnings && Settings->IsActionMappingDeprecated(InputAxisName));
+}
+
+bool UK2Node_InputAxisEvent::HasDeprecatedReference() const
+{
+	const bool bDeprecationWarningEnabled = GetDefault<UInputSettings>()->bEnabledLegacyMappingDeprecationWarnings;
+
+	if (!bDeprecationWarningEnabled)
+	{
+		return false;
+	}
+	
+	return GetDefault<UInputSettings>()->IsActionMappingDeprecated(InputAxisName);	
+}
+
+FEdGraphNodeDeprecationResponse UK2Node_InputAxisEvent::GetDeprecationResponse(EEdGraphNodeDeprecationType DeprecationType) const
+{
+	FEdGraphNodeDeprecationResponse Response = Super::GetDeprecationResponse(DeprecationType);
+	
+	if (DeprecationType == EEdGraphNodeDeprecationType::NodeHasDeprecatedReference)
+	{
+		Response.MessageType = EEdGraphNodeDeprecationMessageType::Warning;
+
+		FFormatNamedArguments Args;
+		Args.Add(TEXT("ActionName"), FText::FromName(InputAxisName));
+		Response.MessageText = FText::Format(LOCTEXT("AxisMappingDeprecationWarning", "The Axis Mapping '{ActionName}' has been deprecated."), Args);
+	}
+
+	return Response;
+}
+
 void UK2Node_InputAxisEvent::GetMenuActions(FBlueprintActionDatabaseRegistrar& ActionRegistrar) const
 {
 	TArray<FName> AxisNames;
@@ -187,8 +221,16 @@ void UK2Node_InputAxisEvent::GetMenuActions(FBlueprintActionDatabaseRegistrar& A
 			FEditorDelegates::OnActionAxisMappingsChanged.AddStatic(RefreshClassActions);
 		}
 
+		const UInputSettings* Settings = GetDefault<UInputSettings>();
+		
 		for (const FName& AxisName : AxisNames)
 		{
+			// Don't show any deprecated axis mappings in the BP action menu
+			if (Settings->bEnabledLegacyMappingDeprecationWarnings && Settings->IsActionMappingDeprecated(AxisName))
+			{
+				continue;
+			}
+			
 			UBlueprintNodeSpawner* NodeSpawner = UBlueprintNodeSpawner::Create(GetClass());
 			check(NodeSpawner != nullptr);
 
diff --git a/Engine/Source/Editor/DetailCustomizations/Private/InputSettingsDetails.cpp b/Engine/Source/Editor/DetailCustomizations/Private/InputSettingsDetails.cpp
--- a/Engine/Source/Editor/DetailCustomizations/Private/InputSettingsDetails.cpp
+++ b/Engine/Source/Editor/DetailCustomizations/Private/InputSettingsDetails.cpp
@@ -32,6 +32,7 @@
 #include "UObject/UnrealType.h"
 #include "Widgets/DeclarativeSyntaxSupport.h"
 #include "Widgets/Images/SImage.h"
+#include "Widgets/Input/SCheckBox.h"
 #include "Widgets/Input/SEditableTextBox.h"
 #include "Widgets/Layout/SBox.h"
 #include "Widgets/SBoxPanel.h"
@@ -42,12 +43,66 @@ class UObject;
 
 #define LOCTEXT_NAMESPACE "InputSettingsDetails"
 
+namespace UE
+{
+	static void HandleInputDeprecationStatusChanged(UInputSettings* InputSettings, ECheckBoxState NewState, const FName MappingName)
+	{
+		if (!InputSettings)
+		{
+			return;
+		}
+		
+		const FScopedTransaction Transaction(LOCTEXT("DeprecateActionMappingGroup_Transaction", "Deprecate action mapping group"));
+		
+		InputSettings->Modify();
+
+		const bool bShouldBeDeprecated = NewState == ECheckBoxState::Checked;
+		InputSettings->SetActionMappingDeprecated(MappingName, bShouldBeDeprecated);
+
+		// Make sure to save out the default config file changes
+		InputSettings->SaveKeyMappings();
+		InputSettings->TryUpdateDefaultConfigFile();
+	}
+
+	static ECheckBoxState GetInputDeprecationCheckBoxState(UInputSettings* InputSettings, const FName MappingName)
+	{
+		if (InputSettings)
+		{
+			const bool bIsDeprecated = InputSettings->IsActionMappingDeprecated(MappingName);
+			return bIsDeprecated ? ECheckBoxState::Checked : ECheckBoxState::Unchecked;
+		}
 
+		return ECheckBoxState::Undetermined;
+	}
+}
 
 ////////////////////////////////
 // FActionMappingsNodeBuilder //
 ////////////////////////////////
 
+FName InputSettingsDetails::FMappingSet::GetActionMappingName() const
+{
+	// Get the name of the mapping to deprecate
+	FName MappingName = NAME_None;
+	if (Mappings.Num() > 0)
+	{
+		Mappings[0]->GetChildHandle(GET_MEMBER_NAME_CHECKED(FInputActionKeyMapping, ActionName))->GetValue(MappingName);
+	}
+	
+	return MappingName;
+}
+
+FName InputSettingsDetails::FMappingSet::GetAxisMappingName() const
+{
+	FName MappingName = NAME_None;
+	if (Mappings.Num() > 0)
+	{
+		Mappings[0]->GetChildHandle(GET_MEMBER_NAME_CHECKED(FInputAxisKeyMapping, AxisName))->GetValue(MappingName);
+	}
+	return MappingName;
+}
+
+
 FActionMappingsNodeBuilder::FActionMappingsNodeBuilder( IDetailLayoutBuilder* InDetailLayoutBuilder, const TSharedPtr<IPropertyHandle>& InPropertyHandle )
 	: DetailLayoutBuilder( InDetailLayoutBuilder )
 	, ActionMappingsPropertyHandle( InPropertyHandle )
@@ -155,6 +210,26 @@ void FActionMappingsNodeBuilder::GenerateChildContent( IDetailChildrenBuilder& C
 			[
 				RemoveButton
 			]
+			+SHorizontalBox::Slot()
+			.Padding(InputSettingsDetails::InputConstants::PropertyPadding)
+			.HAlign(HAlign_Center)
+			.VAlign(VAlign_Center)
+			.AutoWidth()
+			[
+				SNew(STextBlock)
+				.Text(LOCTEXT("DeprecationLabel_Checkbox", "Is Deprecated"))
+			]
+			+SHorizontalBox::Slot()
+			.Padding(InputSettingsDetails::InputConstants::PropertyPadding)
+			.HAlign(HAlign_Center)
+			.VAlign(VAlign_Center)
+			.AutoWidth()
+			[
+				SNew(SCheckBox)
+				.IsChecked(this, &FActionMappingsNodeBuilder::GetActionMappingCheckBoxState, MappingSet.GetActionMappingName())
+				.OnCheckStateChanged(this, &FActionMappingsNodeBuilder::OnDeprecationCheckboxChanged, MappingSet.GetActionMappingName())
+				.ToolTipText(LOCTEXT("DeprecateActionGroup", "If true, then this action mapping will be considered deprecated."))
+			]
 		];
 	
 		for (int32 MappingIndex = 0; MappingIndex < MappingSet.Mappings.Num(); ++MappingIndex)
@@ -247,6 +322,9 @@ void FActionMappingsNodeBuilder::RemoveActionMappingGroupButton_OnClick(const In
 {
 	const FScopedTransaction Transaction(LOCTEXT("RemoveActionMappingGroup_Transaction", "Remove Action Mapping Group"));
 
+	// Since this axis mapping has been removed, also remove it from the deprecation array
+	UE::HandleInputDeprecationStatusChanged(GetInputSettings(), ECheckBoxState::Unchecked, MappingSet.GetActionMappingName());
+	
 	TSharedPtr<IPropertyHandleArray> ActionMappingsArrayHandle = ActionMappingsPropertyHandle->AsArray();
 
 	TArray<int32> SortedIndices;
@@ -262,6 +340,28 @@ void FActionMappingsNodeBuilder::RemoveActionMappingGroupButton_OnClick(const In
 	}
 }
 
+void FActionMappingsNodeBuilder::OnDeprecationCheckboxChanged(ECheckBoxState NewState, const FName MappingName)
+{
+	UE::HandleInputDeprecationStatusChanged(GetInputSettings(), NewState, MappingName);
+}
+
+ECheckBoxState FActionMappingsNodeBuilder::GetActionMappingCheckBoxState(const FName MappingName) const
+{
+	return UE::GetInputDeprecationCheckBoxState(GetInputSettings(), MappingName);
+}
+
+UInputSettings* FActionMappingsNodeBuilder::GetInputSettings() const
+{
+	TArray<UObject*> OuterObjects;
+	ActionMappingsPropertyHandle->GetOuterObjects(OuterObjects);
+
+	if (OuterObjects.Num() == 1)
+	{
+		return CastChecked<UInputSettings>(OuterObjects[0]);
+	}
+	return nullptr;
+}
+
 bool FActionMappingsNodeBuilder::GroupsRequireRebuild() const
 {
 	for (int32 GroupIndex = 0; GroupIndex < GroupedMappings.Num(); ++GroupIndex)
@@ -445,6 +545,26 @@ void FAxisMappingsNodeBuilder::GenerateChildContent( IDetailChildrenBuilder& Chi
 			[
 				RemoveButton
 			]
+			+SHorizontalBox::Slot()
+			.Padding(InputSettingsDetails::InputConstants::PropertyPadding)
+			.HAlign(HAlign_Center)
+			.VAlign(VAlign_Center)
+			.AutoWidth()
+			[
+				SNew(STextBlock)
+				.Text(LOCTEXT("DeprecationLabel_Checkbox", "Is Deprecated"))
+			]
+			+SHorizontalBox::Slot()
+			.Padding(InputSettingsDetails::InputConstants::PropertyPadding)
+			.HAlign(HAlign_Center)
+			.VAlign(VAlign_Center)
+			.AutoWidth()
+			[
+				SNew(SCheckBox)
+				.IsChecked(this, &FAxisMappingsNodeBuilder::GetAxisMappingCheckBoxState, MappingSet.GetAxisMappingName())
+				.OnCheckStateChanged(this, &FAxisMappingsNodeBuilder::OnDeprecationCheckboxChanged, MappingSet.GetAxisMappingName())
+				.ToolTipText(LOCTEXT("DeprecateAxisGroup", "If true, then this axis mapping will be considered deprecated."))
+			]
 		];
 
 		for (int32 MappingIndex = 0; MappingIndex < MappingSet.Mappings.Num(); ++MappingIndex)
@@ -538,6 +658,9 @@ void FAxisMappingsNodeBuilder::RemoveAxisMappingGroupButton_OnClick(const InputS
 {
 	const FScopedTransaction Transaction(LOCTEXT("RemoveAxisMappingGroup_Transaction", "Remove Axis Mapping Group"));
 
+	// Since this axis mapping has been removed, also remove it from the deprecation array
+	UE::HandleInputDeprecationStatusChanged(GetInputSettings(), ECheckBoxState::Unchecked, MappingSet.GetAxisMappingName());
+	
 	TSharedPtr<IPropertyHandleArray> AxisMappingsArrayHandle = AxisMappingsPropertyHandle->AsArray();
 
 	TArray<int32> SortedIndices;
@@ -553,6 +676,28 @@ void FAxisMappingsNodeBuilder::RemoveAxisMappingGroupButton_OnClick(const InputS
 	}
 }
 
+void FAxisMappingsNodeBuilder::OnDeprecationCheckboxChanged(ECheckBoxState NewState, const FName MappingName)
+{
+	UE::HandleInputDeprecationStatusChanged(GetInputSettings(), NewState, MappingName);
+}
+
+ECheckBoxState FAxisMappingsNodeBuilder::GetAxisMappingCheckBoxState(const FName MappingName) const
+{
+	return UE::GetInputDeprecationCheckBoxState(GetInputSettings(), MappingName);
+}
+
+UInputSettings* FAxisMappingsNodeBuilder::GetInputSettings() const
+{
+	TArray<UObject*> OuterObjects;
+	AxisMappingsPropertyHandle->GetOuterObjects(OuterObjects);
+
+	if (OuterObjects.Num() == 1)
+	{
+		return CastChecked<UInputSettings>(OuterObjects[0]);
+	}
+	return nullptr;
+}
+
 bool FAxisMappingsNodeBuilder::GroupsRequireRebuild() const
 {
 	for (int32 GroupIndex = 0; GroupIndex < GroupedMappings.Num(); ++GroupIndex)
diff --git a/Engine/Source/Editor/DetailCustomizations/Private/InputSettingsDetails.h b/Engine/Source/Editor/DetailCustomizations/Private/InputSettingsDetails.h
--- a/Engine/Source/Editor/DetailCustomizations/Private/InputSettingsDetails.h
+++ b/Engine/Source/Editor/DetailCustomizations/Private/InputSettingsDetails.h
@@ -13,6 +13,7 @@
 #include "Templates/SharedPointer.h"
 #include "Types/SlateEnums.h"
 #include "UObject/NameTypes.h"
+#include "Styling/SlateTypes.h"
 
 class FDetailWidgetRow;
 class FText;
@@ -35,6 +36,8 @@ namespace InputSettingsDetails
 		FName SharedName;
 		IDetailGroup* DetailGroup;
 		TArray<TSharedRef<IPropertyHandle>> Mappings;
+		FName GetActionMappingName() const;
+		FName GetAxisMappingName() const;
 	};
 }
 
@@ -59,6 +62,11 @@ private:
 	void AddActionMappingToGroupButton_OnClick(const InputSettingsDetails::FMappingSet MappingSet);
 	void RemoveActionMappingGroupButton_OnClick(const InputSettingsDetails::FMappingSet MappingSet);
 
+	void OnDeprecationCheckboxChanged(ECheckBoxState NewState, const FName MappingName);
+	ECheckBoxState GetActionMappingCheckBoxState(const FName MappingName) const;
+	class UInputSettings* GetInputSettings() const;
+	
+
 	bool GroupsRequireRebuild() const;
 	void RebuildGroupedMappings();
 	void RebuildChildren()
@@ -104,6 +112,10 @@ private:
 	void AddAxisMappingToGroupButton_OnClick(const InputSettingsDetails::FMappingSet MappingSet);
 	void RemoveAxisMappingGroupButton_OnClick(const InputSettingsDetails::FMappingSet MappingSet);
 
+	void OnDeprecationCheckboxChanged(ECheckBoxState NewState, const FName MappingName);
+	ECheckBoxState GetAxisMappingCheckBoxState(const FName MappingName) const;
+	class UInputSettings* GetInputSettings() const;
+
 	bool GroupsRequireRebuild() const;
 	void RebuildGroupedMappings();
 	void RebuildChildren()
diff --git a/Engine/Source/Runtime/Engine/Classes/Components/InputComponent.h b/Engine/Source/Runtime/Engine/Classes/Components/InputComponent.h
--- a/Engine/Source/Runtime/Engine/Classes/Components/InputComponent.h
+++ b/Engine/Source/Runtime/Engine/Classes/Components/InputComponent.h
@@ -908,6 +908,8 @@ public:
 	template<class UserClass>
 	FInputAxisBinding& BindAxis( const FName AxisName, UserClass* Object, typename FInputAxisHandlerSignature::TMethodPtr< UserClass > Func )
 	{
+		LogDeprecatedBindingWarning(AxisName);
+		
 		FInputAxisBinding AB( AxisName );
 		AB.AxisDelegate.BindDelegate(Object, Func);
 		AxisBindings.Emplace(MoveTemp(AB));
@@ -933,6 +935,8 @@ public:
 	 */
 	FInputAxisBinding& BindAxis( const FName AxisName )
 	{
+		LogDeprecatedBindingWarning(AxisName);
+		
 		FInputAxisBinding AB( AxisName );
 		AxisBindings.Emplace(MoveTemp(AB));
 		return AxisBindings.Last();
@@ -1058,7 +1062,14 @@ public:
 	}
 
 private:
-
+	
+	/**
+	 * Checks if the given action/axis binding name is deprecated. If it is, then log a warning.
+	 * 
+	 * @param BindingName The name of the 
+	 */
+	ENGINE_API static void LogDeprecatedBindingWarning(const FName BindingName);
+	
 	/** Retrieves the actions bound to the input component which are triggered by a given key. Requires that the internal key map has already been built. */
 	ENGINE_API void GetActionsBoundToKey(UPlayerInput* PlayerInput, FKey Key, TArray<TSharedPtr<FInputActionBinding>>& Actions) const;
 
diff --git a/Engine/Source/Runtime/Engine/Classes/GameFramework/InputSettings.h b/Engine/Source/Runtime/Engine/Classes/GameFramework/InputSettings.h
--- a/Engine/Source/Runtime/Engine/Classes/GameFramework/InputSettings.h
+++ b/Engine/Source/Runtime/Engine/Classes/GameFramework/InputSettings.h
@@ -112,6 +112,12 @@ class UInputSettings
 	/** If enabled, virtual keyboards will have autocorrect enabled. Currently only supported on mobile devices. */
 	UPROPERTY(config, EditAnywhere, Category = "Virtual Keyboard (Mobile)")
 	uint8 bUseAutocorrect:1;
+	
+	/**
+	 * If true, then any input action or axis mappings which are marked as deprecated will throw a warning.
+	 */
+	UPROPERTY(config, EditAnywhere, Category="Bindings")
+	uint8 bEnabledLegacyMappingDeprecationWarnings : 1;
 
 	/** 
 	 * Disables autocorrect for these operating systems, even if autocorrect is enabled. Use the format "[platform] [osversion]"
@@ -157,6 +163,15 @@ private:
 	UPROPERTY(config, EditAnywhere, Category="Bindings")
 	TArray<struct FInputAxisKeyMapping> AxisMappings;
 
+	/**
+	 * List of FNames, where each name is either an Action or Axis mapping, which are marked as deprecated.
+	 * 
+	 * Deprecated action and axis mappings will trigger Blueprint compilation warnings and log warnings when bound
+	 * to in C++ code.
+	 */
+	UPROPERTY(config, VisibleAnywhere, Category="Bindings")
+	TSet<FName> DeprecatedActionAndAxisNames;
+
 	/** List of Speech Mappings */
 	UPROPERTY(config, EditAnywhere, Category = "Bindings")
 	TArray<struct FInputActionSpeechMapping> SpeechMappings;
@@ -238,6 +253,9 @@ public:
 	UFUNCTION(BlueprintCallable, Category = Settings)
 	ENGINE_API void GetActionNames(TArray<FName>& ActionNames) const;
 
+	ENGINE_API bool IsActionMappingDeprecated(const FName ActionName) const;
+	ENGINE_API void SetActionMappingDeprecated(const FName ActionName, const bool bIsDeprecated);
+
 	/** Populate a list of all defined axis names */
 	UFUNCTION(BlueprintCallable, Category = Settings)
 	ENGINE_API void GetAxisNames(TArray<FName>& AxisNames) const;
diff --git a/Engine/Source/Runtime/Engine/Private/Components/InputComponent.cpp b/Engine/Source/Runtime/Engine/Private/Components/InputComponent.cpp
--- a/Engine/Source/Runtime/Engine/Private/Components/InputComponent.cpp
+++ b/Engine/Source/Runtime/Engine/Private/Components/InputComponent.cpp
@@ -3,6 +3,7 @@
 #include "Components/InputComponent.h"
 #include "GameFramework/PlayerController.h"
 #include "GameFramework/PlayerInput.h"
+#include "GameFramework/InputSettings.h"
 
 #include UE_INLINE_GENERATED_CPP_BY_NAME(InputComponent)
 
@@ -124,6 +125,15 @@ void UInputComponent::ClearBindingsForObject(UObject* InOwner)
 	}
 }
 
+void UInputComponent::LogDeprecatedBindingWarning(const FName BindingName)
+{
+	const UInputSettings* Settings = GetDefault<UInputSettings>();
+	if (Settings->bEnabledLegacyMappingDeprecationWarnings && Settings->IsActionMappingDeprecated(BindingName))
+	{
+		UE_LOG(LogPlayerInput, Warning, TEXT("Attempting to access a deprecated action/axis mapping '%s'"), *BindingName.ToString());
+	}
+}
+
 void UInputComponent::GetActionsBoundToKey(UPlayerInput* PlayerInput, const FKey Key, TArray<TSharedPtr<FInputActionBinding>>& Actions) const
 {
 	for (const FCachedKeyToActionInfo& CachedInfo : CachedKeyToActionInfo)
@@ -243,6 +253,8 @@ bool UInputComponent::HasBindings( ) const
 
 FInputActionBinding& UInputComponent::AddActionBinding( FInputActionBinding InBinding )
 {
+	LogDeprecatedBindingWarning(InBinding.GetActionName());
+	
 	ActionBindings.Add(MakeShared<FInputActionBinding>(MoveTemp(InBinding)));
 	FInputActionBinding& Binding = *ActionBindings.Last().Get();
 	Binding.GenerateNewHandle();
@@ -314,6 +326,8 @@ void UInputComponent::RemoveActionBindingForHandle(const int32 Handle)
 
 void UInputComponent::RemoveActionBinding(const FInputActionBinding &BindingToRemove, const int32 BindingIndex)
 {
+	LogDeprecatedBindingWarning(BindingToRemove.GetActionName());
+	
 	// Potentially need to clear some pairings
 	if (BindingToRemove.bPaired)
 	{
@@ -357,6 +371,8 @@ void UInputComponent::RemoveActionBinding(const FInputActionBinding &BindingToRe
 
 void UInputComponent::RemoveActionBinding(FName InActionName, EInputEvent InKeyEvent)
 {
+	LogDeprecatedBindingWarning(InActionName);
+	
 	for (int32 ActionIdx = 0; ActionIdx < ActionBindings.Num(); ++ActionIdx)
 	{
 		const TSharedPtr<FInputActionBinding>& Binding = ActionBindings[ActionIdx];
@@ -390,6 +406,8 @@ void UInputComponent::ClearBindingValues()
 
 void UInputComponent::RemoveAxisBinding(FName AxisName)
 {
+	LogDeprecatedBindingWarning(AxisName);
+	
 	for (int32 AxisIdx = AxisBindings.Num() - 1; AxisIdx >= 0; --AxisIdx)
 	{
 		const FInputAxisBinding& Binding = AxisBindings[AxisIdx];
diff --git a/Engine/Source/Runtime/Engine/Private/UserInterface/InputSettings.cpp b/Engine/Source/Runtime/Engine/Private/UserInterface/InputSettings.cpp
--- a/Engine/Source/Runtime/Engine/Private/UserInterface/InputSettings.cpp
+++ b/Engine/Source/Runtime/Engine/Private/UserInterface/InputSettings.cpp
@@ -45,6 +45,7 @@ UInputSettings::UInputSettings(const FObjectInitializer& ObjectInitializer)
 PRAGMA_DISABLE_DEPRECATION_WARNINGS
 	, bEnableDynamicComponentInputBinding(true)
 PRAGMA_ENABLE_DEPRECATION_WARNINGS
+	, bEnabledLegacyMappingDeprecationWarnings(true)
 	, DefaultViewportMouseCaptureMode(EMouseCaptureMode::CapturePermanently_IncludingInitialMouseDown)
 	, DefaultViewportMouseLockMode(EMouseLockMode::LockOnCapture)
 	, DefaultPlayerInputClass(UPlayerInput::StaticClass())
@@ -338,7 +339,23 @@ void UInputSettings::GetActionNames(TArray<FName>& ActionNames) const
 	{
 		ActionNames.AddUnique(SpeechMapping.GetActionName());
 	}
+}
 
+bool UInputSettings::IsActionMappingDeprecated(const FName ActionName) const
+{
+	return DeprecatedActionAndAxisNames.Contains(ActionName);
+}
+
+void UInputSettings::SetActionMappingDeprecated(const FName ActionName, const bool bIsDeprecated)
+{
+	if (bIsDeprecated)
+	{
+		DeprecatedActionAndAxisNames.Add(ActionName);
+	}
+	else
+	{
+		DeprecatedActionAndAxisNames.Remove(ActionName);
+	}
 }
 
 void UInputSettings::GetAxisNames(TArray<FName>& AxisNames) const
