Update MathStructCustomization metadata extraction to include InstanceMetadata

This change updates metadata extraction for header row to include InstanceMetadata. It is useful, when assigning metadata to Vector/IntVector types via other customizations.

PR #13431: Update MathStructCustomization metadata extraction to include InstanceMetadata

diff --git a/Engine/Source/Editor/DetailCustomizations/Private/MathStructCustomizations.cpp b/Engine/Source/Editor/DetailCustomizations/Private/MathStructCustomizations.cpp
--- a/Engine/Source/Editor/DetailCustomizations/Private/MathStructCustomizations.cpp
+++ b/Engine/Source/Editor/DetailCustomizations/Private/MathStructCustomizations.cpp
@@ -328,24 +328,22 @@ template void FMathStructCustomization::ExtractNumericMetadata<uint8>(TSharedRef
 template<typename NumericType>
 void FMathStructCustomization::ExtractNumericMetadata(TSharedRef<IPropertyHandle>& PropertyHandle, FNumericMetadata<NumericType>& MetadataOut)
 {
-	FProperty* Property = PropertyHandle->GetProperty();
-
-	const FString& MetaUIMinString = Property->GetMetaData(TEXT("UIMin"));
-	const FString& MetaUIMaxString = Property->GetMetaData(TEXT("UIMax"));
-	const FString& SliderExponentString = Property->GetMetaData(TEXT("SliderExponent"));
-	const FString& DeltaString = Property->GetMetaData(TEXT("Delta"));
-	const FString& LinearDeltaSensitivityString = Property->GetMetaData(TEXT("LinearDeltaSensitivity"));
-	const FString& ShiftMultiplierString = Property->GetMetaData(TEXT("ShiftMultiplier"));
-	const FString& CtrlMultiplierString = Property->GetMetaData(TEXT("CtrlMultiplier"));
-	const FString& SupportDynamicSliderMaxValueString = Property->GetMetaData(TEXT("SupportDynamicSliderMaxValue"));
-	const FString& SupportDynamicSliderMinValueString = Property->GetMetaData(TEXT("SupportDynamicSliderMinValue"));
-	const FString& ClampMinString = Property->GetMetaData(TEXT("ClampMin"));
-	const FString& ClampMaxString = Property->GetMetaData(TEXT("ClampMax"));
+	const FString& MetaUIMinString = PropertyHandle->GetMetaData(TEXT("UIMin"));
+	const FString& MetaUIMaxString = PropertyHandle->GetMetaData(TEXT("UIMax"));
+	const FString& SliderExponentString = PropertyHandle->GetMetaData(TEXT("SliderExponent"));
+	const FString& DeltaString = PropertyHandle->GetMetaData(TEXT("Delta"));
+	const FString& LinearDeltaSensitivityString = PropertyHandle->GetMetaData(TEXT("LinearDeltaSensitivity"));
+	const FString& ShiftMultiplierString = PropertyHandle->GetMetaData(TEXT("ShiftMultiplier"));
+	const FString& CtrlMultiplierString = PropertyHandle->GetMetaData(TEXT("CtrlMultiplier"));
+	const FString& SupportDynamicSliderMaxValueString = PropertyHandle->GetMetaData(TEXT("SupportDynamicSliderMaxValue"));
+	const FString& SupportDynamicSliderMinValueString = PropertyHandle->GetMetaData(TEXT("SupportDynamicSliderMinValue"));
+	const FString& ClampMinString = PropertyHandle->GetMetaData(TEXT("ClampMin"));
+	const FString& ClampMaxString = PropertyHandle->GetMetaData(TEXT("ClampMax"));
 
 	// If no UIMin/Max was specified then use the clamp string
 	const FString& UIMinString = MetaUIMinString.Len() ? MetaUIMinString : ClampMinString;
 	const FString& UIMaxString = MetaUIMaxString.Len() ? MetaUIMaxString : ClampMaxString;
-	bool bAllowSpin = !Property->GetBoolMetaData(TEXT("NoSpinBox"));
+	bool bAllowSpin = !PropertyHandle->GetBoolMetaData(TEXT("NoSpinBox"));
 
 	NumericType ClampMin = TNumericLimits<NumericType>::Lowest();
 	NumericType ClampMax = TNumericLimits<NumericType>::Max();
@@ -424,7 +422,7 @@ void FMathStructCustomization::ExtractNumericMetadata(TSharedRef<IPropertyHandle
 	// By default allow widget to determine default interface
 	MetadataOut.TypeInterface = nullptr;
 
-	if (FStructProperty* StructProp = CastField<FStructProperty>(Property))
+	if (FStructProperty* StructProp = CastField<FStructProperty>(PropertyHandle->GetProperty()))
 	{
 		if (StructProp->Struct == TBaseStructure<FRotator>::Get())
 		{
