Fixed an issue where FQuat variables always displayed with the name "Rotation" in the details panel, regardless of DisplayName metadata or the variable's actual name. This fix ensures that the intended name is respected and shown properly in the UI.

PR #13434

diff --git a/Engine/Source/Editor/DetailCustomizations/Private/MathStructProxyCustomizations.cpp b/Engine/Source/Editor/DetailCustomizations/Private/MathStructProxyCustomizations.cpp
--- a/Engine/Source/Editor/DetailCustomizations/Private/MathStructProxyCustomizations.cpp
+++ b/Engine/Source/Editor/DetailCustomizations/Private/MathStructProxyCustomizations.cpp
@@ -161,7 +161,7 @@ void FMatrixStructCustomization<T>::CustomizeLocation(TSharedRef<class IProperty
 }
 
 template<typename T>
-void FMatrixStructCustomization<T>::CustomizeRotation(TSharedRef<class IPropertyHandle> StructPropertyHandle, FDetailWidgetRow& Row)
+void FMatrixStructCustomization<T>::CustomizeRotation(TSharedRef<class IPropertyHandle> StructPropertyHandle, FDetailWidgetRow& Row, TOptional<FText> OverrideName)
 {
 	TWeakPtr<IPropertyHandle> WeakHandlePtr = StructPropertyHandle;
 
@@ -215,7 +215,7 @@ void FMatrixStructCustomization<T>::CustomizeRotation(TSharedRef<class IProperty
 	.PasteAction(FUIAction(FExecuteAction::CreateSP(this, &FMatrixStructCustomization<T>::OnPaste, FTransformField::Rotation, WeakHandlePtr)))
 	.NameContent()
 	[
-		StructPropertyHandle->CreatePropertyNameWidget(LOCTEXT("RotationLabel", "Rotation"))
+		StructPropertyHandle->CreatePropertyNameWidget(OverrideName.IsSet() ? OverrideName.GetValue() : LOCTEXT("RotationLabel", "Rotation"))
 	]
 	.ValueContent()
 	.MinDesiredWidth(375.0f)
@@ -301,8 +301,8 @@ void FMatrixStructCustomization<T>::CustomizeChildren(TSharedRef<class IProperty
 		AxisDisplayInfo::GetAxisDisplayCoordinateSystem() == EAxisList::LeftUpForward
 		&& !StructPropertyHandle->GetProperty()->HasAnyPropertyFlags(CPF_BlueprintVisible | CPF_BlueprintReadOnly);
 
-	CustomizeLocation(StructPropertyHandle, StructBuilder.AddCustomRow(LOCTEXT("RotationLabel", "Rotation")));
-	CustomizeRotation(StructPropertyHandle, StructBuilder.AddCustomRow(LOCTEXT("LocationLabel", "Location")));
+	CustomizeLocation(StructPropertyHandle, StructBuilder.AddCustomRow(LOCTEXT("LocationLabel", "Location")));
+	CustomizeRotation(StructPropertyHandle, StructBuilder.AddCustomRow(LOCTEXT("RotationLabel", "Rotation")));
 	CustomizeScale(StructPropertyHandle, StructBuilder.AddCustomRow(LOCTEXT("ScaleLabel", "Scale")));
 }
 
@@ -782,7 +782,7 @@ TSharedRef<IPropertyTypeCustomization> FQuatStructCustomization<T>::MakeInstance
 template<typename T>
 void FQuatStructCustomization<T>::MakeHeaderRow(TSharedRef<class IPropertyHandle>& InStructPropertyHandle, FDetailWidgetRow& Row)
 {
-	this->CustomizeRotation(InStructPropertyHandle, Row);
+	this->CustomizeRotation(InStructPropertyHandle, Row, InStructPropertyHandle->GetPropertyDisplayName());
 }
 
 template<typename T>
diff --git a/Engine/Source/Editor/DetailCustomizations/Public/Customizations/MathStructProxyCustomizations.h b/Engine/Source/Editor/DetailCustomizations/Public/Customizations/MathStructProxyCustomizations.h
--- a/Engine/Source/Editor/DetailCustomizations/Public/Customizations/MathStructProxyCustomizations.h
+++ b/Engine/Source/Editor/DetailCustomizations/Public/Customizations/MathStructProxyCustomizations.h
@@ -362,7 +362,7 @@ public:
 protected:
 	/** Customization utility functions */
 	void CustomizeLocation(TSharedRef<class IPropertyHandle> StructPropertyHandle, FDetailWidgetRow& Row);
-	void CustomizeRotation(TSharedRef<class IPropertyHandle> StructPropertyHandle, FDetailWidgetRow& Row);
+	void CustomizeRotation(TSharedRef<class IPropertyHandle> StructPropertyHandle, FDetailWidgetRow& Row, TOptional<FText> OverrideName = TOptional<FText>{});
 	void CustomizeScale(TSharedRef<class IPropertyHandle> StructPropertyHandle, FDetailWidgetRow& Row);
 
 	/** FMathStructCustomization interface */
