PR #13430: FMathStructProxyCustomization : Support degrees display units when enabled in the settings to match the behavior of the FComponentTransformDetails customization.

diff --git a/Engine/Source/Editor/DetailCustomizations/Public/Customizations/MathStructProxyCustomizations.h b/Engine/Source/Editor/DetailCustomizations/Public/Customizations/MathStructProxyCustomizations.h
--- a/Engine/Source/Editor/DetailCustomizations/Public/Customizations/MathStructProxyCustomizations.h
+++ b/Engine/Source/Editor/DetailCustomizations/Public/Customizations/MathStructProxyCustomizations.h
@@ -12,6 +12,7 @@
 #include "Math/Color.h"
 #include "Math/Rotator.h"
 #include "Math/Vector.h"
+#include "Math/UnitConversion.h"
 #include "Misc/Optional.h"
 #include "PropertyHandle.h"
 #include "Styling/CoreStyle.h"
@@ -20,6 +21,8 @@
 #include "Types/SlateEnums.h"
 #include "Widgets/SWidget.h"
 #include "Widgets/Input/SNumericEntryBox.h"
+#include "Widgets/Input/NumericTypeInterface.h"
+#include "Widgets/Input/NumericUnitTypeInterface.inl"
 
 #define UE_API DETAILCUSTOMIZATIONS_API
 
@@ -255,6 +258,12 @@ inline TSharedRef<SWidget> FMathStructProxyCustomization::MakeNumericProxyWidget
 {
 	TWeakPtr<IPropertyHandle> WeakHandlePtr = StructPropertyHandle;
 
+	TSharedPtr<INumericTypeInterface<NumericType>> TypeInterface = nullptr;
+	if (FUnitConversion::Settings().ShouldDisplayUnits() && bRotationInDegrees)
+	{
+		TypeInterface = MakeShared<TNumericUnitTypeInterface<NumericType>>(EUnit::Degrees);
+	}
+
 	return
 		SNew(SNumericEntryBox<NumericType>)
 		.IsEnabled(this, &FMathStructProxyCustomization::IsValueEnabled, WeakHandlePtr)
@@ -274,6 +283,7 @@ inline TSharedRef<SWidget> FMathStructProxyCustomization::MakeNumericProxyWidget
 		.LabelPadding(FMargin(3))
 		.ToolTipTextFormat(this, &FMathStructProxyCustomization::OnGetValueToolTipTextFormat, Label)
 		.LabelLocation(SNumericEntryBox<NumericType>::ELabelLocation::Inside)
+		.TypeInterface(TypeInterface)
 		.Label()
 		[
 			SNumericEntryBox<NumericType>::BuildNarrowColorLabel(LabelBackgroundColor)
