Allow directory path struct customization to respect ForceShowPluginContent and ForceShowEngineContent meta tags.

diff --git a/Engine/Source/Editor/DetailCustomizations/Private/DirectoryPathStructCustomization.cpp b/Engine/Source/Editor/DetailCustomizations/Private/DirectoryPathStructCustomization.cpp
--- a/Engine/Source/Editor/DetailCustomizations/Private/DirectoryPathStructCustomization.cpp
+++ b/Engine/Source/Editor/DetailCustomizations/Private/DirectoryPathStructCustomization.cpp
@@ -56,6 +56,8 @@ void FDirectoryPathStructCustomization::CustomizeHeader( TSharedRef<IPropertyHan
 	const bool bRelativeToGameContentDir = StructPropertyHandle->HasMetaData( TEXT("RelativeToGameContentDir") );
 	const bool bUseRelativePath = StructPropertyHandle->HasMetaData( TEXT("RelativePath") );
 	const bool bContentDir = StructPropertyHandle->HasMetaData( TEXT("ContentDir") ) || StructPropertyHandle->HasMetaData(TEXT("LongPackageName"));
+	const bool bForceShowPluginContent = StructPropertyHandle->HasMetaData(TEXT("ForceShowPluginContent"));
+	const bool bForceShowEngineContent = StructPropertyHandle->HasMetaData(TEXT("ForceShowEngineContent"));
 
 	AbsoluteGameContentDir = FPaths::ConvertRelativePathToFull(FPaths::ProjectContentDir());
 
@@ -67,8 +69,8 @@ void FDirectoryPathStructCustomization::CustomizeHeader( TSharedRef<IPropertyHan
 		{
 			PickerWidget = SAssignNew(PickerButton, SButton)
 			.ButtonStyle( FAppStyle::Get(), "HoverHintOnly" )
-			.ToolTipText( LOCTEXT( "FolderComboToolTipText", "Choose a content directory") )
-			.OnClicked( this, &FDirectoryPathStructCustomization::OnPickContent, PathProperty.ToSharedRef() )
+			.ToolTipText( LOCTEXT("FolderComboToolTipText", "Choose a content directory") )
+			.OnClicked( this, &FDirectoryPathStructCustomization::OnPickContent, PathProperty.ToSharedRef(), bForceShowPluginContent, bForceShowEngineContent )
 			.ContentPadding(2.0f)
 			.ForegroundColor( FSlateColor::UseForeground() )
 			.IsFocusable(false)
@@ -127,14 +129,16 @@ void FDirectoryPathStructCustomization::CustomizeChildren( TSharedRef<IPropertyH
 {
 }
 
-FReply FDirectoryPathStructCustomization::OnPickContent(TSharedRef<IPropertyHandle> PropertyHandle) 
+FReply FDirectoryPathStructCustomization::OnPickContent(TSharedRef<IPropertyHandle> PropertyHandle, const bool bForceShowPluginContent, const bool bForceShowEngineContent) 
 {
 	FContentBrowserModule& ContentBrowserModule = FModuleManager::LoadModuleChecked<FContentBrowserModule>("ContentBrowser");
 	FPathPickerConfig PathPickerConfig;
 	PropertyHandle->GetValue(PathPickerConfig.DefaultPath);
 	PathPickerConfig.bAllowContextMenu = false;
 	PathPickerConfig.OnPathSelected = FOnPathSelected::CreateSP(this, &FDirectoryPathStructCustomization::OnPathPicked, PropertyHandle);
-
+	PathPickerConfig.bForceShowEngineContent = bForceShowEngineContent;
+	PathPickerConfig.bForceShowPluginContent = bForceShowPluginContent;
+	
 	FMenuBuilder MenuBuilder(true, NULL);
 	MenuBuilder.AddWidget(SNew(SBox)
 		.WidthOverride(300.0f)
diff --git a/Engine/Source/Editor/DetailCustomizations/Private/DirectoryPathStructCustomization.h b/Engine/Source/Editor/DetailCustomizations/Private/DirectoryPathStructCustomization.h
--- a/Engine/Source/Editor/DetailCustomizations/Private/DirectoryPathStructCustomization.h
+++ b/Engine/Source/Editor/DetailCustomizations/Private/DirectoryPathStructCustomization.h
@@ -25,7 +25,7 @@ public:
 private:
 
 	/** Delegate used to display a directory picker */
-	FReply OnPickContent(TSharedRef<IPropertyHandle> PropertyHandle) ;
+	FReply OnPickContent(TSharedRef<IPropertyHandle> PropertyHandle, const bool bForceShowPluginContent, const bool bForceShowEngineContent);
 
 	/** Delegate used to display a directory picker */
 	FReply OnPickDirectory(TSharedRef<IPropertyHandle> PropertyHandle, const bool bRelativeToGameContentDir, const bool bUseRelativePaths) const;
