TextureLODSettingsDetails: Fix segmentation violation when TextureMipGenSettings enum values have assigned integer values in c++ that differ from their index in the enum's list of values; GetMetadata expects the index rather than the value.
Also remove the InitialSelection for the TextureMipGenSettings combo boxes of the first item in the box; making that initial selection causes the user's selection of the first value in the list to fail to activate, until they first click another entry in the list, and this failure to activate makes it fail to normalize the value name in the combo box. Removing the initial selection causes the same behavior of using the first element as the initially presented value, but with the fix that the activation is executed when the user clicks it.
#rnx

diff --git a/Engine/Source/Editor/DetailCustomizations/Private/TextureLODSettingsDetails.cpp b/Engine/Source/Editor/DetailCustomizations/Private/TextureLODSettingsDetails.cpp
--- a/Engine/Source/Editor/DetailCustomizations/Private/TextureLODSettingsDetails.cpp
+++ b/Engine/Source/Editor/DetailCustomizations/Private/TextureLODSettingsDetails.cpp
@@ -244,7 +244,6 @@ void FTextureLODGroupLayout::GenerateChildContent(IDetailChildrenBuilder& Childr
 				.OptionsSource(&FilterComboList)
 				.OnGenerateWidget(this, &FTextureLODGroupLayout::MakeMinMagFilterComboWidget)
 				.OnSelectionChanged(this, &FTextureLODGroupLayout::OnMinMagFilterChanged)
-				.InitiallySelectedItem(FilterComboList[0])
 				.ContentPadding(0.f)
 				.Content()
 				[
@@ -270,7 +269,6 @@ void FTextureLODGroupLayout::GenerateChildContent(IDetailChildrenBuilder& Childr
 				.OptionsSource(&FilterComboList)
 				.OnGenerateWidget(this, &FTextureLODGroupLayout::MakeMipFilterComboWidget)
 				.OnSelectionChanged(this, &FTextureLODGroupLayout::OnMipFilterChanged)
-				.InitiallySelectedItem(FilterComboList[0])
 				.ContentPadding(0.f)
 				.Content()
 				[
@@ -298,7 +296,6 @@ void FTextureLODGroupLayout::GenerateChildContent(IDetailChildrenBuilder& Childr
 				.OptionsSource(&MipGenSettingsComboList)
 				.OnGenerateWidget(this, &FTextureLODGroupLayout::MakeMipGenSettingsComboWidget)
 				.OnSelectionChanged(this, &FTextureLODGroupLayout::OnMipGenSettingsChanged)
-				.InitiallySelectedItem(MipGenSettingsComboList[0])
 				.ContentPadding(0.f)
 				.Content()
 				[
@@ -460,7 +457,7 @@ FText FTextureLODGroupLayout::GetMipFilterComboBoxContent() const
 TSharedRef<SWidget> FTextureLODGroupLayout::MakeMipGenSettingsComboWidget(TSharedPtr<TextureMipGenSettings> InItem)
 {
 	UEnum* TextureGroupEnum = StaticEnum<TextureMipGenSettings>();
-	const FString& MipGenSettingsName = TextureGroupEnum->GetMetaData(TEXT("DisplayName"), *InItem.Get());
+	const FString& MipGenSettingsName = TextureGroupEnum->GetMetaData(TEXT("DisplayName"), TextureGroupEnum->GetIndexByValue(*InItem.Get()));
 
 	return SNew(STextBlock).Text(FText::FromString(MipGenSettingsName)).Font(IDetailLayoutBuilder::GetDetailFont());
 }
@@ -483,7 +480,7 @@ FText FTextureLODGroupLayout::GetMipGenSettingsComboBoxToolTip() const
 FText FTextureLODGroupLayout::GetMipGenSettingsComboBoxContent() const
 {
 	UEnum* TextureGroupEnum = StaticEnum<TextureMipGenSettings>();
-	const FString& MipGenSettingsName = TextureGroupEnum->GetMetaData(TEXT("DisplayName"), LodGroup->MipGenSettings);
+	const FString& MipGenSettingsName = TextureGroupEnum->GetMetaData(TEXT("DisplayName"), TextureGroupEnum->GetIndexByValue(LodGroup->MipGenSettings));
 	
 	return FText::FromString(MipGenSettingsName);
 }
