Editor: Added missing NotifyPreChange/NotifyPostChange

This is needed for Sequencer to store the pre edited value (so that it can restore it when done animating)

diff --git a/Engine/Source/Editor/DetailCustomizations/Private/ColorStructCustomization.cpp b/Engine/Source/Editor/DetailCustomizations/Private/ColorStructCustomization.cpp
--- a/Engine/Source/Editor/DetailCustomizations/Private/ColorStructCustomization.cpp
+++ b/Engine/Source/Editor/DetailCustomizations/Private/ColorStructCustomization.cpp
@@ -316,7 +316,9 @@ void FColorStructCustomization::OnSetColorFromColorPicker(FLinearColor NewColor)
 
 	EPropertyValueSetFlags::Type PropertyFlags = EPropertyValueSetFlags::NotTransactable;
 	PropertyFlags |= bIsInteractive ? EPropertyValueSetFlags::InteractiveChange : 0;
+	StructPropertyHandle->NotifyPreChange();
 	StructPropertyHandle->SetValueFromFormattedString(LastPickerColorString, PropertyFlags);
+	StructPropertyHandle->NotifyPostChange(PropertyFlags);
 	StructPropertyHandle->NotifyFinishedChangingProperties();
 }
 
