Add ISplineDetailsProvider and make FSplineComponentVisualizer implement the interface. Use the interface in FSplineComponentDetails

diff --git a/Engine/Source/Editor/ComponentVisualizers/Private/SplineComponentVisualizer.cpp b/Engine/Source/Editor/ComponentVisualizers/Private/SplineComponentVisualizer.cpp
--- a/Engine/Source/Editor/ComponentVisualizers/Private/SplineComponentVisualizer.cpp
+++ b/Engine/Source/Editor/ComponentVisualizers/Private/SplineComponentVisualizer.cpp
@@ -30,6 +30,9 @@
 #include "EngineUtils.h"
 #include "CanvasItem.h"
 #include "CanvasTypes.h"
+#include "LevelEditorViewport.h"
+#include "UnrealEdGlobals.h"
+#include "Editor/UnrealEdEngine.h"
 #include "Math/UnrealMathUtility.h"
 #include "Widgets/Input/SButton.h"
 #include "Widgets/Input/SComboButton.h"
@@ -387,6 +390,8 @@ PRAGMA_ENABLE_DEPRECATION_WARNINGS
 	}
 	
 	SelectionState = NewObject<USplineComponentVisualizerSelectionState>(GetTransientPackage(), TEXT("SelectionState"), RF_Transactional);
+
+	IModularFeatures::Get().RegisterModularFeature(ISplineDetailsProvider::GetModularFeatureName(), this);
 }
 
 void FSplineComponentVisualizer::OnRegister()
@@ -632,6 +637,7 @@ bool FSplineComponentVisualizer::ShouldShowForSelectedSubcomponents(const UActor
 
 FSplineComponentVisualizer::~FSplineComponentVisualizer()
 {
+	IModularFeatures::Get().UnregisterModularFeature(ISplineDetailsProvider::GetModularFeatureName(), this);
 	FSplineComponentVisualizerCommands::Unregister();
 }
 
@@ -3533,6 +3539,11 @@ bool FSplineComponentVisualizer::HandleSelectFirstLastSplinePoint(USplineCompone
 	return bResetEditedSplineComponent;
 }
 
+void FSplineComponentVisualizer::HandleSelectPrevNextSplinePoint(bool bNext, bool bAddToSelection)
+{
+	OnSelectPrevNextSplinePoint(bNext, bAddToSelection);
+}
+
 bool FSplineComponentVisualizer::HandleSelectAllSplinePoints(USplineComponent* InSplineComponent)
 {
 	const FScopedTransaction Transaction(LOCTEXT("SelectAllSplinePoints", "Select All Spline Points"));
@@ -4255,4 +4266,34 @@ void FSplineComponentVisualizer::OnDeselectedInEditor(TObjectPtr<USplineComponen
 }
 
 
+bool FSplineComponentVisualizer::ShouldUseForSpline(const USplineComponent* SplineComponent) const
+{
+	if (!GetEditedSplineComponent())
+	{
+		return false;
+	}
+	
+	// This implementation should be used when the classic visualizer is enabled for this spline
+	return (GetEditedSplineComponent() == SplineComponent || GetEditedSplineComponent()->GetArchetype() == SplineComponent)
+		&& SplineComponentVisualizerLocals::IsEnabledForSpline(SplineComponent);
+}
+
+bool FSplineComponentVisualizer::IsEnabledForSpline(const USplineComponent* InSplineComponent) const
+{
+	auto ModularFeatureName = ISplineComponentVisualizerSuppressor::GetModularFeatureName();
+
+	if (IModularFeatures::Get().IsModularFeatureAvailable(ModularFeatureName))
+	{
+		return !IModularFeatures::Get().GetModularFeature<ISplineComponentVisualizerSuppressor>(ModularFeatureName).ShouldSuppress(InSplineComponent);
+	}
+	
+	return true;
+}
+
+void FSplineComponentVisualizer::ActivateVisualization()
+{
+	TSharedPtr<FComponentVisualizer> ComponentVis = StaticCastSharedPtr<FComponentVisualizer>(SharedThis(this).ToSharedPtr());
+	GUnrealEd->ComponentVisManager.SetActiveComponentVis(GCurrentLevelEditingViewportClient, ComponentVis);
+}
+
 #undef LOCTEXT_NAMESPACE
diff --git a/Engine/Source/Editor/ComponentVisualizers/Public/SplineComponentVisualizer.h b/Engine/Source/Editor/ComponentVisualizers/Public/SplineComponentVisualizer.h
--- a/Engine/Source/Editor/ComponentVisualizers/Public/SplineComponentVisualizer.h
+++ b/Engine/Source/Editor/ComponentVisualizers/Public/SplineComponentVisualizer.h
@@ -21,6 +21,7 @@
 #include "Math/UnrealMathSSE.h"
 #include "Math/Vector.h"
 #include "Misc/AssertionMacros.h"
+#include "SplineDetailsProvider.h"
 #include "Templates/SharedPointer.h"
 #include "UObject/GCObject.h"
 #include "UObject/Object.h"
@@ -258,6 +259,9 @@ enum class ESplineComponentSnapMode
 
 /** SplineComponent visualizer/edit functionality */
 class FSplineComponentVisualizer : public FComponentVisualizer, public FGCObject
+PRAGMA_DISABLE_INTERNAL_WARNINGS
+	, public ISplineDetailsProvider
+PRAGMA_DISABLE_INTERNAL_WARNINGS
 {
 public:
 	UE_API FSplineComponentVisualizer();
@@ -297,21 +301,23 @@ public:
 	UE_API virtual void GenerateContextMenuSections(FMenuBuilder& InMenuBuilder) const;
 
 	/** Get the spline component we are currently editing */
-	UE_API USplineComponent* GetEditedSplineComponent() const;
+	UE_API virtual USplineComponent* GetEditedSplineComponent() const override;
 
-	const TSet<int32>& GetSelectedKeys() const { check(SelectionState); return SelectionState->GetSelectedKeys(); }
+	const virtual TSet<int32>& GetSelectedKeys() const override { check(SelectionState); return SelectionState->GetSelectedKeys(); }
 
 	/** Select first or last spline point, returns true if the spline component being edited has changed */
-	UE_API bool HandleSelectFirstLastSplinePoint(USplineComponent* InSplineComponent, bool bFirstPoint);
+	UE_API virtual bool HandleSelectFirstLastSplinePoint(USplineComponent* InSplineComponent, bool bFirstPoint) override;
+
+	UE_API virtual void HandleSelectPrevNextSplinePoint(bool bNext, bool bAddToSelection) override;
 
 	/** Select all spline points, , returns true if the spline component being edited has changed */
-	UE_API bool HandleSelectAllSplinePoints(USplineComponent* InSplineComponent);
+	UE_API virtual bool HandleSelectAllSplinePoints(USplineComponent* InSplineComponent) override;
 
 	/** Select next or prev spline point, loops when last point is currently selected */
 	UE_API void OnSelectPrevNextSplinePoint(bool bNextPoint, bool bAddToSelection);
 
 	/** Sets the new cached rotation on the visualizer */
-	UE_API void SetCachedRotation(const FQuat& NewRotation);
+	UE_API virtual void SetCachedRotation(const FQuat& NewRotation) override;
 protected:
 
 	/** Determine if any selected key index is out of range (perhaps because something external has modified the spline) */
@@ -484,6 +490,10 @@ protected:
 	
 	UE_API void OnDeselectedInEditor(TObjectPtr<USplineComponent> SplineComponent);
 
+	UE_API bool ShouldUseForSpline(const USplineComponent* SplineComponent) const override;
+	UE_API void ActivateVisualization() override;
+	UE_API bool IsEnabledForSpline(const USplineComponent* InSplineComponent) const;
+	
 	// FGCObject interface
 	UE_API virtual void AddReferencedObjects(FReferenceCollector& Collector);
 	virtual FString GetReferencerName() const override
@@ -564,7 +574,7 @@ public:
 
 	virtual ~ISplineComponentVisualizerSuppressor() = default;
 	
-	static FName GetModularFeatureName() { return "SplineComponentVisualizerSuppressor"; }
+	static FName GetModularFeatureName() { return FName(TEXT("SplineComponentVisualizerSuppressor")); }
 	
 	virtual bool ShouldSuppress(const USplineComponent* SplineComponent) const = 0;
 };
diff --git a/Engine/Source/Editor/ComponentVisualizers/Public/SplineDetailsProvider.h b/Engine/Source/Editor/ComponentVisualizers/Public/SplineDetailsProvider.h
new file mode 100644
--- /dev/null
+++ b/Engine/Source/Editor/ComponentVisualizers/Public/SplineDetailsProvider.h
@@ -0,0 +1,52 @@
+// Copyright Epic Games, Inc. All Rights Reserved.
+
+#pragma once
+
+#include "Features/IModularFeature.h"
+#include "UObject/NameTypes.h"
+
+class USplineComponent;
+
+class UE_INTERNAL ISplineDetailsProvider;
+
+/**
+ * Interface for spline details panels to interact with spline editing tools.
+ * 
+ * This interface was extracted from FSplineComponentVisualizer to integrate a
+ * replacement of that visualizer into the spline component details customization.
+ */
+class ISplineDetailsProvider : public IModularFeature
+{
+public:
+	virtual ~ISplineDetailsProvider() = default;
+
+	/** Get the modular feature name for registration/lookup */
+	static FName GetModularFeatureName() 
+	{ 
+		return FName(TEXT("SplineDetailsInterface")); 
+	}
+
+	/** Check if this implementation should be used for the given spline component */
+	virtual bool ShouldUseForSpline(const USplineComponent* SplineComponent) const = 0;
+
+	/** Get the currently selected spline point indices */
+	virtual const TSet<int32>& GetSelectedKeys() const = 0;
+
+	/** Get the spline component currently being edited by this system */
+	virtual USplineComponent* GetEditedSplineComponent() const = 0;
+
+	/** Select the first or last spline point */
+	virtual bool HandleSelectFirstLastSplinePoint(USplineComponent* SplineComp, bool bFirst) = 0;
+
+	/** Select the previous or next spline point */
+	virtual void HandleSelectPrevNextSplinePoint(bool bNext, bool bAddToSelection) = 0;
+
+	/** Select all spline points */
+	virtual bool HandleSelectAllSplinePoints(USplineComponent* SplineComp) = 0;
+
+	/** Set the cached rotation for the spline component */
+	virtual void SetCachedRotation(const FQuat& NewRotation) {}
+
+	/** Activate the visualization for the edited spline component */
+	virtual void ActivateVisualization() {}
+};
diff --git a/Engine/Source/Editor/DetailCustomizations/Private/SplineComponentDetails.cpp b/Engine/Source/Editor/DetailCustomizations/Private/SplineComponentDetails.cpp
--- a/Engine/Source/Editor/DetailCustomizations/Private/SplineComponentDetails.cpp
+++ b/Engine/Source/Editor/DetailCustomizations/Private/SplineComponentDetails.cpp
@@ -21,6 +21,7 @@
 #include "Editor/EditorEngine.h"
 #include "Editor/UnrealEdEngine.h"
 #include "Engine/Blueprint.h"
+#include "Features/IModularFeatures.h"
 #include "Fonts/SlateFontInfo.h"
 #include "Framework/Commands/UIAction.h"
 #include "Framework/Commands/UICommandInfo.h"
@@ -59,6 +60,7 @@
 #include "Serialization/Archive.h"
 #include "SlotBase.h"
 #include "SplineComponentVisualizer.h"
+#include "SplineDetailsProvider.h"
 #include "SplineMetadataDetailsFactory.h"
 #include "Styling/AppStyle.h"
 #include "Styling/SlateColor.h"
@@ -320,7 +322,9 @@ private:
 	FSharedRotatorValue Rotation;
 	TSharedValue<ESplinePointType::Type> PointType;
 
-	TSharedPtr<FSplineComponentVisualizer> SplineVisualizer;
+	void TryAcquireDetailsInterface();
+	ISplineDetailsProvider* SplineDetailsInterface = nullptr;
+	
 	TArray<FProperty*> SplineProperties;
 	TArray<TSharedPtr<FString>> SplinePointTypes;
 	TSharedPtr<ISplineMetadataDetails> SplineMetaDataDetails;
@@ -337,11 +341,6 @@ bool FSplinePointDetails::bAlreadyWarnedInvalidIndex = false;
 FSplinePointDetails::FSplinePointDetails(USplineComponent* InOwningSplineComponent)
 	: SplineComp(nullptr)
 {
-	if (TSharedPtr<FComponentVisualizer> Visualizer = GUnrealEd->FindComponentVisualizer(InOwningSplineComponent->GetClass()))
-	{
-		SplineVisualizer = StaticCastSharedPtr<FSplineComponentVisualizer>(Visualizer);
-	}
-
 	for (const FName& Property : USplineComponent::GetSplinePropertyNames())
 	{
 		SplineProperties.Add(FindFProperty<FProperty>(USplineComponent::StaticClass(), Property));
@@ -376,6 +375,41 @@ FSplinePointDetails::FSplinePointDetails(USplineComponent* InOwningSplineCompone
 	bAlreadyWarnedInvalidIndex = false;
 }
 
+void FSplinePointDetails::TryAcquireDetailsInterface()
+{
+	if (!SplineComp && !SplineCompArchetype)
+	{
+		return;
+	}
+	
+	if (!SplineDetailsInterface)
+	{
+		ISplineDetailsProvider* FoundInterface = nullptr;
+		for (ISplineDetailsProvider* Interface : IModularFeatures::Get().GetModularFeatureImplementations<ISplineDetailsProvider>(
+			ISplineDetailsProvider::GetModularFeatureName()))
+		{
+			if (Interface && Interface->ShouldUseForSpline(SplineComp ? SplineComp : SplineCompArchetype))
+			{
+				FoundInterface = Interface;
+			}
+		}
+
+		if (FoundInterface)
+		{
+			SplineDetailsInterface = FoundInterface;
+			
+			IModularFeatures::Get().OnModularFeatureUnregistered().AddSPLambda(this, [this](const FName& Type, class IModularFeature* ModularFeature)
+			{
+				if (SplineDetailsInterface == ModularFeature)
+				{
+					SplineDetailsInterface = nullptr;
+					IModularFeatures::Get().OnModularFeatureUnregistered().RemoveAll(this);
+				}
+			});
+		}
+	}
+}
+
 void FSplinePointDetails::SetOnRebuildChildren(FSimpleDelegate InOnRegenerateChildren)
 {
 	OnRegenerateChildren = InOnRegenerateChildren;
@@ -771,7 +805,8 @@ void FSplinePointDetails::GenerateChildContent(IDetailChildrenBuilder& ChildrenB
 			];
 	}
 
-	if (SplineVisualizer.IsValid() && SplineVisualizer->GetSelectedKeys().Num() > 0)
+	TryAcquireDetailsInterface();
+	if (SplineDetailsInterface && SplineDetailsInterface->GetSelectedKeys().Num() > 0)
 	{
 		for (TObjectIterator<UClass> ClassIterator; ClassIterator; ++ClassIterator)
 		{
@@ -798,11 +833,13 @@ void FSplinePointDetails::Tick(float DeltaTime)
 
 void FSplinePointDetails::UpdateValues()
 {
+	TryAcquireDetailsInterface();
+	
 	// If this is a blueprint spline, always update the spline component based on 
 	// the spline component visualizer's currently edited spline component.
 	if (SplineCompArchetype)
 	{
-		USplineComponent* EditedSplineComp = SplineVisualizer.IsValid() ? SplineVisualizer->GetEditedSplineComponent() : nullptr;
+		USplineComponent* EditedSplineComp = SplineDetailsInterface ? SplineDetailsInterface->GetEditedSplineComponent() : nullptr;
 
 		if (!EditedSplineComp || (EditedSplineComp->GetArchetype() != SplineCompArchetype))
 		{
@@ -812,13 +849,13 @@ void FSplinePointDetails::UpdateValues()
 		SplineComp = EditedSplineComp;
 	}
 
-	if (!SplineComp || !SplineVisualizer.IsValid())
+	if (!SplineComp || !SplineDetailsInterface)
 	{
 		return;
 	}
 
 	bool bNeedsRebuild = false;
-	const TSet<int32>& NewSelectedKeys = SplineVisualizer->GetSelectedKeys();
+	const TSet<int32>& NewSelectedKeys = SplineDetailsInterface->GetSelectedKeys();
 
 	if (NewSelectedKeys.Num() != SelectedKeys.Num())
 	{
@@ -915,7 +952,9 @@ PRAGMA_DISABLE_DEPRECATION_WARNINGS
 	// I have made this assumption valid by introducing CanSetInputKey() which verifies this by checking with the selected component.
 	// This assumption is necessary because there is no interface on the component which allows us to modify input keys, we must directly write to SplineCurves.
 	
-	if ((CommitInfo != ETextCommit::OnEnter && CommitInfo != ETextCommit::OnUserMovedFocus) || !SplineComp)
+	TryAcquireDetailsInterface();
+	
+	if ((CommitInfo != ETextCommit::OnEnter && CommitInfo != ETextCommit::OnUserMovedFocus) || !SplineComp || !SplineDetailsInterface)
 	{
 		return;
 	}
@@ -996,7 +1035,9 @@ PRAGMA_ENABLE_DEPRECATION_WARNINGS
 
 void FSplinePointDetails::OnSetPosition(FVector::FReal NewValue, ETextCommit::Type CommitInfo, EAxis::Type Axis)
 {
-	if (!SplineComp)
+	TryAcquireDetailsInterface();
+	
+	if (!SplineComp || !SplineDetailsInterface)
 	{
 		return;
 	}
@@ -1041,7 +1082,9 @@ void FSplinePointDetails::OnSetPosition(FVector::FReal NewValue, ETextCommit::Ty
 
 void FSplinePointDetails::OnSetArriveTangent(FVector::FReal NewValue, ETextCommit::Type CommitInfo, EAxis::Type Axis)
 {
-	if (!SplineComp)
+	TryAcquireDetailsInterface();
+	
+	if (!SplineComp || !SplineDetailsInterface)
 	{
 		return;
 	}
@@ -1088,7 +1131,9 @@ void FSplinePointDetails::OnSetArriveTangent(FVector::FReal NewValue, ETextCommi
 
 void FSplinePointDetails::OnSetLeaveTangent(FVector::FReal NewValue, ETextCommit::Type CommitInfo, EAxis::Type Axis)
 {
-	if (!SplineComp)
+	TryAcquireDetailsInterface();
+	
+	if (!SplineComp || !SplineDetailsInterface)
 	{
 		return;
 	}
@@ -1134,7 +1179,9 @@ void FSplinePointDetails::OnSetLeaveTangent(FVector::FReal NewValue, ETextCommit
 
 void FSplinePointDetails::OnSetRotation(FRotator::FReal NewValue, ETextCommit::Type CommitInfo, EAxis::Type Axis)
 {
-	if (!SplineComp)
+	TryAcquireDetailsInterface();
+	
+	if (!SplineComp || !SplineDetailsInterface)
 	{
 		return;
 	}
@@ -1190,10 +1237,7 @@ void FSplinePointDetails::OnSetRotation(FRotator::FReal NewValue, ETextCommit::T
 		}
 	}
 
-	if (SplineVisualizer.IsValid())
-	{
-		SplineVisualizer->SetCachedRotation(NewRotationRelative);
-	}
+	SplineDetailsInterface->SetCachedRotation(NewRotationRelative);
 
 	if (CommitInfo == ETextCommit::OnEnter || CommitInfo == ETextCommit::OnUserMovedFocus)
 	{
@@ -1211,7 +1255,9 @@ void FSplinePointDetails::OnSetRotation(FRotator::FReal NewValue, ETextCommit::T
 
 void FSplinePointDetails::OnSetScale(FVector::FReal NewValue, ETextCommit::Type CommitInfo, EAxis::Type Axis)
 {
-	if (!SplineComp)
+	TryAcquireDetailsInterface();
+	
+	if (!SplineComp || !SplineDetailsInterface)
 	{
 		return;
 	}
@@ -1267,7 +1313,9 @@ FText FSplinePointDetails::GetPointType() const
 
 void FSplinePointDetails::OnSplinePointTypeChanged(TSharedPtr<FString> NewValue, ESelectInfo::Type SelectInfo)
 {
-	if (!SplineComp)
+	TryAcquireDetailsInterface();
+	
+	if (!SplineComp || !SplineDetailsInterface)
 	{
 		return;
 	}
@@ -1375,7 +1423,9 @@ USplineComponent* FSplinePointDetails::GetSplineComponentToVisualize() const
 
 FReply FSplinePointDetails::OnSelectFirstLastSplinePoint(bool bFirst)
 {
-	if (SplineVisualizer.IsValid())
+	TryAcquireDetailsInterface();
+	
+	if (SplineDetailsInterface)
 	{
 		bool bActivateComponentVis = false;
 
@@ -1387,12 +1437,11 @@ FReply FSplinePointDetails::OnSelectFirstLastSplinePoint(bool bFirst)
 
 		if (SplineComp)
 		{
-			if (SplineVisualizer->HandleSelectFirstLastSplinePoint(SplineComp, bFirst))
+			if (SplineDetailsInterface->HandleSelectFirstLastSplinePoint(SplineComp, bFirst))
 			{
 				if (bActivateComponentVis)
 				{
-					TSharedPtr<FComponentVisualizer> Visualizer = StaticCastSharedPtr<FComponentVisualizer>(SplineVisualizer);
-					GUnrealEd->ComponentVisManager.SetActiveComponentVis(GCurrentLevelEditingViewportClient, Visualizer);
+					SplineDetailsInterface->ActivateVisualization();
 				}
 			}
 		}
@@ -1402,16 +1451,20 @@ FReply FSplinePointDetails::OnSelectFirstLastSplinePoint(bool bFirst)
 
 FReply FSplinePointDetails::OnSelectPrevNextSplinePoint(bool bNext, bool bAddToSelection)
 {
-	if (SplineVisualizer.IsValid())
+	TryAcquireDetailsInterface();
+	
+	if (SplineDetailsInterface)
 	{
-		SplineVisualizer->OnSelectPrevNextSplinePoint(bNext, bAddToSelection);
+		SplineDetailsInterface->HandleSelectPrevNextSplinePoint(bNext, bAddToSelection);
 	}
 	return FReply::Handled();
 }
 
 FReply FSplinePointDetails::OnSelectAllSplinePoints()
 {
-	if (SplineVisualizer.IsValid())
+	TryAcquireDetailsInterface();
+
+	if (SplineDetailsInterface)
 	{
 		bool bActivateComponentVis = false;
 
@@ -1423,12 +1476,11 @@ FReply FSplinePointDetails::OnSelectAllSplinePoints()
 
 		if (SplineComp)
 		{
-			if (SplineVisualizer->HandleSelectAllSplinePoints(SplineComp))
+			if (SplineDetailsInterface->HandleSelectAllSplinePoints(SplineComp))
 			{
 				if (bActivateComponentVis)
 				{
-					TSharedPtr<FComponentVisualizer> Visualizer = StaticCastSharedPtr<FComponentVisualizer>(SplineVisualizer);
-					GUnrealEd->ComponentVisManager.SetActiveComponentVis(GCurrentLevelEditingViewportClient, Visualizer);
+					SplineDetailsInterface->ActivateVisualization();
 				}
 			}
 		}
