This change fixes an issue where Spline-related UI elements were displaying XYZ axis labels and tooltips even when the editor was running in LUF coordinate mode.
Fix Details:
Updated the context menu and details panel for Splines so that axis labels, tooltips, and action names now correctly reflect LUF when active.
This includes context menu areas like SnapAllToSelected, SnapToLastSelected,  and SetLockedAxis (temporarily disabled).
This also includes details panel areas for Spline settings like Selected Spline Node?s Location, Rotation, Scale, Arrive Tangent, Leave Tangent, and Up Vector
Logic was added to ensure these UI elements are now dynamically aligned with the current coordinate system.
XYZ (FLU) remains the fallback when LUF is not enabled.

diff --git a/Engine/Source/Editor/ComponentVisualizers/Private/SplineComponentVisualizer.cpp b/Engine/Source/Editor/ComponentVisualizers/Private/SplineComponentVisualizer.cpp
--- a/Engine/Source/Editor/ComponentVisualizers/Private/SplineComponentVisualizer.cpp
+++ b/Engine/Source/Editor/ComponentVisualizers/Private/SplineComponentVisualizer.cpp
@@ -34,6 +34,7 @@
 #include "UnrealEdGlobals.h"
 #include "Editor/UnrealEdEngine.h"
 #include "Math/UnrealMathUtility.h"
+#include "Misc/AxisDisplayInfo.h"
 #include "Widgets/Input/SButton.h"
 #include "Widgets/Input/SComboButton.h"
 #include "Widgets/Input/SEditableTextBox.h"
@@ -223,18 +224,32 @@ public:
 		UI_COMMAND(AlignKeyToActor, "Align to Actor", "Align selected spline point to actor, Ctrl-LMB to select the actor after choosing this option.", EUserInterfaceActionType::Button, FInputChord());
 		UI_COMMAND(AlignKeyPerpendicularToActor, "Align Perpendicular to Actor", "Align perpendicular  selected spline point to actor, Ctrl-LMB to select the actor after choosing this option.", EUserInterfaceActionType::Button, FInputChord());
 		UI_COMMAND(ToggleSnapTangentAdjustments, "Allow Tangents Updates On Snap", "Allow tangents to update when performing snap operations on points.", EUserInterfaceActionType::ToggleButton, FInputChord());
-		UI_COMMAND(SnapAllToSelectedX, "Snap All To Selected X", "Snap all spline points to selected spline point world X position.", EUserInterfaceActionType::Button, FInputChord());
-		UI_COMMAND(SnapAllToSelectedY, "Snap All To Selected Y", "Snap all spline points to selected spline point world Y position.", EUserInterfaceActionType::Button, FInputChord());
-		UI_COMMAND(SnapAllToSelectedZ, "Snap All To Selected Z", "Snap all spline points to selected spline point world Z position.", EUserInterfaceActionType::Button, FInputChord());
-		UI_COMMAND(SnapToLastSelectedX, "Snap To Last Selected X", "Snap selected spline points to world X position of last selected spline point.", EUserInterfaceActionType::Button, FInputChord());
-		UI_COMMAND(SnapToLastSelectedY, "Snap To Last Selected Y", "Snap selected spline points to world Y position of last selected spline point.", EUserInterfaceActionType::Button, FInputChord());
-		UI_COMMAND(SnapToLastSelectedZ, "Snap To Last Selected Z", "Snap selected spline points to world Z position of last selected spline point.", EUserInterfaceActionType::Button, FInputChord());
+
+#define AXIS_UI_COMMAND(CommandName, Axis, LabelText, TooltipText, ButtonType) \
+	CommandName = FUICommandInfoDecl(this->AsShared(), \
+		TEXT(#CommandName), \
+		FText::Format(LOCTEXT(#CommandName, LabelText), AxisDisplayInfo::GetAxisDisplayName(Axis)), \
+		FText::Format(LOCTEXT(#CommandName TEXT("_ToolTip"), TooltipText), AxisDisplayInfo::GetAxisDisplayName(Axis))) \
+		.UserInterfaceType(ButtonType) \
+		.DefaultChord(FInputChord());
+
+		AXIS_UI_COMMAND(SnapAllToSelectedY, EAxisList::Left, "Snap All To Selected {0}", "Snap all spline points to selected spline point world {0} position.", EUserInterfaceActionType::Button)
+		AXIS_UI_COMMAND(SnapAllToSelectedZ, EAxisList::Up, "Snap All To Selected {0}", "Snap all spline points to selected spline point world {0} position.", EUserInterfaceActionType::Button)
+		AXIS_UI_COMMAND(SnapAllToSelectedX, EAxisList::Forward ,"Snap All To Selected {0}", "Snap all spline points to selected spline point world {0} position.", EUserInterfaceActionType::Button)
+
+		AXIS_UI_COMMAND(SnapToLastSelectedY, EAxisList::Left, "Snap To Last Selected {0}", "Snap selected spline points to world {0} position of last selected spline point.", EUserInterfaceActionType::Button)
+		AXIS_UI_COMMAND(SnapToLastSelectedZ, EAxisList::Up, "Snap To Last Selected {0}", "Snap selected spline points to world {0} position of last selected spline point.", EUserInterfaceActionType::Button)
+		AXIS_UI_COMMAND(SnapToLastSelectedX, EAxisList::Forward, "Snap To Last Selected {0}", "Snap selected spline points to world {0} position of last selected spline point.", EUserInterfaceActionType::Button)
+
+		AXIS_UI_COMMAND(SetLockedAxisY, EAxisList::Left, "{0}", "Fix {0} axis when adding new spline points.", EUserInterfaceActionType::RadioButton)
+		AXIS_UI_COMMAND(SetLockedAxisZ, EAxisList::Up, "{0}", "Fix {0} axis when adding new spline points.", EUserInterfaceActionType::RadioButton)
+		AXIS_UI_COMMAND(SetLockedAxisX, EAxisList::Forward, "{0}", "Fix {0} axis when adding new spline points.", EUserInterfaceActionType::RadioButton)
+#undef AXIS_UI_COMMAND
+
+
 		UI_COMMAND(StraightenToNext, "Straighten To Next Point", "Straighten selected points toward next sequential point", EUserInterfaceActionType::Button, FInputChord());
 		UI_COMMAND(StraightenToPrevious, "Straighten To Previous Point", "Straighten selected points toward previous sequential point", EUserInterfaceActionType::Button, FInputChord());
 		UI_COMMAND(SetLockedAxisNone, "None", "New spline point axis is not fixed.", EUserInterfaceActionType::RadioButton, FInputChord());
-		UI_COMMAND(SetLockedAxisX, "X", "Fix X axis when adding new spline points.", EUserInterfaceActionType::RadioButton, FInputChord());
-		UI_COMMAND(SetLockedAxisY, "Y", "Fix Y axis when adding new spline points.", EUserInterfaceActionType::RadioButton, FInputChord());
-		UI_COMMAND(SetLockedAxisZ, "Z", "Fix Z axis when adding new spline points.", EUserInterfaceActionType::RadioButton, FInputChord());
 		UI_COMMAND(VisualizeRollAndScale, "Visualize Roll and Scale", "Whether the visualization should show roll and scale on this spline.", EUserInterfaceActionType::ToggleButton, FInputChord());
 		UI_COMMAND(DiscontinuousSpline, "Allow Discontinuous Splines", "Whether the visualization allows Arrive and Leave tangents to be set separately.", EUserInterfaceActionType::ToggleButton, FInputChord());
 		UI_COMMAND(ToggleClosedLoop, "Closed Loop", "Toggle the Closed Loop setting of the spline", EUserInterfaceActionType::ToggleButton, FInputChord());
@@ -3859,13 +3874,28 @@ void FSplineComponentVisualizer::GenerateSnapAlignSubMenu(FMenuBuilder& MenuBuil
 	MenuBuilder.AddMenuEntry(FSplineComponentVisualizerCommands::Get().AlignKeyToActor);
 	MenuBuilder.AddMenuEntry(FSplineComponentVisualizerCommands::Get().AlignKeyPerpendicularToActor);
 	MenuBuilder.AddSeparator();
-	MenuBuilder.AddMenuEntry(FSplineComponentVisualizerCommands::Get().SnapAllToSelectedX);
-	MenuBuilder.AddMenuEntry(FSplineComponentVisualizerCommands::Get().SnapAllToSelectedY);
-	MenuBuilder.AddMenuEntry(FSplineComponentVisualizerCommands::Get().SnapAllToSelectedZ);
-	MenuBuilder.AddSeparator();
-	MenuBuilder.AddMenuEntry(FSplineComponentVisualizerCommands::Get().SnapToLastSelectedX);
-	MenuBuilder.AddMenuEntry(FSplineComponentVisualizerCommands::Get().SnapToLastSelectedY);
-	MenuBuilder.AddMenuEntry(FSplineComponentVisualizerCommands::Get().SnapToLastSelectedZ);
+
+	if (AxisDisplayInfo::GetAxisDisplayCoordinateSystem() == EAxisList::LeftUpForward)
+	{
+		MenuBuilder.AddMenuEntry(FSplineComponentVisualizerCommands::Get().SnapAllToSelectedY);
+		MenuBuilder.AddMenuEntry(FSplineComponentVisualizerCommands::Get().SnapAllToSelectedZ);
+		MenuBuilder.AddMenuEntry(FSplineComponentVisualizerCommands::Get().SnapAllToSelectedX);
+		MenuBuilder.AddSeparator();
+		MenuBuilder.AddMenuEntry(FSplineComponentVisualizerCommands::Get().SnapToLastSelectedY);
+		MenuBuilder.AddMenuEntry(FSplineComponentVisualizerCommands::Get().SnapToLastSelectedZ);
+		MenuBuilder.AddMenuEntry(FSplineComponentVisualizerCommands::Get().SnapToLastSelectedX);
+	}
+	else
+	{
+		MenuBuilder.AddMenuEntry(FSplineComponentVisualizerCommands::Get().SnapAllToSelectedX);
+		MenuBuilder.AddMenuEntry(FSplineComponentVisualizerCommands::Get().SnapAllToSelectedY);
+		MenuBuilder.AddMenuEntry(FSplineComponentVisualizerCommands::Get().SnapAllToSelectedZ);
+		MenuBuilder.AddSeparator();
+		MenuBuilder.AddMenuEntry(FSplineComponentVisualizerCommands::Get().SnapToLastSelectedX);
+		MenuBuilder.AddMenuEntry(FSplineComponentVisualizerCommands::Get().SnapToLastSelectedY);
+		MenuBuilder.AddMenuEntry(FSplineComponentVisualizerCommands::Get().SnapToLastSelectedZ);
+	}
+
 	MenuBuilder.AddSeparator();
 	MenuBuilder.AddMenuEntry(FSplineComponentVisualizerCommands::Get().StraightenToNext);
 	MenuBuilder.AddMenuEntry(FSplineComponentVisualizerCommands::Get().StraightenToPrevious);
@@ -3876,9 +3906,18 @@ void FSplineComponentVisualizer::GenerateSnapAlignSubMenu(FMenuBuilder& MenuBuil
 void FSplineComponentVisualizer::GenerateLockAxisSubMenu(FMenuBuilder& MenuBuilder) const
 {
 	MenuBuilder.AddMenuEntry(FSplineComponentVisualizerCommands::Get().SetLockedAxisNone);
-	MenuBuilder.AddMenuEntry(FSplineComponentVisualizerCommands::Get().SetLockedAxisX);
-	MenuBuilder.AddMenuEntry(FSplineComponentVisualizerCommands::Get().SetLockedAxisY);
-	MenuBuilder.AddMenuEntry(FSplineComponentVisualizerCommands::Get().SetLockedAxisZ);
+	if (AxisDisplayInfo::GetAxisDisplayCoordinateSystem() == EAxisList::LeftUpForward)
+	{
+		MenuBuilder.AddMenuEntry(FSplineComponentVisualizerCommands::Get().SetLockedAxisY);
+		MenuBuilder.AddMenuEntry(FSplineComponentVisualizerCommands::Get().SetLockedAxisZ);
+		MenuBuilder.AddMenuEntry(FSplineComponentVisualizerCommands::Get().SetLockedAxisX);
+	}
+	else
+	{
+		MenuBuilder.AddMenuEntry(FSplineComponentVisualizerCommands::Get().SetLockedAxisX);
+		MenuBuilder.AddMenuEntry(FSplineComponentVisualizerCommands::Get().SetLockedAxisY);
+		MenuBuilder.AddMenuEntry(FSplineComponentVisualizerCommands::Get().SetLockedAxisZ);
+	}
 }
 
 void FSplineComponentVisualizer::GenerateAttributeMenuSection(FMenuBuilder& InMenuBuilder) const
diff --git a/Engine/Source/Editor/DetailCustomizations/Private/SplineComponentDetails.cpp b/Engine/Source/Editor/DetailCustomizations/Private/SplineComponentDetails.cpp
--- a/Engine/Source/Editor/DetailCustomizations/Private/SplineComponentDetails.cpp
+++ b/Engine/Source/Editor/DetailCustomizations/Private/SplineComponentDetails.cpp
@@ -82,6 +82,7 @@
 #include "UObject/WeakObjectPtr.h"
 #include "UObject/WeakObjectPtrTemplates.h"
 #include "UnrealEdGlobals.h"
+#include "Math/UnitConversion.h"
 #include "Widgets/DeclarativeSyntaxSupport.h"
 #include "Widgets/Input/SButton.h"
 #include "Widgets/Input/SComboBox.h"
@@ -310,6 +311,9 @@ private:
 	void OnBeginScaleSlider();
 	void OnEndSlider(FVector::FReal);
 
+	FVector ConvertToLUFEuler(const FQuat& RotationQuaternion) const;
+	FQuat ConvertFromLUFEuler(const FVector& RotationEuler) const;
+
 	USplineComponent* SplineComp;
 	USplineComponent* SplineCompArchetype;
 	TSet<int32> SelectedKeys;
@@ -345,7 +349,7 @@ FSplinePointDetails::FSplinePointDetails(USplineComponent* InOwningSplineCompone
 	{
 		SplineProperties.Add(FindFProperty<FProperty>(USplineComponent::StaticClass(), Property));
 	}
-	
+
 	const TArray<ESplinePointType::Type> EnabledSplinePointTypes = InOwningSplineComponent->GetEnabledSplinePointTypes();
 
 	UEnum* SplinePointTypeEnum = StaticEnum<ESplinePointType::Type>();
@@ -557,7 +561,6 @@ void FSplinePointDetails::GenerateChildContent(IDetailChildrenBuilder& ChildrenB
 	{
 		return;
 	}
-
 	// Input key
 	ChildrenBuilder.AddCustomRow(LOCTEXT("InputKey", "Input Key"))
 		.RowTag(TEXT("InputKey"))
@@ -610,6 +613,10 @@ void FSplinePointDetails::GenerateChildContent(IDetailChildrenBuilder& ChildrenB
 				.X(this, &FSplinePointDetails::GetPositionX)
 				.Y(this, &FSplinePointDetails::GetPositionY)
 				.Z(this, &FSplinePointDetails::GetPositionZ)
+				.XDisplayName(AxisDisplayInfo::GetAxisToolTip(EAxisList::Forward))
+				.YDisplayName(AxisDisplayInfo::GetAxisToolTip(EAxisList::Left))
+				.ZDisplayName(AxisDisplayInfo::GetAxisToolTip(EAxisList::Up))
+				.Swizzle(AxisDisplayInfo::GetTransformAxisSwizzle())
 				.AllowSpin(true)
 				.bColorAxisLabels(true)
 				.SpinDelta(1.f)
@@ -628,8 +635,15 @@ void FSplinePointDetails::GenerateChildContent(IDetailChildrenBuilder& ChildrenB
 	// Rotation
 	if (SplineComp->AllowsSplinePointRotationEditing())
 	{
-		PasteFromTextDelegate->AddSP(this, &FSplinePointDetails::OnPasteFromText, ESplinePointProperty::Rotation);	
+		TSharedPtr<INumericTypeInterface<FRotator::FReal>> TypeInterface = MakeShareable(FUnitConversion::Settings().ShouldDisplayUnits() ? new TNumericUnitTypeInterface<FRotator::FReal>(EUnit::Degrees) : new TDefaultNumericTypeInterface<FRotator::FReal>());
+		if (AxisDisplayInfo::GetAxisDisplayCoordinateSystem() == EAxisList::LeftUpForward)
+		{
+			TypeInterface->SetMaxFractionalDigits(3);
+			TypeInterface->SetIndicateNearlyInteger(false);
+		}
 		
+		PasteFromTextDelegate->AddSP(this, &FSplinePointDetails::OnPasteFromText, ESplinePointProperty::Rotation);
+
 		ChildrenBuilder.AddCustomRow(LOCTEXT("Rotation", "Rotation"))
 			.RowTag(TEXT("Rotation"))
 			.CopyAction(CreateCopyAction(ESplinePointProperty::Rotation))
@@ -641,20 +655,25 @@ void FSplinePointDetails::GenerateChildContent(IDetailChildrenBuilder& ChildrenB
 			[
 				BuildSplinePointPropertyLabel(ESplinePointProperty::Rotation)
 			]
-		.ValueContent()
+			.ValueContent()
 			.MinDesiredWidth(375.0f)
 			.MaxDesiredWidth(375.0f)
 			[
 				SNew(SNumericRotatorInputBox<FRotator::FReal>)
+				.AllowSpin(false)
 				.Roll(this, &FSplinePointDetails::GetRotationRoll)
 				.Pitch(this, &FSplinePointDetails::GetRotationPitch)
 				.Yaw(this, &FSplinePointDetails::GetRotationYaw)
-				.AllowSpin(false)
-				.bColorAxisLabels(false)
+				.RollDisplayName(AxisDisplayInfo::GetRotationAxisToolTip(EAxisList::Forward))
+				.PitchDisplayName(AxisDisplayInfo::GetRotationAxisToolTip(EAxisList::Left))
+				.YawDisplayName(AxisDisplayInfo::GetRotationAxisToolTip(EAxisList::Up))
+				.bColorAxisLabels(true)
+				.Swizzle(AxisDisplayInfo::GetTransformAxisSwizzle())
 				.OnRollCommitted(this, &FSplinePointDetails::OnSetRotation, EAxis::X)
 				.OnPitchCommitted(this, &FSplinePointDetails::OnSetRotation, EAxis::Y)
 				.OnYawCommitted(this, &FSplinePointDetails::OnSetRotation, EAxis::Z)
 				.Font(IDetailLayoutBuilder::GetDetailFont())
+				.TypeInterface(TypeInterface)
 			];
 	}
 
@@ -662,7 +681,7 @@ void FSplinePointDetails::GenerateChildContent(IDetailChildrenBuilder& ChildrenB
 	if (SplineComp->AllowsSplinePointScaleEditing())
 	{
 		PasteFromTextDelegate->AddSP(this, &FSplinePointDetails::OnPasteFromText, ESplinePointProperty::Scale);
-		
+
 		ChildrenBuilder.AddCustomRow(LOCTEXT("Scale", "Scale"))
 			.RowTag(TEXT("Scale"))
 			.Visibility(TAttribute<EVisibility>(this, &FSplinePointDetails::IsEnabled))
@@ -684,6 +703,10 @@ void FSplinePointDetails::GenerateChildContent(IDetailChildrenBuilder& ChildrenB
 				.X(this, &FSplinePointDetails::GetScaleX)
 				.Y(this, &FSplinePointDetails::GetScaleY)
 				.Z(this, &FSplinePointDetails::GetScaleZ)
+				.XDisplayName(AxisDisplayInfo::GetAxisToolTip(EAxisList::Forward))
+				.YDisplayName(AxisDisplayInfo::GetAxisToolTip(EAxisList::Left))
+				.ZDisplayName(AxisDisplayInfo::GetAxisToolTip(EAxisList::Up))
+				.Swizzle(AxisDisplayInfo::GetTransformAxisSwizzle())
 				.AllowSpin(true)
 				.bColorAxisLabels(true)
 				.OnXChanged(this, &FSplinePointDetails::OnSetScale, ETextCommit::Default, EAxis::X)
@@ -698,79 +721,98 @@ void FSplinePointDetails::GenerateChildContent(IDetailChildrenBuilder& ChildrenB
 			];
 	}
 
-	// ArriveTangent
-	if (SplineComp->AllowsSplinePointArriveTangentEditing())
 	{
-		PasteFromTextDelegate->AddSP(this, &FSplinePointDetails::OnPasteFromText, ESplinePointProperty::ArriveTangent);
-	
-		ChildrenBuilder.AddCustomRow(LOCTEXT("ArriveTangent", "Arrive Tangent"))
-			.RowTag(TEXT("ArriveTangent"))
-			.Visibility(TAttribute<EVisibility>(this, &FSplinePointDetails::IsEnabled))
-			.CopyAction(CreateCopyAction(ESplinePointProperty::ArriveTangent))
-			.PasteAction(CreatePasteAction(ESplinePointProperty::ArriveTangent))
-			.NameContent()
-			.HAlign(HAlign_Left)
-			.VAlign(VAlign_Center)
-			[
-				SNew(STextBlock)
-				.Text(LOCTEXT("ArriveTangent", "Arrive Tangent"))
-				.Font(IDetailLayoutBuilder::GetDetailFont())
-				.ToolTipText(LOCTEXT("ArriveTangent_Tooltip", "Incoming tangent. Note that the size shown in viewport "
-					"is controlled by Spline Tangent Scale in editor preferences (and hidden if 0). Only allowed to "
-					"differ from Leave Tangent if Allow Discontinuous Spline is true."))
-			]
-		.ValueContent()
-			.MinDesiredWidth(375.0f)
-			.MaxDesiredWidth(375.0f)
-			[
-				SNew(SNumericVectorInputBox<FVector::FReal>)
-				.X(this, &FSplinePointDetails::GetArriveTangentX)
-				.Y(this, &FSplinePointDetails::GetArriveTangentY)
-				.Z(this, &FSplinePointDetails::GetArriveTangentZ)
-				.AllowSpin(false)
-				.bColorAxisLabels(false)
-				.OnXCommitted(this, &FSplinePointDetails::OnSetArriveTangent, EAxis::X)
-				.OnYCommitted(this, &FSplinePointDetails::OnSetArriveTangent, EAxis::Y)
-				.OnZCommitted(this, &FSplinePointDetails::OnSetArriveTangent, EAxis::Z)
-				.Font(IDetailLayoutBuilder::GetDetailFont())
-			];
-	}
+		TSharedPtr<INumericTypeInterface<FRotator::FReal>> TypeInterface = MakeShareable(new TDefaultNumericTypeInterface<FRotator::FReal>());
+		if (AxisDisplayInfo::GetAxisDisplayCoordinateSystem() == EAxisList::LeftUpForward)
+		{
+			TypeInterface->SetMaxFractionalDigits(3);
+			TypeInterface->SetIndicateNearlyInteger(false);
+		}
 
-	// LeaveTangent
-	if (SplineComp->AllowsSplinePointLeaveTangentEditing())
-	{
-		PasteFromTextDelegate->AddSP(this, &FSplinePointDetails::OnPasteFromText, ESplinePointProperty::LeaveTangent);
-	
-		ChildrenBuilder.AddCustomRow(LOCTEXT("LeaveTangent", "Leave Tangent"))
-			.RowTag(TEXT("LeaveTangent"))
-			.Visibility(TAttribute<EVisibility>(this, &FSplinePointDetails::IsEnabled))
-			.CopyAction(CreateCopyAction(ESplinePointProperty::LeaveTangent))
-			.PasteAction(CreatePasteAction(ESplinePointProperty::LeaveTangent))
-			.NameContent()
-			.HAlign(HAlign_Left)
-			.VAlign(VAlign_Center)
-			[
-				SNew(STextBlock)
-				.Text(LOCTEXT("LeaveTangent", "Leave Tangent"))
-				.Font(IDetailLayoutBuilder::GetDetailFont())
-				.ToolTipText(LOCTEXT("LeaveTangent_Tooltip", "Outgoing tangent. Note that the size shown in viewport "
-					"is controlled by Spline Tangent Scale in editor preferences (and hidden if 0)."))
-			]
-		.ValueContent()
-			.MinDesiredWidth(375.0f)
-			.MaxDesiredWidth(375.0f)
-			[
-				SNew(SNumericVectorInputBox<FVector::FReal>)
-				.X(this, &FSplinePointDetails::GetLeaveTangentX)
-				.Y(this, &FSplinePointDetails::GetLeaveTangentY)
-				.Z(this, &FSplinePointDetails::GetLeaveTangentZ)
-				.AllowSpin(false)
-				.bColorAxisLabels(false)
-				.OnXCommitted(this, &FSplinePointDetails::OnSetLeaveTangent, EAxis::X)
-				.OnYCommitted(this, &FSplinePointDetails::OnSetLeaveTangent, EAxis::Y)
-				.OnZCommitted(this, &FSplinePointDetails::OnSetLeaveTangent, EAxis::Z)
-				.Font(IDetailLayoutBuilder::GetDetailFont())
-			];
+		// ArriveTangent
+		if (SplineComp->AllowsSplinePointArriveTangentEditing())
+		{
+			PasteFromTextDelegate->AddSP(this, &FSplinePointDetails::OnPasteFromText, ESplinePointProperty::ArriveTangent);
+
+			ChildrenBuilder.AddCustomRow(LOCTEXT("ArriveTangent", "Arrive Tangent"))
+				.RowTag(TEXT("ArriveTangent"))
+				.Visibility(TAttribute<EVisibility>(this, &FSplinePointDetails::IsEnabled))
+				.CopyAction(CreateCopyAction(ESplinePointProperty::ArriveTangent))
+				.PasteAction(CreatePasteAction(ESplinePointProperty::ArriveTangent))
+				.NameContent()
+				.HAlign(HAlign_Left)
+				.VAlign(VAlign_Center)
+				[
+					SNew(STextBlock)
+					.Text(LOCTEXT("ArriveTangent", "Arrive Tangent"))
+					.Font(IDetailLayoutBuilder::GetDetailFont())
+					.ToolTipText(LOCTEXT("ArriveTangent_Tooltip", "Incoming tangent. Note that the size shown in viewport "
+						"is controlled by Spline Tangent Scale in editor preferences (and hidden if 0). Only allowed to "
+						"differ from Leave Tangent if Allow Discontinuous Spline is true."))
+				]
+			.ValueContent()
+				.MinDesiredWidth(375.0f)
+				.MaxDesiredWidth(375.0f)
+				[
+					SNew(SNumericVectorInputBox<FVector::FReal>)
+					.X(this, &FSplinePointDetails::GetArriveTangentX)
+					.Y(this, &FSplinePointDetails::GetArriveTangentY)
+					.Z(this, &FSplinePointDetails::GetArriveTangentZ)
+					.XDisplayName(AxisDisplayInfo::GetAxisToolTip(EAxisList::Forward))
+					.YDisplayName(AxisDisplayInfo::GetAxisToolTip(EAxisList::Left))
+					.ZDisplayName(AxisDisplayInfo::GetAxisToolTip(EAxisList::Up))
+					.Swizzle(AxisDisplayInfo::GetTransformAxisSwizzle())
+					.AllowSpin(false)
+					.bColorAxisLabels(false)
+					.OnXCommitted(this, &FSplinePointDetails::OnSetArriveTangent, EAxis::X)
+					.OnYCommitted(this, &FSplinePointDetails::OnSetArriveTangent, EAxis::Y)
+					.OnZCommitted(this, &FSplinePointDetails::OnSetArriveTangent, EAxis::Z)
+					.TypeInterface(TypeInterface)
+					.Font(IDetailLayoutBuilder::GetDetailFont())
+				];
+		}
+
+		// LeaveTangent
+		if (SplineComp->AllowsSplinePointLeaveTangentEditing())
+		{
+			PasteFromTextDelegate->AddSP(this, &FSplinePointDetails::OnPasteFromText, ESplinePointProperty::LeaveTangent);
+
+			ChildrenBuilder.AddCustomRow(LOCTEXT("LeaveTangent", "Leave Tangent"))
+				.RowTag(TEXT("LeaveTangent"))
+				.Visibility(TAttribute<EVisibility>(this, &FSplinePointDetails::IsEnabled))
+				.CopyAction(CreateCopyAction(ESplinePointProperty::LeaveTangent))
+				.PasteAction(CreatePasteAction(ESplinePointProperty::LeaveTangent))
+				.NameContent()
+				.HAlign(HAlign_Left)
+				.VAlign(VAlign_Center)
+				[
+					SNew(STextBlock)
+					.Text(LOCTEXT("LeaveTangent", "Leave Tangent"))
+					.Font(IDetailLayoutBuilder::GetDetailFont())
+					.ToolTipText(LOCTEXT("LeaveTangent_Tooltip", "Outgoing tangent. Note that the size shown in viewport "
+						"is controlled by Spline Tangent Scale in editor preferences (and hidden if 0)."))
+				]
+			.ValueContent()
+				.MinDesiredWidth(375.0f)
+				.MaxDesiredWidth(375.0f)
+				[
+					SNew(SNumericVectorInputBox<FVector::FReal>)
+					.X(this, &FSplinePointDetails::GetLeaveTangentX)
+					.Y(this, &FSplinePointDetails::GetLeaveTangentY)
+					.Z(this, &FSplinePointDetails::GetLeaveTangentZ)
+					.XDisplayName(AxisDisplayInfo::GetAxisToolTip(EAxisList::Forward))
+					.YDisplayName(AxisDisplayInfo::GetAxisToolTip(EAxisList::Left))
+					.ZDisplayName(AxisDisplayInfo::GetAxisToolTip(EAxisList::Up))
+					.Swizzle(AxisDisplayInfo::GetTransformAxisSwizzle())
+					.AllowSpin(false)
+					.bColorAxisLabels(false)
+					.OnXCommitted(this, &FSplinePointDetails::OnSetLeaveTangent, EAxis::X)
+					.OnYCommitted(this, &FSplinePointDetails::OnSetLeaveTangent, EAxis::Y)
+					.OnZCommitted(this, &FSplinePointDetails::OnSetLeaveTangent, EAxis::Z)
+					.TypeInterface(TypeInterface)
+					.Font(IDetailLayoutBuilder::GetDetailFont())
+				];
+		}
 	}
 
 	// Type
@@ -899,17 +941,39 @@ void FSplinePointDetails::UpdateValues()
 			for (int32 Index : SelectedKeys)
 			{
 				// possibly could get this data in bulk via GetSplinePoint(Index), but doing a 1:1 swap for now.
-				
+
 				InputKey.Add(SplineComp->GetInputKeyValueAtSplinePoint(Index));
-				
-				Position.Add(SplineComp->GetLocationAtSplinePoint(Index, bEditingLocationAbsolute ? ESplineCoordinateSpace::World : ESplineCoordinateSpace::Local));
-				Rotation.Add(SplineComp->GetRotationAtSplinePoint(Index, bEditingRotationAbsolute ? ESplineCoordinateSpace::World : ESplineCoordinateSpace::Local));
+
+				FVector SelectedSplinePointLocation = SplineComp->GetLocationAtSplinePoint(Index, bEditingLocationAbsolute ? ESplineCoordinateSpace::World : ESplineCoordinateSpace::Local);
+
+				FQuat SelectedSplinePointQuaternion = SplineComp->GetQuaternionAtSplinePoint(Index, bEditingRotationAbsolute ? ESplineCoordinateSpace::World : ESplineCoordinateSpace::Local);
+
 				Scale.Add(SplineComp->GetScaleAtSplinePoint(Index));
-				
-				ArriveTangent.Add(SplineComp->GetArriveTangentAtSplinePoint(Index, ESplineCoordinateSpace::Local));
-				LeaveTangent.Add(SplineComp->GetLeaveTangentAtSplinePoint(Index, ESplineCoordinateSpace::Local));
-				
+
+				FVector SelectedPointArriveTangent = SplineComp->GetArriveTangentAtSplinePoint(Index, ESplineCoordinateSpace::Local);
+				FVector SelectedPointLeaveTangent = SplineComp->GetLeaveTangentAtSplinePoint(Index, ESplineCoordinateSpace::Local);
+
 				PointType.Add(SplineComp->GetSplinePointType(Index));
+
+				if (AxisDisplayInfo::GetAxisDisplayCoordinateSystem() == EAxisList::LeftUpForward)
+				{
+					SelectedSplinePointLocation.Y = -SelectedSplinePointLocation.Y;
+
+					SelectedPointArriveTangent.Y = -SelectedPointArriveTangent.Y;
+
+					SelectedPointLeaveTangent.Y = -SelectedPointLeaveTangent.Y;
+
+					FVector SelectedSplinePointEuler = ConvertToLUFEuler(SelectedSplinePointQuaternion);
+					Rotation.Add(FRotator(SelectedSplinePointEuler.X, SelectedSplinePointEuler.Y, SelectedSplinePointEuler.Z));
+				}
+				else
+				{
+					Rotation.Add(SelectedSplinePointQuaternion.Rotator());
+				}
+
+				Position.Add(SelectedSplinePointLocation);
+				ArriveTangent.Add(SelectedPointArriveTangent);
+				LeaveTangent.Add(SelectedPointLeaveTangent);
 			}
 
 			if (SplineMetaDataDetails)
@@ -1060,7 +1124,7 @@ void FSplinePointDetails::OnSetPosition(FVector::FReal NewValue, ETextCommit::Ty
 
 			const ESplineCoordinateSpace::Type SplineCoordinateSpace = bEditingLocationAbsolute ? ESplineCoordinateSpace::World : ESplineCoordinateSpace::Local;
 			FVector PointPosition = SplineComp->GetLocationAtSplinePoint(Index, SplineCoordinateSpace);
-			PointPosition.SetComponentForAxis(Axis, NewValue);
+			PointPosition.SetComponentForAxis(Axis, AxisDisplayInfo::GetAxisDisplayCoordinateSystem() == EAxisList::LeftUpForward && Axis == EAxis::Type::Y ? -NewValue : NewValue);
 			SplineComp->SetLocationAtSplinePoint(Index, PointPosition, SplineCoordinateSpace, false);
 		}
 	}
@@ -1095,7 +1159,7 @@ void FSplinePointDetails::OnSetArriveTangent(FVector::FReal NewValue, ETextCommi
 		SplineComp->Modify();
 
 		const int32 NumPoints = SplineComp->GetNumberOfSplinePoints();
-		
+
 		for (int32 Index : SelectedKeys)
 		{
 			if (Index < 0 || Index >= NumPoints)
@@ -1106,9 +1170,9 @@ void FSplinePointDetails::OnSetArriveTangent(FVector::FReal NewValue, ETextCommi
 			}
 
 			FVector LocalArriveTangent = SplineComp->GetArriveTangentAtSplinePoint(Index, ESplineCoordinateSpace::Local);
-			LocalArriveTangent.SetComponentForAxis(Axis, NewValue);
+			LocalArriveTangent.SetComponentForAxis(Axis, AxisDisplayInfo::GetAxisDisplayCoordinateSystem() == EAxisList::LeftUpForward && Axis == EAxis::Type::Y ? -NewValue : NewValue);
 			FVector LocalLeaveTangent = SplineComp->bAllowDiscontinuousSpline ? SplineComp->GetLeaveTangentAtSplinePoint(Index, ESplineCoordinateSpace::Local) : LocalArriveTangent;
-			
+
 			SplineComp->SetTangentsAtSplinePoint(Index, LocalArriveTangent, LocalLeaveTangent, ESplineCoordinateSpace::Local, false);
 
 		}
@@ -1144,7 +1208,7 @@ void FSplinePointDetails::OnSetLeaveTangent(FVector::FReal NewValue, ETextCommit
 		SplineComp->Modify();
 
 		const int32 NumPoints = SplineComp->GetNumberOfSplinePoints();
-		
+
 		for (int32 Index : SelectedKeys)
 		{
 			if (Index < 0 || Index >= NumPoints)
@@ -1155,9 +1219,9 @@ void FSplinePointDetails::OnSetLeaveTangent(FVector::FReal NewValue, ETextCommit
 			}
 
 			FVector LocalLeaveTangent = SplineComp->GetLeaveTangentAtSplinePoint(Index, ESplineCoordinateSpace::Local);
-			LocalLeaveTangent.SetComponentForAxis(Axis, NewValue);
+			LocalLeaveTangent.SetComponentForAxis(Axis, AxisDisplayInfo::GetAxisDisplayCoordinateSystem() == EAxisList::LeftUpForward && Axis == EAxis::Type::Y ? -NewValue : NewValue);
 			FVector LocalArriveTangent = SplineComp->bAllowDiscontinuousSpline ? SplineComp->GetArriveTangentAtSplinePoint(Index, ESplineCoordinateSpace::Local) : LocalLeaveTangent;
-			
+
 			SplineComp->SetTangentsAtSplinePoint(Index, LocalArriveTangent, LocalLeaveTangent, ESplineCoordinateSpace::Local, false);
 		}
 	}
@@ -1185,15 +1249,15 @@ void FSplinePointDetails::OnSetRotation(FRotator::FReal NewValue, ETextCommit::T
 	{
 		return;
 	}
-	
+
 	FQuat NewRotationRelative;
 	// Scope the transaction to only include the value change and none of the derived data changes that might arise from NotifyPropertyModified
 	{
 		const FScopedTransaction Transaction(LOCTEXT("SetSplinePointRotation", "Set spline point rotation"));
 		SplineComp->Modify();
-		
+
 		const int32 NumPoints = SplineComp->GetNumberOfSplinePoints();
-		
+
 		FQuat SplineComponentRotation = SplineComp->GetComponentQuat();
 		for (int32 Index : SelectedKeys)
 		{
@@ -1208,29 +1272,66 @@ void FSplinePointDetails::OnSetRotation(FRotator::FReal NewValue, ETextCommit::T
 
 			if (bEditingRotationAbsolute)
 			{
-				FRotator AbsoluteRot = (SplineComponentRotation * CurrentRotationRelative).Rotator();
+				FQuat CurrentAbsoluteRot = SplineComponentRotation * CurrentRotationRelative;
 
-				switch (Axis)
+				if (AxisDisplayInfo::GetAxisDisplayCoordinateSystem() == EAxisList::LeftUpForward)
 				{
-				case EAxis::X: AbsoluteRot.Roll = NewValue; break;
-				case EAxis::Y: AbsoluteRot.Pitch = NewValue; break;
-				case EAxis::Z: AbsoluteRot.Yaw = NewValue; break;
+					FVector CurrentComponentValue = ConvertToLUFEuler(CurrentAbsoluteRot);
+
+					switch (Axis)
+					{
+					case EAxis::X: CurrentComponentValue.Z = NewValue; break;
+					case EAxis::Y: CurrentComponentValue.X = NewValue; break;
+					case EAxis::Z: CurrentComponentValue.Y = NewValue; break;
+					}
+
+					CurrentAbsoluteRot = ConvertFromLUFEuler(CurrentComponentValue);
 				}
+				else
+				{
+					FRotator CurrentComponentValue = (SplineComponentRotation * CurrentRotationRelative).Rotator();
 
-				NewRotationRelative = SplineComponentRotation.Inverse() * AbsoluteRot.Quaternion();
+					switch (Axis)
+					{
+					case EAxis::X: CurrentComponentValue.Roll = NewValue; break;
+					case EAxis::Y: CurrentComponentValue.Pitch = NewValue; break;
+					case EAxis::Z: CurrentComponentValue.Yaw = NewValue; break;
+					}
+
+					CurrentAbsoluteRot = CurrentComponentValue.Quaternion();
+				}
+
+				NewRotationRelative = SplineComponentRotation.Inverse() * CurrentAbsoluteRot;
 			}
 			else
 			{
-				FRotator NewRotationRotator(CurrentRotationRelative);
 
-				switch (Axis)
+				if (AxisDisplayInfo::GetAxisDisplayCoordinateSystem() == EAxisList::LeftUpForward)
 				{
-				case EAxis::X: NewRotationRotator.Roll = NewValue; break;
-				case EAxis::Y: NewRotationRotator.Pitch = NewValue; break;
-				case EAxis::Z: NewRotationRotator.Yaw = NewValue; break;
+					FVector CurrentComponentValue = ConvertToLUFEuler(CurrentRotationRelative);
+
+					switch (Axis)
+					{
+					case EAxis::X: CurrentComponentValue.Z = NewValue; break;
+					case EAxis::Y: CurrentComponentValue.X = NewValue; break;
+					case EAxis::Z: CurrentComponentValue.Y = NewValue; break;
+					}
+
+					NewRotationRelative = ConvertFromLUFEuler(CurrentComponentValue);
 				}
+				else
+				{
+					FRotator NewRotationRotator(CurrentRotationRelative);
+
+					switch (Axis)
+					{
+					case EAxis::X: NewRotationRotator.Roll = NewValue; break;
+					case EAxis::Y: NewRotationRotator.Pitch = NewValue; break;
+					case EAxis::Z: NewRotationRotator.Yaw = NewValue; break;
+					}
 
-				NewRotationRelative = NewRotationRotator.Quaternion();
+					NewRotationRelative = NewRotationRotator.Quaternion();
+				}
 			}
 
 			SplineComp->SetQuaternionAtSplinePoint(Index, NewRotationRelative, ESplineCoordinateSpace::Local);
@@ -1858,6 +1959,46 @@ void FSplinePointDetails::OnEndSlider(FVector::FReal)
 	GEditor->EndTransaction();
 }
 
+FVector FSplinePointDetails::ConvertToLUFEuler(const FQuat& RotationQuaternion) const
+{
+	const TTuple<FQuat::FReal, FQuat::FReal, FQuat::FReal> VerseEulerRads = RotationQuaternion.GetNormalized().ToLUFEuler();
+
+	// Since the value is converted from quaternion, will likely have denormals.  Clamp those values.
+	auto SanitizeFloat = [](FQuat::FReal Val)->FQuat::FReal
+		{
+			if (FMath::IsNearlyZero(Val))
+			{
+				return 0.0;
+			}
+			return Val;
+		};
+
+	const FVector VerseEulerRadsV
+	{
+		SanitizeFloat(VerseEulerRads.Get<0>()),
+		SanitizeFloat(VerseEulerRads.Get<1>()),
+		SanitizeFloat(VerseEulerRads.Get<2>())
+	};
+	return FMath::RadiansToDegrees(VerseEulerRadsV);
+}
+
+FQuat FSplinePointDetails::ConvertFromLUFEuler(const FVector& RotationEuler) const
+{
+	const FVector RotationRads = FMath::DegreesToRadians(RotationEuler);
+
+	const TTuple<FQuat::FReal, FQuat::FReal, FQuat::FReal> RotationRadsT =
+	{
+		RotationRads.X,
+		RotationRads.Y,
+		RotationRads.Z
+	};
+
+
+	FQuat Quat = FQuat::MakeFromLUFEuler(RotationRadsT);
+	Quat.Normalize();
+	return Quat;
+}
+
 ////////////////////////////////////
 
 TSharedRef<IDetailCustomization> FSplineComponentDetails::MakeInstance()
diff --git a/Engine/Source/Runtime/Engine/Classes/Components/SplineComponent.h b/Engine/Source/Runtime/Engine/Classes/Components/SplineComponent.h
--- a/Engine/Source/Runtime/Engine/Classes/Components/SplineComponent.h
+++ b/Engine/Source/Runtime/Engine/Classes/Components/SplineComponent.h
@@ -303,7 +303,7 @@ private:
 	
 public:
 	/** Default up vector in local space to be used when calculating transforms along the spline */
-	UPROPERTY(EditAnywhere, Replicated, BlueprintReadWrite, Category = Spline)
+	UPROPERTY(EditAnywhere, Replicated, BlueprintReadWrite, Category = Spline, meta=(Spatial))
 	FVector DefaultUpVector;
 
 #if WITH_EDITORONLY_DATA
diff --git a/Engine/Source/Runtime/Slate/Public/Widgets/Input/SVectorInputBox.h b/Engine/Source/Runtime/Slate/Public/Widgets/Input/SVectorInputBox.h
--- a/Engine/Source/Runtime/Slate/Public/Widgets/Input/SVectorInputBox.h
+++ b/Engine/Source/Runtime/Slate/Public/Widgets/Input/SVectorInputBox.h
@@ -520,6 +520,8 @@ private:
 			.UndeterminedString(NSLOCTEXT("SVectorInputBox", "MultipleValues", "Multiple Values"))
 			.ContextMenuExtender(OnContextMenuExtenderComponent)
 			.TypeInterface(InArgs._TypeInterface)
+			.MinFractionalDigits(InArgs._TypeInterface ? InArgs._TypeInterface->GetMinFractionalDigits() : TOptional<int32>())
+			.MaxFractionalDigits(InArgs._TypeInterface ? InArgs._TypeInterface->GetMaxFractionalDigits() : TOptional<int32>())
 			.MinValue(CreatePerComponentGetter(ComponentIndex, TOptional<NumericType>(), InArgs._MinVector))
 			.MaxValue(CreatePerComponentGetter(ComponentIndex, TOptional<NumericType>(), InArgs._MaxVector))
 			.MinSliderValue(CreatePerComponentGetter(ComponentIndex, TOptional<NumericType>(), InArgs._MinSliderVector))
