[Curve Editor] Fix crash on editor exit when detail panel is opened with a media plate selected.

diff --git a/Engine/Source/Editor/DetailCustomizations/Private/CurveStructCustomization.cpp b/Engine/Source/Editor/DetailCustomizations/Private/CurveStructCustomization.cpp
--- a/Engine/Source/Editor/DetailCustomizations/Private/CurveStructCustomization.cpp
+++ b/Engine/Source/Editor/DetailCustomizations/Private/CurveStructCustomization.cpp
@@ -73,7 +73,11 @@ FCurveStructCustomization::~FCurveStructCustomization()
 	}
 
 	DestroyPopOutWindow();
-	GEditor->UnregisterForUndo(this);
+
+	if (GEditor)
+	{
+		GEditor->UnregisterForUndo(this);
+	}
 }
 
 FCurveStructCustomization::FCurveStructCustomization()
