SceneCapture: Expose lens distortion show flag on scene captures.

diff --git a/Engine/Source/Editor/DetailCustomizations/Private/SceneCaptureDetails.cpp b/Engine/Source/Editor/DetailCustomizations/Private/SceneCaptureDetails.cpp
--- a/Engine/Source/Editor/DetailCustomizations/Private/SceneCaptureDetails.cpp
+++ b/Engine/Source/Editor/DetailCustomizations/Private/SceneCaptureDetails.cpp
@@ -180,6 +180,7 @@ void FSceneCaptureDetails::CustomizeDetails( IDetailLayoutBuilder& DetailLayout
 	ShowFlagsToAllowForCaptures.Add(FEngineShowFlags::EShowFlag::SF_Tonemapper);
 	ShowFlagsToAllowForCaptures.Add(FEngineShowFlags::EShowFlag::SF_PostProcessMaterial);
 	ShowFlagsToAllowForCaptures.Add(FEngineShowFlags::EShowFlag::SF_AllowPrimitiveAlphaHoldout);
+	ShowFlagsToAllowForCaptures.Add(FEngineShowFlags::EShowFlag::SF_LensDistortion);
 
 	// Create array of flag name strings for each group
 	TArray< TArray<FString> > ShowFlagsByGroup;
