Fix spline component crash on linux

diff --git a/Engine/Source/Editor/DetailCustomizations/Private/SplineComponentDetails.cpp b/Engine/Source/Editor/DetailCustomizations/Private/SplineComponentDetails.cpp
--- a/Engine/Source/Editor/DetailCustomizations/Private/SplineComponentDetails.cpp
+++ b/Engine/Source/Editor/DetailCustomizations/Private/SplineComponentDetails.cpp
@@ -84,6 +84,8 @@
 #include "UnrealEdGlobals.h"
 #include "Math/UnitConversion.h"
 #include "Widgets/DeclarativeSyntaxSupport.h"
+#include "Widgets/Input/NumericTypeInterface.h"
+#include "Widgets/Input/NumericUnitTypeInterface.inl"
 #include "Widgets/Input/SButton.h"
 #include "Widgets/Input/SComboBox.h"
 #include "Widgets/Input/SComboButton.h"
