[EQS] fixed code that recovers from invalid data setup by resetting the field property

diff --git a/Engine/Source/Editor/DetailCustomizations/Private/AIDataProviderValueDetails.cpp b/Engine/Source/Editor/DetailCustomizations/Private/AIDataProviderValueDetails.cpp
--- a/Engine/Source/Editor/DetailCustomizations/Private/AIDataProviderValueDetails.cpp
+++ b/Engine/Source/Editor/DetailCustomizations/Private/AIDataProviderValueDetails.cpp
@@ -130,7 +130,7 @@ void FAIDataProviderValueDetails::OnBindingChanged()
 			ensureAlwaysMsgf(false, TEXT("We expect UEnvQueryNode::PostEditChangeProperty to prevent "
 				"copying an incompatible provider value on top of an existing one before calling NotifyAssetUpdate and refreshing the details."));
 			DataBindingProperty->ResetToDefault();
-			DefaultValueProperty->ResetToDefault();
+			DataFieldProperty->ResetToDefault();
 		}
 	}
 }
