Mac - Fix Mac detail customization broken
This is due to my introduction of a new detail customization instead of adding it to the existing

#tests local testing & preflight

diff --git a/Engine/Source/Developer/Mac/MacPlatformEditor/Private/MacPlatformEditorModule.cpp b/Engine/Source/Developer/Mac/MacPlatformEditor/Private/MacPlatformEditorModule.cpp
--- a/Engine/Source/Developer/Mac/MacPlatformEditor/Private/MacPlatformEditorModule.cpp
+++ b/Engine/Source/Developer/Mac/MacPlatformEditor/Private/MacPlatformEditorModule.cpp
@@ -10,7 +10,6 @@
 #include "UObject/Class.h"
 #include "PropertyEditorModule.h"
 #include "XcodeProjectSettingsDetailsCustomization.h"
-#include "MacTargetSettingsCustomization.h"
 
 #define LOCTEXT_NAMESPACE "MacPlatformEditorModule"
 
@@ -45,11 +44,6 @@ class FMacPlatformEditorModule
                 "XcodeProjectSettings",
                 FOnGetDetailCustomizationInstance::CreateStatic(&FXcodeProjectSettingsDetailsCustomization::MakeInstance)
             );
-			
-			PropertyModule.RegisterCustomClassLayout(
-				"MacTargetSettings",
-				FOnGetDetailCustomizationInstance::CreateStatic(&FMacTargetSettingsCustomization::MakeInstance)
-			);
 		}
 	}
 
diff --git a/Engine/Source/Developer/Mac/MacPlatformEditor/Private/MacTargetSettingsCustomization.cpp b/Engine/Source/Developer/Mac/MacPlatformEditor/Private/MacTargetSettingsCustomization.cpp
deleted file mode 100644
diff --git a/Engine/Source/Developer/Mac/MacPlatformEditor/Private/MacTargetSettingsCustomization.h b/Engine/Source/Developer/Mac/MacPlatformEditor/Private/MacTargetSettingsCustomization.h
deleted file mode 100644
diff --git a/Engine/Source/Editor/DetailCustomizations/DetailCustomizations.Build.cs b/Engine/Source/Editor/DetailCustomizations/DetailCustomizations.Build.cs
--- a/Engine/Source/Editor/DetailCustomizations/DetailCustomizations.Build.cs
+++ b/Engine/Source/Editor/DetailCustomizations/DetailCustomizations.Build.cs
@@ -72,6 +72,7 @@ public class DetailCustomizations : ModuleRules
 				"Json",
 				"JsonUtilities",
 				"SettingsEditor",
+				"MacTargetPlatformSettings",
 			}
 		);
 
diff --git a/Engine/Source/Editor/DetailCustomizations/Private/MacTargetSettingsDetails.cpp b/Engine/Source/Editor/DetailCustomizations/Private/MacTargetSettingsDetails.cpp
--- a/Engine/Source/Editor/DetailCustomizations/Private/MacTargetSettingsDetails.cpp
+++ b/Engine/Source/Editor/DetailCustomizations/Private/MacTargetSettingsDetails.cpp
@@ -1,9 +1,12 @@
 // Copyright Epic Games, Inc. All Rights Reserved.
 
 #include "MacTargetSettingsDetails.h"
+#include "MacTargetSettings.h"
 #include "Misc/Paths.h"
 #include "Misc/App.h"
 #include "Misc/EngineVersion.h"
+#include "Misc/ConfigCacheIni.h"
+#include "Misc/ConfigContext.h"
 #include "Modules/ModuleManager.h"
 #include "Layout/Margin.h"
 #include "Widgets/DeclarativeSyntaxSupport.h"
@@ -25,6 +28,8 @@
 #include "IDetailPropertyRow.h"
 #include "RHI.h"
 #include "ShaderFormatsPropertyDetails.h"
+#include "ISettingsEditorModule.h"
+#include "SSettingsEditorCheckoutNotice.h"
 
 namespace MacTargetSettingsDetailsConstants
 {
@@ -215,6 +220,21 @@ void FMacTargetSettingsDetails::CustomizeDetails( IDetailLayoutBuilder& DetailBu
 			]
 		];
 		
+		RenderCategory.AddCustomRow(LOCTEXT("UseBindlessWithSM5", "Use Bindless With SM5 (Experimental)"), false)
+		.RowTag("UseBindlessWithSM5")
+		.NameContent()
+		[
+			SNew(STextBlock)
+			.Font(IDetailLayoutBuilder::GetDetailFont())
+			.Text(LOCTEXT("UseBindlessWithSM5", "Use Bindless With SM5 (Experimental)"))
+		]
+		.ValueContent()
+		[
+			SNew(SCheckBox)
+			.IsChecked(this, &FMacTargetSettingsDetails::ShouldUseBindlessWithSM5)
+			.OnCheckStateChanged(this, &FMacTargetSettingsDetails::OnUseBindlessWithSM5Changed)
+		];
+		
 		UpdateShaderStandardWarning();
     }
 	
@@ -418,4 +438,45 @@ void FMacTargetSettingsDetails::UpdateShaderStandardWarning()
     }
 }
 
+constexpr TCHAR SM5SPSection[] = TEXT("ShaderPlatformConfig METAL_SM5");
+constexpr TCHAR BindlessConfigurationKey[] = TEXT("BindlessConfiguration");
+
+ECheckBoxState FMacTargetSettingsDetails::ShouldUseBindlessWithSM5() const
+{
+	FString OutSetting;
+	if (GConfig->GetString(SM5SPSection, BindlessConfigurationKey, OutSetting, GEngineIni))
+	{
+		return OutSetting.Equals(TEXT("Disabled"), ESearchCase::IgnoreCase) ? ECheckBoxState::Unchecked : ECheckBoxState::Checked;
+	}
+	else
+	{
+		return ECheckBoxState::Unchecked;
+	}
+}
+void FMacTargetSettingsDetails::OnUseBindlessWithSM5Changed(ECheckBoxState NewState)
+{
+	static const FString ClassConfigName = UMacTargetSettings::StaticClass()->ClassConfigName.ToString();
+	static const FString DefaultConfigLayerFilename = UMacTargetSettings::StaticClass()->GetDefaultConfigFilename();
+
+	if (!SettingsHelpers::CheckOutOrAddFile(DefaultConfigLayerFilename))
+	{
+		SettingsHelpers::MakeWritable(DefaultConfigLayerFilename);
+	}
+
+	// CF UObject::UpdateSingleSectionOfConfigFile implementation
+	GConfig->Flush(false, *ClassConfigName);
+	
+	// Because BindlessConfiguration is not a UPROPERTY with config keyboard we have to do it this way.
+	FConfigFile Temp;
+	Temp.SetString(SM5SPSection, BindlessConfigurationKey, NewState == ECheckBoxState::Unchecked ? TEXT("Disabled") : TEXT("All"));
+	Temp.UpdateSinglePropertyInSection(*DefaultConfigLayerFilename, BindlessConfigurationKey, SM5SPSection);
+	
+	// CF UObject::UpdateSingleSectionOfConfigFile implementation
+	FConfigContext Context = FConfigContext::ForceReloadIntoGConfig();
+	Context.bWriteDestIni = false;
+	Context.Load(*ClassConfigName);
+	
+	FModuleManager::GetModuleChecked<ISettingsEditorModule>("SettingsEditor").OnApplicationRestartRequired();
+}
+
 #undef LOCTEXT_NAMESPACE
diff --git a/Engine/Source/Editor/DetailCustomizations/Private/MacTargetSettingsDetails.h b/Engine/Source/Editor/DetailCustomizations/Private/MacTargetSettingsDetails.h
--- a/Engine/Source/Editor/DetailCustomizations/Private/MacTargetSettingsDetails.h
+++ b/Engine/Source/Editor/DetailCustomizations/Private/MacTargetSettingsDetails.h
@@ -49,6 +49,12 @@ private:
 	/** Delegate to update the shader standard warning */
 	void UpdateShaderStandardWarning();
 	
+	/** Delegate handler to know if we currently use bindless with SM5 */
+	ECheckBoxState ShouldUseBindlessWithSM5() const;
+	
+	/** Delegate handler to set if we use bindless with SM5 */
+	void OnUseBindlessWithSM5Changed(ECheckBoxState NewState);
+	
 private:
 	/** Reference to the target shader formats property view */
 	TSharedPtr<FShaderFormatsPropertyDetails> TargetShaderFormatsDetails;
