Ensure that color widget sends a final update with EPropertyChangeType::ValueSet at the end of an interactive update.
This matches the contract in the comments for EPropertyChangeType.
It fixes an issue where color updates weren't being set on children of the material instance being edited. For performance we only apply value changes to all children at the end of the interactive edit.
#rnx

diff --git a/Engine/Source/Editor/DetailCustomizations/Private/ColorStructCustomization.cpp b/Engine/Source/Editor/DetailCustomizations/Private/ColorStructCustomization.cpp
--- a/Engine/Source/Editor/DetailCustomizations/Private/ColorStructCustomization.cpp
+++ b/Engine/Source/Editor/DetailCustomizations/Private/ColorStructCustomization.cpp
@@ -388,7 +388,16 @@ void FColorStructCustomization::OnColorPickerInteractiveBegin()
 
 void FColorStructCustomization::OnColorPickerInteractiveEnd()
 {
-	bIsInteractive = false;
+	if (bIsInteractive)
+	{
+		bIsInteractive = false;
+
+		if (!LastPickerColorString.IsEmpty())
+		{
+			StructPropertyHandle->SetValueFromFormattedString(LastPickerColorString, EPropertyValueSetFlags::NotTransactable);
+			StructPropertyHandle->NotifyFinishedChangingProperties();
+		}
+	}
 }
 
 
