SplineComponentDetails: Fix bad paren in compound condition

diff --git a/Engine/Source/Editor/DetailCustomizations/Private/SplineComponentDetails.cpp b/Engine/Source/Editor/DetailCustomizations/Private/SplineComponentDetails.cpp
--- a/Engine/Source/Editor/DetailCustomizations/Private/SplineComponentDetails.cpp
+++ b/Engine/Source/Editor/DetailCustomizations/Private/SplineComponentDetails.cpp
@@ -1023,8 +1023,8 @@ PRAGMA_DISABLE_DEPRECATION_WARNINGS
 	const int32 NumPoints = SplineComp->GetNumberOfSplinePoints();
 
 	bool bModifyOtherPoints = false;
-	if ((Index > 0 && NewValue <= SplineComp->GetInputKeyValueAtSplinePoint(Index - 1) ||
-		(Index < NumPoints - 1 && NewValue >= SplineComp->GetInputKeyValueAtSplinePoint(Index + 1))))
+	if ((Index > 0 && NewValue <= SplineComp->GetInputKeyValueAtSplinePoint(Index - 1)) ||
+		(Index < NumPoints - 1 && NewValue >= SplineComp->GetInputKeyValueAtSplinePoint(Index + 1)))
 	{
 		const FText Title(LOCTEXT("InputKeyTitle", "Input key out of range"));
 		const FText Message(LOCTEXT("InputKeyMessage", "Spline input keys must be numerically ascending. Would you like to modify other input keys in the spline in order to be able to set this value?"));
