[Slate] Fixed triggered ensure inside Interval Struct Customization when Property Node is no longer valid.

diff --git a/Engine/Source/Editor/DetailCustomizations/Private/IntervalStructCustomization.cpp b/Engine/Source/Editor/DetailCustomizations/Private/IntervalStructCustomization.cpp
--- a/Engine/Source/Editor/DetailCustomizations/Private/IntervalStructCustomization.cpp
+++ b/Engine/Source/Editor/DetailCustomizations/Private/IntervalStructCustomization.cpp
@@ -260,6 +260,14 @@ void FIntervalStructCustomization<NumericType>::SetValue(NumericType NewValue, E
 	IPropertyHandle* Handle = (Field == EIntervalField::Min) ? MinValueHandle.Get() : MaxValueHandle.Get();
 	IPropertyHandle* OtherHandle = (Field == EIntervalField::Min) ? MaxValueHandle.Get() : MinValueHandle.Get();
 
+	// This is to guard out 2 cases
+	// a). SSpinBox flushes out any attempted uncommited value during destruction, and Custom Widgets was destructed before IPropertyTypeCustomization(see FDetailPropertyRow)
+	// b). It's possible that the event DetailsView resets its Object gets routed before the event SSpinBox loses user focus, which will commit value.   
+	if (Handle->GetNumOuterObjects() == 0)
+	{
+		return;
+	}
+
 	const TOptional<NumericType> OtherValue = GetValue<NumericType>(OtherHandle);
 	const bool bOutOfRange = OtherValue.IsSet() && ((Field == EIntervalField::Min && NewValue > OtherValue.GetValue()) || (Field == EIntervalField::Max && NewValue < OtherValue.GetValue()));
 
diff --git a/Engine/Source/Editor/UnrealEd/Private/SCurveEditor.cpp b/Engine/Source/Editor/UnrealEd/Private/SCurveEditor.cpp
--- a/Engine/Source/Editor/UnrealEd/Private/SCurveEditor.cpp
+++ b/Engine/Source/Editor/UnrealEd/Private/SCurveEditor.cpp
@@ -3840,6 +3840,9 @@ void SCurveEditor::RefreshDetailsView()
 	//To clean up dangling references
 	EmptyAllSelection();
 	EmptyAllCurveViewModels();
+
+	// Refresh will still happen in the current frame, but it's important to make it asynchronous because its possible the user sets value through property handle multiple times inside an atomic action
+	// Making it synchronous will invalidate property nodes for all the subsequent property handle set value operation.
 	RegisterActiveTimer(0.f, FWidgetActiveTimerDelegate::CreateSPLambda(
 		this,
 		[this](double, float)
