Unshelved from pending changelist '47868638':

[Slate] Fixed Interval Customization not able to set value when the outer is a struct

diff --git a/Engine/Source/Editor/DetailCustomizations/Private/IntervalStructCustomization.cpp b/Engine/Source/Editor/DetailCustomizations/Private/IntervalStructCustomization.cpp
--- a/Engine/Source/Editor/DetailCustomizations/Private/IntervalStructCustomization.cpp
+++ b/Engine/Source/Editor/DetailCustomizations/Private/IntervalStructCustomization.cpp
@@ -265,7 +265,13 @@ void FIntervalStructCustomization<NumericType>::SetValue(NumericType NewValue, E
 	// b). It's possible that the event DetailsView resets its Object gets routed before the event SSpinBox loses user focus, which will commit value.   
 	if (Handle->GetNumOuterObjects() == 0)
 	{
-		return;
+		TArray<TSharedPtr<FStructOnScope>> OuterStructs;
+		Handle->GetOuterStructs(OuterStructs);
+
+		if (OuterStructs.Num() == 0)
+		{
+			return;
+		}
 	}
 
 	const TOptional<NumericType> OtherValue = GetValue<NumericType>(OtherHandle);
