Component Transform Details: Speculative fix for crash when a SceneComponent is not found

diff --git a/Engine/Source/Editor/DetailCustomizations/Private/ComponentTransformDetails.cpp b/Engine/Source/Editor/DetailCustomizations/Private/ComponentTransformDetails.cpp
--- a/Engine/Source/Editor/DetailCustomizations/Private/ComponentTransformDetails.cpp
+++ b/Engine/Source/Editor/DetailCustomizations/Private/ComponentTransformDetails.cpp
@@ -1679,7 +1679,7 @@ void FComponentTransformDetails::OnSetTransform(ETransformField::Type TransformF
 					}
 				}
 				
-				if (!Object->IsTemplate())
+				if (SceneComponent && !Object->IsTemplate())
 				{
 					if (TransformField == ETransformField::Rotation || TransformField == ETransformField::Location)
 					{
