Compacting FDataTableRowHandle into a single line

diff --git a/Engine/Source/Editor/DetailCustomizations/Private/DataTableCustomization.cpp b/Engine/Source/Editor/DetailCustomizations/Private/DataTableCustomization.cpp
--- a/Engine/Source/Editor/DetailCustomizations/Private/DataTableCustomization.cpp
+++ b/Engine/Source/Editor/DetailCustomizations/Private/DataTableCustomization.cpp
@@ -23,6 +23,7 @@
 #include "Templates/Casts.h"
 #include "UObject/Class.h"
 #include "UObject/Object.h"
+#include "UObject/PropertyAccessUtil.h"
 #include "Widgets/DeclarativeSyntaxSupport.h"
 #include "Widgets/Text/STextBlock.h"
 
@@ -41,16 +42,46 @@ void FDataTableCustomizationLayout::CustomizeHeader(TSharedRef<class IPropertyHa
 		RowFilterStruct = UClass::TryFindTypeSlow<UScriptStruct>(RowType);
 	}
 
+	if (StructPropertyHandle->HasMetaData(TEXT("RowPreviewProperty")))
+	{
+		RowPreviewPropertyName = FName(*StructPropertyHandle->GetMetaData(TEXT("RowPreviewProperty")));
+	}
+	else
+	{
+		RowPreviewPropertyName = NAME_None;
+	}
+
 	FSimpleDelegate OnDataTableChangedDelegate = FSimpleDelegate::CreateSP(this, &FDataTableCustomizationLayout::OnDataTableChanged);
 	StructPropertyHandle->SetOnPropertyValueChanged(OnDataTableChangedDelegate);
+}
 	
-	HeaderRow
-	.NameContent()
-	[
-		InStructPropertyHandle->CreatePropertyNameWidget()
-	];
+FText FDataTableCustomizationLayout::GetPreviewText() const
+{
+	UDataTable *CurrentDataTable = nullptr;
+	FName CurrentRowName;
+	if (!RowPreviewPropertyName.IsNone() && GetCurrentValue(CurrentDataTable, CurrentRowName))
+	{
+		FProperty* DataTableRowProperty = PropertyAccessUtil::FindPropertyByName(RowPreviewPropertyName, CurrentDataTable->GetRowStruct());
+		uint8* DataTableRow = CurrentDataTable->FindRowUnchecked(CurrentRowName);
+
+		if (DataTableRowProperty && DataTableRow)
+		{
+			uint8* RowData = DataTableRowProperty->ContainerPtrToValuePtr<uint8>(DataTableRow, 0);
+
+			FString PreviewValue;
+			DataTableRowProperty->ExportText_Direct(
+				PreviewValue,
+				RowData,
+				nullptr,
+				nullptr,
+				PPF_None
+			);
 
-	FDataTableEditorUtils::AddSearchForReferencesContextMenu(HeaderRow, FExecuteAction::CreateSP(this, &FDataTableCustomizationLayout::OnSearchForReferences));
+			return FText::FromString(FString::Printf(TEXT("%s = %s"), *RowPreviewPropertyName.ToString(), *PreviewValue));
+		}
+	}
+
+	return FText::GetEmpty();
 }
 
 void FDataTableCustomizationLayout::CustomizeChildren(TSharedRef<class IPropertyHandle> InStructPropertyHandle, class IDetailChildrenBuilder& StructBuilder, IPropertyTypeCustomizationUtils& StructCustomizationUtils)
@@ -66,43 +97,64 @@ void FDataTableCustomizationLayout::CustomizeChildren(TSharedRef<class IProperty
 		FSimpleDelegate OnDataTableChangedDelegate = FSimpleDelegate::CreateSP(this, &FDataTableCustomizationLayout::OnDataTableChanged);
 		DataTablePropertyHandle->SetOnPropertyValueChanged(OnDataTableChangedDelegate);
 
-		/** Construct a asset picker widget with a custom filter */
-		StructBuilder.AddCustomRow(LOCTEXT("DataTable_TableName", "Data Table"))
-			.RowTag(PropertyName)
+		FPropertyComboBoxArgs ComboArgs(RowNamePropertyHandle, 
+			FOnGetPropertyComboBoxStrings::CreateSP(this, &FDataTableCustomizationLayout::OnGetRowStrings), 
+			FOnGetPropertyComboBoxValue::CreateSP(this, &FDataTableCustomizationLayout::OnGetRowValueString));
+		ComboArgs.ShowSearchForItemCount = 1;
+
+		FUIAction CopyAction;
+		FUIAction PasteAction;
+		InStructPropertyHandle->CreateDefaultPropertyCopyPasteActions(CopyAction, PasteAction);
+
+		FDetailWidgetRow& InlineRow = StructBuilder.AddCustomRow(LOCTEXT("DataTableRowHandle", "Data Table Row Handle"))
+			.CopyAction(CopyAction)
+			.PasteAction(PasteAction)
 			.NameContent()
 			[
-				SNew(STextBlock)
-				.Text(LOCTEXT("DataTable_TableName", "Data Table"))
-				.Font(StructCustomizationUtils.GetRegularFont())
+				InStructPropertyHandle->CreatePropertyNameWidget()
 			]
 			.ValueContent()
-			.MaxDesiredWidth(0.0f) // don't constrain the combo button width
+			.MaxDesiredWidth(0.0f)
 			[
+				SNew(SHorizontalBox)
+				+ SHorizontalBox::Slot()
+					.AutoWidth()
+					[
 				SNew(SObjectPropertyEntryBox)
 				.PropertyHandle(DataTablePropertyHandle)
 				.AllowedClass(UDataTable::StaticClass())
 				.OnShouldFilterAsset(this, &FDataTableCustomizationLayout::ShouldFilterAsset)
-			];
-
-		FPropertyComboBoxArgs ComboArgs(RowNamePropertyHandle, 
-			FOnGetPropertyComboBoxStrings::CreateSP(this, &FDataTableCustomizationLayout::OnGetRowStrings), 
-			FOnGetPropertyComboBoxValue::CreateSP(this, &FDataTableCustomizationLayout::OnGetRowValueString));
-		ComboArgs.ShowSearchForItemCount = 1;
-
-		/** Construct a combo box widget to select from a list of valid options */
-		StructBuilder.AddCustomRow(LOCTEXT("DataTable_RowName", "Row Name"))
-			.RowTag(PropertyName)
-			.NameContent()
+					]
+				+ SHorizontalBox::Slot()
+					.AutoWidth()
+					[
+						PropertyCustomizationHelpers::MakePropertyComboBox(ComboArgs)
+					]			
+				+ SHorizontalBox::Slot()
+					.AutoWidth()
+					.Padding(8, 0)
 			[
 				SNew(STextBlock)
-				.Text(LOCTEXT("DataTable_RowName", "Row Name"))
-				.Font(StructCustomizationUtils.GetRegularFont())
+						.Visibility_Lambda([this]() { return RowPreviewPropertyName.IsNone() ? EVisibility::Hidden : EVisibility::Visible; })
+						.Text(this, &FDataTableCustomizationLayout::GetPreviewText)
+					]	
 			]
-			.ValueContent()
-			.MaxDesiredWidth(0.0f) // don't constrain the combo button width
-			[
-				PropertyCustomizationHelpers::MakePropertyComboBox(ComboArgs)
-			];
+			.OverrideResetToDefault(FResetToDefaultOverride::Create(
+				TAttribute<bool>::CreateLambda([InStructPropertyHandle]() -> bool
+				{
+					return InStructPropertyHandle->IsValidHandle() && InStructPropertyHandle->DiffersFromDefault();
+				}),
+				FSimpleDelegate::CreateLambda([InStructPropertyHandle]()
+				{
+					if (InStructPropertyHandle->IsValidHandle())
+					{
+						InStructPropertyHandle->ResetToDefault();
+					}
+				})
+			)
+		);
+
+		FDataTableEditorUtils::AddSearchForReferencesContextMenu(InlineRow, FExecuteAction::CreateSP(this, &FDataTableCustomizationLayout::OnSearchForReferences));				
 	}
 }
 
diff --git a/Engine/Source/Editor/DetailCustomizations/Private/DataTableCustomization.h b/Engine/Source/Editor/DetailCustomizations/Private/DataTableCustomization.h
--- a/Engine/Source/Editor/DetailCustomizations/Private/DataTableCustomization.h
+++ b/Engine/Source/Editor/DetailCustomizations/Private/DataTableCustomization.h
@@ -49,11 +49,16 @@ private:
 	/** Gets currently selected row display string */
 	FString OnGetRowValueString() const;
 
+	/** Enables previewing of a selected row value **/
+	FText GetPreviewText() const;
+
 	/** Handle to the struct properties being customized */
 	TSharedPtr<IPropertyHandle> StructPropertyHandle;
 	TSharedPtr<IPropertyHandle> DataTablePropertyHandle;
 	TSharedPtr<IPropertyHandle> RowNamePropertyHandle;
 	/** The MetaData derived filter for the row type */
 	FName RowTypeFilter;
+	FName RowPreviewPropertyName;
 	UScriptStruct* RowFilterStruct;
+	bool bShowOnlyInnerProperties = false;
 };
