Fix confusing error popups about package names that stop you from selecting a map to cook in project settings

diff --git a/Engine/Source/Editor/DetailCustomizations/Private/FilePathStructCustomization.cpp b/Engine/Source/Editor/DetailCustomizations/Private/FilePathStructCustomization.cpp
--- a/Engine/Source/Editor/DetailCustomizations/Private/FilePathStructCustomization.cpp
+++ b/Engine/Source/Editor/DetailCustomizations/Private/FilePathStructCustomization.cpp
@@ -102,9 +102,14 @@ void FFilePathStructCustomization::HandleFilePathPickerPathPicked( const FString
 	{
 		FString LongPackageName;
 		FString StringFailureReason;
-		if (FPackageName::TryConvertFilenameToLongPackageName(PickedPath, LongPackageName, &StringFailureReason) == false)
+
+		// Don't try to convert empty string
+		if (PickedPath.Len() > 0)
 		{
-			FMessageDialog::Open(EAppMsgType::Ok, FText::FromString(StringFailureReason));
+			if (FPackageName::TryConvertFilenameToLongPackageName(PickedPath, LongPackageName, &StringFailureReason) == false)
+			{
+				FMessageDialog::Open(EAppMsgType::Ok, FText::FromString(StringFailureReason));
+			}
 		}
 		FinalPath = LongPackageName;
 	}
