[Audio] Fix inconsistent ordering of attenuation categories in AudioComponent/SoundCue/SoundCue Attenuation node overrides and in attenuation assets
+ Fixes unrelated categories intertwined with attenuation categories in AudioComponent/SoundCue/SoundCue Attenuation node overrides
+ Fixes ordering not applied in attenuation assets

#tests Verified all said locations now use the same order for their attenuation categories contiguously

diff --git a/Engine/Source/Editor/DetailCustomizations/Private/AttenuationSettingsCustomizations.cpp b/Engine/Source/Editor/DetailCustomizations/Private/AttenuationSettingsCustomizations.cpp
--- a/Engine/Source/Editor/DetailCustomizations/Private/AttenuationSettingsCustomizations.cpp
+++ b/Engine/Source/Editor/DetailCustomizations/Private/AttenuationSettingsCustomizations.cpp
@@ -45,15 +45,29 @@ namespace AttenuationSettingsUtils
 
 	void SortCategories(const TMap<FName, IDetailCategoryBuilder*>& AllCategoryMap)
 	{
-		const IDetailCategoryBuilder* AttenuationBuilder = AllCategoryMap.FindRef("Attenuation");
-		if (!AttenuationBuilder)
+		int32 AttenuationSortIndexBase = 0;
+		if (const IDetailCategoryBuilder* AttenuationBuilder = AllCategoryMap.FindRef("Attenuation"))
 		{
-			return;
+			AttenuationSortIndexBase = AttenuationBuilder->GetSortOrder();
 		}
 
+		static const TArray<FName> OrderedCategories = {
+			FName("AttenuationDistance"),
+			FName("AttenuationSpatialization"),
+			FName("AttenuationOcclusion"),
+			FName("AttenuationSubmixSend"),
+			FName("AttenuationReverbSend"),
+			FName("AttenuationListenerFocus"),
+			FName("AttenuationPriority"),
+			FName("AttenuationAirAbsorption"),
+			FName("AttenuationPluginSettings"),
+			FName("AttenuationSourceDataOverride"),
+			FName("AttenuationAudioLink"),
+		};
+		static const int32 FirstNonAttenuationIndexOffsetFromBase = OrderedCategories.Num() + 1;
+		
 		for (const TPair<FName, IDetailCategoryBuilder*>& Pair : AllCategoryMap)
 		{
-			int32 SortOrder = AttenuationBuilder->GetSortOrder();
 			const FName CategoryName = Pair.Key;
 
 			// Early out if attenuation category
@@ -62,68 +76,28 @@ namespace AttenuationSettingsUtils
 				continue;
 			}
 
-			// Organize related categories
-			if (CategoryName == "AttenuationDistance")
+			int32 CategoryNewSortIndex;
+			const int32 CategoryIndex = OrderedCategories.Find(CategoryName);
+			if (CategoryIndex != INDEX_NONE)
 			{
-				SortOrder += 1;
+				CategoryNewSortIndex = AttenuationSortIndexBase + CategoryIndex + 1;
 			}
-			else if (CategoryName == "AttenuationSpatialization")
-			{
-				SortOrder += 2;
-			}
-			else if (CategoryName == "AttenuationOcclusion")
-			{
-				SortOrder += 3;
-			}
-			else if (CategoryName == "AttenuationSubmixSend")
-			{
-				SortOrder += 4;
-			}
-			else if (CategoryName == "AttenuationReverbSend")
-			{
-				SortOrder += 5;
-			}
-			else if (CategoryName == "AttenuationListenerFocus")
-			{
-				SortOrder += 6;
-			}
-			else if (CategoryName == "AttenuationPriority")
-			{
-				SortOrder += 7;
-			}
-			else if (CategoryName == "AttenuationAirAbsorption")
-			{
-				SortOrder += 8;
-			}
-			else if (CategoryName == "AttenuationPluginSettings")
-			{
-				SortOrder += 9;
-			}
-			else if (CategoryName == "AttenuationSourceDataOverride")
-			{
-				SortOrder += 10;
-			}
-			else if (CategoryName == "AttenuationAudioLink")
-			{
-				SortOrder += 11;
-			}
-
 			else
 			{
-				// Add space to any other categories interfering with space for attenuation-related categories
-				const int32 ValueSortOrder = Pair.Value->GetSortOrder();
-				if (ValueSortOrder>= SortOrder)
+				// Add space to any other categories that might be interfering with space for attenuation-related categories
+				const int32 UnrelatedCategorySortIndex = Pair.Value->GetSortOrder();
+				if (AttenuationSortIndexBase <= UnrelatedCategorySortIndex)
 				{
-					SortOrder = ValueSortOrder + 12;
+					CategoryNewSortIndex = UnrelatedCategorySortIndex + FirstNonAttenuationIndexOffsetFromBase;
 				}
-				// Else, leave current sort order alone
+				// Otherwise, this unrelated category index doesn't need to be moved around (it's already before us)
 				else
 				{
 					continue;
 				}
 			}
 
-			Pair.Value->SetSortOrder(SortOrder);
+			Pair.Value->SetSortOrder(CategoryNewSortIndex);
 		}
 	}
 }
