[Audio] Fix attenuation shape settings not properly copying/pasting on attenuation detail views

diff --git a/Engine/Source/Editor/DetailCustomizations/Private/AttenuationSettingsCustomizations.cpp b/Engine/Source/Editor/DetailCustomizations/Private/AttenuationSettingsCustomizations.cpp
--- a/Engine/Source/Editor/DetailCustomizations/Private/AttenuationSettingsCustomizations.cpp
+++ b/Engine/Source/Editor/DetailCustomizations/Private/AttenuationSettingsCustomizations.cpp
@@ -204,108 +204,183 @@ void FBaseAttenuationSettingsCustomization::CustomizeChildren(TSharedRef<IProper
 	LayoutBuilder.AddPropertyToCategory(AttenuationShapeHandle)
 		.Visibility(IsAttenuationOverriddenVisibleAttribute())
 		.EditCondition(GetIsAttenuationEnabledAttribute(), nullptr);
-	
-	LayoutBuilder.AddPropertyToCategory(AttenuationExtentsHandle)
-		.Visibility(TAttribute<EVisibility>(this, &FBaseAttenuationSettingsCustomization::IsBoxSelected))
-		.DisplayName(LOCTEXT("BoxExtentsLabel", "Extents"))
-		.ToolTip(LOCTEXT("BoxExtents", "The dimensions of the of the box."))
-		.EditCondition(GetIsAttenuationEnabledAttribute(), nullptr);
 
 	// Get the attenuation category directly here otherwise our category is going to be incorrect for the following custom rows (e.g. "Vector" vs "Attenuation")
 	IDetailCategoryBuilder& AttenuationCategory = LayoutBuilder.EditCategory("AttenuationDistance");
 
-	const FText RadiusLabel(LOCTEXT("RadiusLabel", "Inner Radius"));
-
-	AttenuationCategory.AddCustomRow(RadiusLabel)
-		.NameContent()
-		[
-			SNew(STextBlock)
-				.Text(RadiusLabel)
-				.ToolTipText(LOCTEXT("RadiusToolTip", "The radius that defines when sound attenuation begins (or when a custom attenuation curve begins). Sounds played at a distance less than this will not be attenuated."))
-				.Font(StructCustomizationUtils.GetRegularFont())
-				.IsEnabled(GetIsAttenuationEnabledAttribute())
-		]
-		.ValueContent()
-		[
-			ExtentXHandle->CreatePropertyValueWidget()
-		]
-		.Visibility(TAttribute<EVisibility>(this, &FBaseAttenuationSettingsCustomization::IsSphereSelected))
-		.IsEnabled(GetIsAttenuationEnabledAttribute());
-
-	AttenuationCategory.AddCustomRow(LOCTEXT("CapsuleHalfHeightLabel", "Capsule Half Height"))
-		.NameContent()
-		[
-			SNew(STextBlock)
-			.Text(LOCTEXT("CapsuleHalfHeightLabel", "Capsule Half Height"))
-			.ToolTipText(LOCTEXT("CapsuleHalfHeightToolTip", "The attenuation capsule's half height."))
-			.Font(StructCustomizationUtils.GetRegularFont())
-		]
-		.ValueContent()
-		[
-			ExtentXHandle->CreatePropertyValueWidget()
-		]
-		.Visibility(TAttribute<EVisibility>(this, &FBaseAttenuationSettingsCustomization::IsCapsuleSelected))
-		.IsEnabled(GetIsAttenuationEnabledAttribute());
-
-	AttenuationCategory.AddCustomRow(LOCTEXT("CapsuleRadiusLabel", "Capsule Radius"))
-		.NameContent()
-		[
-			SNew(STextBlock)
-			.Text(LOCTEXT("CapsuleRadiusLabel", "Capsule Radius"))
-			.ToolTipText(LOCTEXT("CapsuleRadiusToolTip", "The attenuation capsule's radius."))
-			.Font(StructCustomizationUtils.GetRegularFont())
-		]
-		.ValueContent()
-		[
-			ExtentYHandle->CreatePropertyValueWidget()
-		]
-		.Visibility(TAttribute<EVisibility>(this, &FBaseAttenuationSettingsCustomization::IsCapsuleSelected))
-		.IsEnabled(GetIsAttenuationEnabledAttribute());
-
-	AttenuationCategory.AddCustomRow(LOCTEXT("ConeRadiusLabel", "Cone Radius"))
-		.NameContent()
-		[
-			SNew(STextBlock)
-			.Text(LOCTEXT("ConeRadiusLabel", "Cone Radius"))
-			.ToolTipText(LOCTEXT("ConeRadiusToolTip", "The attenuation cone's radius."))
-			.Font(StructCustomizationUtils.GetRegularFont())
-		]
-		.ValueContent()
-		[
-			ExtentXHandle->CreatePropertyValueWidget()
-		]
-		.Visibility(TAttribute<EVisibility>(this, &FBaseAttenuationSettingsCustomization::IsConeSelected))
-		.IsEnabled(GetIsAttenuationEnabledAttribute());
-
-	AttenuationCategory.AddCustomRow(LOCTEXT("ConeAngleLabel", "Cone Angle"))
-		.NameContent()
-		[
-			SNew(STextBlock)
-			.Text(LOCTEXT("ConeAngleLabel", "Cone Angle"))
-			.ToolTipText(LOCTEXT("ConeAngleToolTip", "The angle of the inner edge of the attenuation cone's falloff. Inside this angle sounds will be at full volume."))
-			.Font(StructCustomizationUtils.GetRegularFont())
-		]
-		.ValueContent()
-		[
-			ExtentYHandle->CreatePropertyValueWidget()
-		]
-		.Visibility(TAttribute<EVisibility>(this, &FBaseAttenuationSettingsCustomization::IsConeSelected))
-		.IsEnabled(GetIsAttenuationEnabledAttribute());
+	// Attenuation shapes all use the same FVector to store the shape's parameters.
+	// All shapes don't need 3 numbers however, so we only show the relevant ones for each shape.
+	// The Copy/Paste Category feature of the editor expects property handles to be referred exactly once per category, so we can't have multiple
+	// rows for the same property just to, say, change its label/tooltip.
+	// The unreal way is to use slate attributes for customization, and instantiate the property rows exactly once for each coordinate of the vector.
+	
+	// ExtentsX customization
+	{
+		TAttribute<FText> ExtentXLabelAttribute = TAttribute<FText>::CreateSPLambda(this, [this]()
+			{
+				if (IsSphereSelected() == EVisibility::Visible)
+				{
+					return LOCTEXT("RadiusLabel", "Inner Radius");
+				}
+				if (IsCapsuleSelected() == EVisibility::Visible)
+				{
+					return LOCTEXT("CapsuleHalfHeightLabel", "Capsule Half Height");
+				}
+				if (IsConeSelected() == EVisibility::Visible)
+				{
+					return LOCTEXT("ConeRadiusLabel", "Cone Radius");
+				}
+				if (IsBoxSelected() == EVisibility::Visible)
+				{
+					return LOCTEXT("BoxXLabel", "Extents X");
+				}
+				return FText();
+			});
+		TAttribute<FText> ExtentXTooltipAttribute = TAttribute<FText>::CreateSPLambda(this, [this]()
+			{
+				if (IsSphereSelected() == EVisibility::Visible)
+				{
+					return LOCTEXT("RadiusToolTip", "The radius that defines when sound attenuation begins (or when a custom attenuation curve begins). Sounds played at a distance less than this will not be attenuated.");
+				}
+				if (IsCapsuleSelected() == EVisibility::Visible)
+				{
+					return LOCTEXT("CapsuleHalfHeightToolTip", "The attenuation capsule's half height.");
+				}
+				if (IsConeSelected() == EVisibility::Visible)
+				{
+					return LOCTEXT("ConeRadiusToolTip", "The attenuation cone's radius.");
+				}
+				if (IsBoxSelected() == EVisibility::Visible)
+				{
+					return LOCTEXT("BoxXTooltip", "The X dimension of the box");
+				}
+				return FText();
+			});
+		
+		AttenuationCategory.AddProperty(ExtentXHandle)
+			.IsEnabled(GetIsAttenuationEnabledAttribute())
+			.CustomWidget()
+			.NameContent()
+			[
+				SNew(STextBlock)
+					.Text(ExtentXLabelAttribute)
+					.ToolTipText(ExtentXTooltipAttribute)
+					.Font(StructCustomizationUtils.GetRegularFont())
+					.IsEnabled(GetIsAttenuationEnabledAttribute())
+			]
+			.ValueContent()
+			[
+				ExtentXHandle->CreatePropertyValueWidget()
+			];
+	}
+	
+	// ExtentsY customization
+	{
+		TAttribute<FText> ExtentYLabelAttribute = TAttribute<FText>::CreateSPLambda(this, [this]()
+			{
+				if (IsCapsuleSelected() == EVisibility::Visible)
+				{
+					return LOCTEXT("CapsuleRadiusLabel", "Capsule Radius");
+				}
+				if (IsConeSelected() == EVisibility::Visible)
+				{
+					return LOCTEXT("ConeAngleLabel", "Cone Angle");
+				}
+				if (IsBoxSelected() == EVisibility::Visible)
+				{
+					return LOCTEXT("BoxYLabel", "Extents Y");
+				}
+				return FText();
+			});
+		TAttribute<FText> ExtentYTooltipAttribute = TAttribute<FText>::CreateSPLambda(this, [this]()
+			{
+				if (IsCapsuleSelected() == EVisibility::Visible)
+				{
+					return LOCTEXT("CapsuleRadiusToolTip", "The attenuation capsule's radius.");
+				}
+				if (IsConeSelected() == EVisibility::Visible)
+				{
+					return LOCTEXT("ConeAngleToolTip", "The angle of the inner edge of the attenuation cone's falloff. Inside this angle sounds will be at full volume.");
+				}
+				if (IsBoxSelected() == EVisibility::Visible)
+				{
+					return LOCTEXT("BoxYTooltip", "The Y dimension of the box");
+				}
+				return FText();
+			});
+		
+		AttenuationCategory.AddProperty(ExtentYHandle)
+		   .IsEnabled(GetIsAttenuationEnabledAttribute())
+		   .Visibility(TAttribute<EVisibility>::CreateSPLambda(this, [this]()
+			   {
+				   return IsSphereSelected() == EVisibility::Hidden ? EVisibility::Visible : EVisibility::Hidden;
+			   }))
+		   .CustomWidget()
+		   .NameContent()
+		   [
+			   SNew(STextBlock)
+				   .Text(ExtentYLabelAttribute)
+				   .ToolTipText(ExtentYTooltipAttribute)
+				   .Font(StructCustomizationUtils.GetRegularFont())
+				   .IsEnabled(GetIsAttenuationEnabledAttribute())
+		   ]
+		   .ValueContent()
+		   [
+			   ExtentYHandle->CreatePropertyValueWidget()
+		   ];
+	}
 
-	AttenuationCategory.AddCustomRow(LOCTEXT("ConeFalloffAngleLabel", "Cone Falloff Angle"))
-		.NameContent()
-		[
-			SNew(STextBlock)
-			.Text(LOCTEXT("ConeFalloffAngleLabel", "Cone Falloff Angle"))
-			.ToolTipText(LOCTEXT("ConeFalloffAngleToolTip", "The angle of the outer edge of the attenuation cone's falloff. Outside this angle sounds will be inaudible."))
-			.Font(StructCustomizationUtils.GetRegularFont())
-		]
-		.ValueContent()
-		[
-			ExtentZHandle->CreatePropertyValueWidget()
-		]
-		.Visibility(TAttribute<EVisibility>(this, &FBaseAttenuationSettingsCustomization::IsConeSelected))
-		.IsEnabled(GetIsAttenuationEnabledAttribute());
+	// ExtentsZ customization
+	{
+		TAttribute<FText> ExtentZLabelAttribute = TAttribute<FText>::CreateSPLambda(this, [this]()
+			{
+				if (IsConeSelected() == EVisibility::Visible)
+				{
+					return LOCTEXT("ConeFalloffAngleLabel", "Cone Falloff Angle");
+				}
+				if (IsBoxSelected() == EVisibility::Visible)
+				{
+					return LOCTEXT("BoxZLabel", "Extents Z");
+				}
+				return FText();
+			});
+		TAttribute<FText> ExtentZTooltipAttribute = TAttribute<FText>::CreateSPLambda(this, [this]()
+			{
+				if (IsConeSelected() == EVisibility::Visible)
+				{
+					return LOCTEXT("ConeFalloffAngleToolTip", "The angle of the outer edge of the attenuation cone's falloff. Outside this angle sounds will be inaudible.");
+				}
+				if (IsBoxSelected() == EVisibility::Visible)
+				{
+					return LOCTEXT("BoxZTooltip", "The Z dimension of the box");
+				}
+				return FText();
+			});
+		
+		AttenuationCategory.AddProperty(ExtentZHandle)
+	   .IsEnabled(GetIsAttenuationEnabledAttribute())
+		   .Visibility(TAttribute<EVisibility>::CreateSPLambda(this, [this]()
+			   {
+				   if (IsConeSelected() == EVisibility::Visible || IsBoxSelected() == EVisibility::Visible)
+				   {
+					   return EVisibility::Visible;
+				   }
+				   return EVisibility::Hidden;
+			   }))
+		   .CustomWidget()
+		   .NameContent()
+		   [
+			   SNew(STextBlock)
+				   .Text(ExtentZLabelAttribute)
+				   .ToolTipText(ExtentZTooltipAttribute)
+				   .Font(StructCustomizationUtils.GetRegularFont())
+				   .IsEnabled(GetIsAttenuationEnabledAttribute())
+		   ]
+		   .ValueContent()
+		   [
+			   ExtentZHandle->CreatePropertyValueWidget()
+		   ];
+	}
 
 	LayoutBuilder.AddPropertyToCategory(PropertyHandles.FindChecked(GET_MEMBER_NAME_CHECKED(FBaseAttenuationSettings, ConeOffset)))
 		.Visibility(TAttribute<EVisibility>(this, &FBaseAttenuationSettingsCustomization::IsConeSelected))
