[Audio] Fix Attenuation Volume category showing up when attenuation is not overridden on the audio component/sound cue

diff --git a/Engine/Source/Editor/DetailCustomizations/Private/AttenuationSettingsCustomizations.cpp b/Engine/Source/Editor/DetailCustomizations/Private/AttenuationSettingsCustomizations.cpp
--- a/Engine/Source/Editor/DetailCustomizations/Private/AttenuationSettingsCustomizations.cpp
+++ b/Engine/Source/Editor/DetailCustomizations/Private/AttenuationSettingsCustomizations.cpp
@@ -258,15 +258,16 @@ void FBaseAttenuationSettingsCustomization::CustomizeChildren(TSharedRef<IProper
 			});
 		
 		AttenuationCategory.AddProperty(ExtentXHandle)
+			.Visibility(IsAttenuationOverriddenVisibleAttribute())
 			.IsEnabled(GetIsAttenuationEnabledAttribute())
 			.CustomWidget()
 			.NameContent()
 			[
 				SNew(STextBlock)
-					.Text(ExtentXLabelAttribute)
-					.ToolTipText(ExtentXTooltipAttribute)
-					.Font(StructCustomizationUtils.GetRegularFont())
-					.IsEnabled(GetIsAttenuationEnabledAttribute())
+				.Text(ExtentXLabelAttribute)
+				.ToolTipText(ExtentXTooltipAttribute)
+				.Font(StructCustomizationUtils.GetRegularFont())
+				.IsEnabled(GetIsAttenuationEnabledAttribute())
 			]
 			.ValueContent()
 			[
@@ -310,24 +311,29 @@ void FBaseAttenuationSettingsCustomization::CustomizeChildren(TSharedRef<IProper
 			});
 		
 		AttenuationCategory.AddProperty(ExtentYHandle)
-		   .IsEnabled(GetIsAttenuationEnabledAttribute())
-		   .Visibility(TAttribute<EVisibility>::CreateSPLambda(this, [this]()
-			   {
-				   return IsSphereSelected() == EVisibility::Hidden ? EVisibility::Visible : EVisibility::Hidden;
-			   }))
-		   .CustomWidget()
-		   .NameContent()
-		   [
-			   SNew(STextBlock)
-				   .Text(ExtentYLabelAttribute)
-				   .ToolTipText(ExtentYTooltipAttribute)
-				   .Font(StructCustomizationUtils.GetRegularFont())
-				   .IsEnabled(GetIsAttenuationEnabledAttribute())
-		   ]
-		   .ValueContent()
-		   [
-			   ExtentYHandle->CreatePropertyValueWidget()
-		   ];
+			.IsEnabled(GetIsAttenuationEnabledAttribute())
+			.Visibility(TAttribute<EVisibility>::CreateSPLambda(this, [this]()
+			{
+				return
+					!AttenuationSettingsUtils::GetValue(bOverrideAttenuationHandle)
+		   				? EVisibility::Hidden
+		   				: IsSphereSelected() == EVisibility::Hidden
+							? EVisibility::Visible
+							: EVisibility::Hidden;
+			}))
+			.CustomWidget()
+			.NameContent()
+			[
+				SNew(STextBlock)
+				.Text(ExtentYLabelAttribute)
+				.ToolTipText(ExtentYTooltipAttribute)
+				.Font(StructCustomizationUtils.GetRegularFont())
+				.IsEnabled(GetIsAttenuationEnabledAttribute())
+			]
+			.ValueContent()
+			[
+				ExtentYHandle->CreatePropertyValueWidget()
+			];
 	}
 
 	// ExtentsZ customization
@@ -358,28 +364,28 @@ void FBaseAttenuationSettingsCustomization::CustomizeChildren(TSharedRef<IProper
 			});
 		
 		AttenuationCategory.AddProperty(ExtentZHandle)
-	   .IsEnabled(GetIsAttenuationEnabledAttribute())
-		   .Visibility(TAttribute<EVisibility>::CreateSPLambda(this, [this]()
-			   {
-				   if (IsConeSelected() == EVisibility::Visible || IsBoxSelected() == EVisibility::Visible)
-				   {
-					   return EVisibility::Visible;
-				   }
-				   return EVisibility::Hidden;
-			   }))
-		   .CustomWidget()
-		   .NameContent()
-		   [
-			   SNew(STextBlock)
-				   .Text(ExtentZLabelAttribute)
-				   .ToolTipText(ExtentZTooltipAttribute)
-				   .Font(StructCustomizationUtils.GetRegularFont())
-				   .IsEnabled(GetIsAttenuationEnabledAttribute())
-		   ]
-		   .ValueContent()
-		   [
-			   ExtentZHandle->CreatePropertyValueWidget()
-		   ];
+			.IsEnabled(GetIsAttenuationEnabledAttribute())
+			.Visibility(TAttribute<EVisibility>::CreateSPLambda(this, [this]()
+			{
+				if (IsConeSelected() == EVisibility::Visible || IsBoxSelected() == EVisibility::Visible)
+				{
+					return EVisibility::Visible;
+				}
+				return EVisibility::Hidden;
+			}))
+			.CustomWidget()
+			.NameContent()
+			[
+				SNew(STextBlock)
+				.Text(ExtentZLabelAttribute)
+				.ToolTipText(ExtentZTooltipAttribute)
+				.Font(StructCustomizationUtils.GetRegularFont())
+				.IsEnabled(GetIsAttenuationEnabledAttribute())
+			]
+			.ValueContent()
+			[
+				ExtentZHandle->CreatePropertyValueWidget()
+			];
 	}
 
 	LayoutBuilder.AddPropertyToCategory(PropertyHandles.FindChecked(GET_MEMBER_NAME_CHECKED(FBaseAttenuationSettings, ConeOffset)))
