Fix for crash in the RowPreviewProperty feature of the data table customization.
Made the RowPreviewProperty use a more display friendly output format.
Made the customization use a SWrapBox to make it work correctly when sized down.

diff --git a/Engine/Source/Editor/DetailCustomizations/Private/DataTableCustomization.cpp b/Engine/Source/Editor/DetailCustomizations/Private/DataTableCustomization.cpp
--- a/Engine/Source/Editor/DetailCustomizations/Private/DataTableCustomization.cpp
+++ b/Engine/Source/Editor/DetailCustomizations/Private/DataTableCustomization.cpp
@@ -3,6 +3,7 @@
 #include "DataTableCustomization.h"
 
 #include "AssetRegistry/AssetData.h"
+#include "Components/WrapBox.h"
 #include "Containers/Map.h"
 #include "DataTableEditorUtils.h"
 #include "Delegates/Delegate.h"
@@ -59,7 +60,7 @@ FText FDataTableCustomizationLayout::GetPreviewText() const
 {
 	UDataTable *CurrentDataTable = nullptr;
 	FName CurrentRowName;
-	if (!RowPreviewPropertyName.IsNone() && GetCurrentValue(CurrentDataTable, CurrentRowName))
+	if (!RowPreviewPropertyName.IsNone() && GetCurrentValue(CurrentDataTable, CurrentRowName) && CurrentDataTable)
 	{
 		FProperty* DataTableRowProperty = PropertyAccessUtil::FindPropertyByName(RowPreviewPropertyName, CurrentDataTable->GetRowStruct());
 		uint8* DataTableRow = CurrentDataTable->FindRowUnchecked(CurrentRowName);
@@ -72,9 +73,9 @@ FText FDataTableCustomizationLayout::GetPreviewText() const
 			DataTableRowProperty->ExportText_Direct(
 				PreviewValue,
 				RowData,
+				RowData,
 				nullptr,
-				nullptr,
-				PPF_None
+				PPF_PropertyWindow | PPF_Delimited | PPF_SimpleObjectText
 			);
 
 			return FText::FromString(FString::Printf(TEXT("%s = %s"), *RowPreviewPropertyName.ToString(), *PreviewValue));
@@ -114,30 +115,30 @@ void FDataTableCustomizationLayout::CustomizeChildren(TSharedRef<class IProperty
 				InStructPropertyHandle->CreatePropertyNameWidget()
 			]
 			.ValueContent()
-			.MaxDesiredWidth(0.0f)
+			.MinDesiredWidth(800.f)
 			[
-				SNew(SHorizontalBox)
-				+ SHorizontalBox::Slot()
-					.AutoWidth()
+				SNew(SWrapBox)
+				.UseAllottedSize(true)
+				.Orientation(Orient_Horizontal)
+				.HAlign(HAlign_Fill)
+				+ SWrapBox::Slot()
 					[
-				SNew(SObjectPropertyEntryBox)
-				.PropertyHandle(DataTablePropertyHandle)
-				.AllowedClass(UDataTable::StaticClass())
-				.OnShouldFilterAsset(this, &FDataTableCustomizationLayout::ShouldFilterAsset)
+						SNew(SObjectPropertyEntryBox)
+						.PropertyHandle(DataTablePropertyHandle)
+						.AllowedClass(UDataTable::StaticClass())
+						.OnShouldFilterAsset(this, &FDataTableCustomizationLayout::ShouldFilterAsset)
 					]
-				+ SHorizontalBox::Slot()
-					.AutoWidth()
+				+ SWrapBox::Slot()
 					[
 						PropertyCustomizationHelpers::MakePropertyComboBox(ComboArgs)
 					]			
-				+ SHorizontalBox::Slot()
-					.AutoWidth()
+				+ SWrapBox::Slot()
 					.Padding(8, 0)
-			[
-				SNew(STextBlock)
-						.Visibility_Lambda([this]() { return RowPreviewPropertyName.IsNone() ? EVisibility::Hidden : EVisibility::Visible; })
-						.Text(this, &FDataTableCustomizationLayout::GetPreviewText)
-					]	
+					[
+						SNew(STextBlock)
+							.Visibility_Lambda([this]() { return RowPreviewPropertyName.IsNone() ? EVisibility::Hidden : EVisibility::Visible; })
+							.Text(this, &FDataTableCustomizationLayout::GetPreviewText)
+					]
 			]
 			.OverrideResetToDefault(FResetToDefaultOverride::Create(
 				TAttribute<bool>::CreateLambda([InStructPropertyHandle]() -> bool
