[Rendering] Refactor Custom Primitive Data details view to move some of the logic to the primitive component and allow setuping the custom primitive data on the primitive component using the assigned material

* The widget had custom logic to gather all the parameters for Custom Primitive data on a primitive component
* Logic is now factorized so it can be re-used for the function that setup default values from the material

diff --git a/Engine/Source/Editor/DetailCustomizations/Private/CustomPrimitiveDataCustomization.cpp b/Engine/Source/Editor/DetailCustomizations/Private/CustomPrimitiveDataCustomization.cpp
--- a/Engine/Source/Editor/DetailCustomizations/Private/CustomPrimitiveDataCustomization.cpp
+++ b/Engine/Source/Editor/DetailCustomizations/Private/CustomPrimitiveDataCustomization.cpp
@@ -148,7 +148,7 @@ void FCustomPrimitiveDataCustomization::CustomizeChildren(TSharedRef<IPropertyHa
 		if (VectorGroup == NULL && VectorParameterData.Contains(PrimIdx))
 		{
 			bool bContainsFirstElementOfVector = false;
-			for (const FParameterData& ParameterData : VectorParameterData[PrimIdx])
+			for (const FCustomPrimitiveDataParameter& ParameterData : VectorParameterData[PrimIdx])
 			{
 				if (ParameterData.IndexOffset == 0)
 				{
@@ -217,7 +217,7 @@ IDetailGroup* FCustomPrimitiveDataCustomization::CreateVectorGroup(IDetailChildr
 	TSet<FName> ParameterNames;
 	TSet<TSoftObjectPtr<UMaterialInterface>> Materials;
 
-	for (const FParameterData& ParameterData : VectorParameterData[PrimIdx])
+	for (const FCustomPrimitiveDataParameter& ParameterData : VectorParameterData[PrimIdx])
 	{
 		if (AddedParametersForThisGroup.Contains(ParameterData.ExpressionID))
 		{
@@ -326,7 +326,7 @@ void FCustomPrimitiveDataCustomization::CreateParameterRow(IDetailChildrenBuilde
 
 	if (VectorParameterData.Contains(PrimIdx))
 	{
-		for (const FParameterData& ParameterData : VectorParameterData[PrimIdx])
+		for (const FCustomPrimitiveDataParameter& ParameterData : VectorParameterData[PrimIdx])
 		{
 			if (AddedParametersForThisRow.Contains(ParameterData.ExpressionID))
 			{
@@ -388,7 +388,7 @@ void FCustomPrimitiveDataCustomization::CreateParameterRow(IDetailChildrenBuilde
 
 	if (ScalarParameterData.Contains(PrimIdx))
 	{
-		for (const FParameterData& ParameterData : ScalarParameterData[PrimIdx])
+		for (const FCustomPrimitiveDataParameter& ParameterData : ScalarParameterData[PrimIdx])
 		{
 			if (AddedParametersForThisRow.Contains(ParameterData.ExpressionID))
 			{
@@ -535,16 +535,20 @@ void FCustomPrimitiveDataCustomization::Cleanup()
 
 void FCustomPrimitiveDataCustomization::PopulateParameterData(UPrimitiveComponent* PrimitiveComponent, int32& MaxPrimitiveDataIndex)
 {
+	check(PrimitiveComponent);
+	TArray<TPair<uint8, FCustomPrimitiveDataParameter>> Parameters;
+	PrimitiveComponent->GatherDefaultCustomPrimitiveDataParametersFromMaterials(Parameters, MaxPrimitiveDataIndex);
+
 	const int32 NumMaterials = PrimitiveComponent->GetNumMaterials();
 
 	TSet<TSoftObjectPtr<UMaterial>>& CachedComponentMaterials = ComponentsToWatch.FindOrAdd(PrimitiveComponent);
 	ComponentMaterialCounts.FindOrAdd(PrimitiveComponent) = NumMaterials;
-	for (int32 i = 0; i < NumMaterials; ++i)
+
+	for (auto& [PrimIdx, Parameter] : Parameters)
 	{
-		UMaterialInterface* MaterialInterface = PrimitiveComponent->GetMaterial(i);
-		UMaterial* Material = MaterialInterface ? MaterialInterface->GetMaterial() : NULL;
+		UMaterial* Material = Parameter.Material.IsValid() ? Parameter.Material->GetMaterial() : nullptr;
 
-		if (Material == NULL)
+		if (!Material)
 		{
 			continue;
 		}
@@ -552,76 +556,17 @@ void FCustomPrimitiveDataCustomization::PopulateParameterData(UPrimitiveComponen
 		MaterialsToWatch.Add(Material);
 		CachedComponentMaterials.Add(Material);
 
-		TMap<FMaterialParameterInfo, FMaterialParameterMetadata> Parameters;
-
-		MaterialInterface->GetAllParametersOfType(EMaterialParameterType::Vector, Parameters);
-
-		for (const TPair<FMaterialParameterInfo, FMaterialParameterMetadata>& Parameter : Parameters)
+		if (Parameter.Type == EMaterialParameterType::Scalar)
 		{
-			const FMaterialParameterInfo& Info = Parameter.Key;
-			const FMaterialParameterMetadata& ParameterMetadata = Parameter.Value;
-			
-			if (ParameterMetadata.PrimitiveDataIndex > INDEX_NONE)
-			{
-				const uint8 PrimIdx = static_cast<uint8>(ParameterMetadata.PrimitiveDataIndex);
-				// Add each element individually, so that we can overlap vector parameter names
-				VectorParameterData.FindOrAdd(PrimIdx + 0).Add(
-				{
-					PrimitiveComponent,
-					MaterialInterface,
-					Info,
-					ParameterMetadata.ExpressionGuid,
-					0
-				});
-				VectorParameterData.FindOrAdd(PrimIdx + 1).Add(
-				{
-					PrimitiveComponent,
-					MaterialInterface,
-					Info,
-					ParameterMetadata.ExpressionGuid,
-					1
-				});
-				VectorParameterData.FindOrAdd(PrimIdx + 2).Add(
-				{
-					PrimitiveComponent,
-					MaterialInterface,
-					Info,
-					ParameterMetadata.ExpressionGuid,
-					2
-				});
-				VectorParameterData.FindOrAdd(PrimIdx + 3).Add(
-				{
-					PrimitiveComponent,
-					MaterialInterface,
-					Info,
-					ParameterMetadata.ExpressionGuid,
-					3
-				});
-				MaxPrimitiveDataIndex = FMath::Max(MaxPrimitiveDataIndex, PrimIdx + 3);
-			}
+			ScalarParameterData.FindOrAdd(PrimIdx).Add(MoveTemp(Parameter));
 		}
-
-		Parameters.Reset();
-
-		MaterialInterface->GetAllParametersOfType(EMaterialParameterType::Scalar, Parameters);
-
-		for (const TPair<FMaterialParameterInfo, FMaterialParameterMetadata>& Parameter : Parameters)
+		else if (Parameter.Type == EMaterialParameterType::Vector)
 		{
-			const FMaterialParameterInfo& Info = Parameter.Key;
-			const FMaterialParameterMetadata& ParameterMetadata = Parameter.Value;
-
-			if (ParameterMetadata.PrimitiveDataIndex > INDEX_NONE)
-			{
-				ScalarParameterData.FindOrAdd(static_cast<uint8>(ParameterMetadata.PrimitiveDataIndex)).Add(
-				{
-					PrimitiveComponent,
-					MaterialInterface,
-					Info,
-					ParameterMetadata.ExpressionGuid,
-					0
-				});
-				MaxPrimitiveDataIndex = FMath::Max(MaxPrimitiveDataIndex, (int32)ParameterMetadata.PrimitiveDataIndex);
-			}
+			// Add each element individually, so that we can overlap vector parameter names. Can move the last one.
+			VectorParameterData.FindOrAdd(PrimIdx).Emplace_GetRef(Parameter).IndexOffset = 0;
+			VectorParameterData.FindOrAdd(PrimIdx + 1).Emplace_GetRef(Parameter).IndexOffset = 1;
+			VectorParameterData.FindOrAdd(PrimIdx + 2).Emplace_GetRef(Parameter).IndexOffset = 2;
+			VectorParameterData.FindOrAdd(PrimIdx + 3).Emplace_GetRef(MoveTemp(Parameter)).IndexOffset = 3;
 		}
 	}
 }
@@ -833,7 +778,7 @@ void FCustomPrimitiveDataCustomization::SetDefaultValue(uint8 PrimIdx)
 	// Prioritize vector data since we have a color picker
 	if (VectorParameterData.Contains(PrimIdx))
 	{
-		for (const FParameterData& ParameterData : VectorParameterData[PrimIdx])
+		for (const FCustomPrimitiveDataParameter& ParameterData : VectorParameterData[PrimIdx])
 		{
 			if (ParameterData.Component.IsValid() && !ChangedComponents.Contains(ParameterData.Component))
 			{
@@ -851,7 +796,7 @@ void FCustomPrimitiveDataCustomization::SetDefaultValue(uint8 PrimIdx)
 
 	if (ScalarParameterData.Contains(PrimIdx))
 	{
-		for (const FParameterData& ParameterData : ScalarParameterData[PrimIdx])
+		for (const FCustomPrimitiveDataParameter& ParameterData : ScalarParameterData[PrimIdx])
 		{
 			if (ParameterData.Component.IsValid() && !ChangedComponents.Contains(ParameterData.Component))
 			{
diff --git a/Engine/Source/Editor/DetailCustomizations/Private/CustomPrimitiveDataCustomization.h b/Engine/Source/Editor/DetailCustomizations/Private/CustomPrimitiveDataCustomization.h
--- a/Engine/Source/Editor/DetailCustomizations/Private/CustomPrimitiveDataCustomization.h
+++ b/Engine/Source/Editor/DetailCustomizations/Private/CustomPrimitiveDataCustomization.h
@@ -8,10 +8,10 @@
 #include "HAL/Platform.h"
 #include "IPropertyTypeCustomization.h"
 #include "Input/Reply.h"
-#include "Materials/MaterialParameters.h"
 #include "Math/Color.h"
 #include "Misc/Guid.h"
 #include "PropertyEditorModule.h"
+#include "Components/PrimitiveComponent.h"
 #include "Templates/SharedPointer.h"
 #include "UObject/SoftObjectPtr.h"
 #include "UObject/WeakObjectPtr.h"
@@ -48,16 +48,6 @@ public:
 	void CustomizeChildren(TSharedRef<IPropertyHandle> PropertyHandle, IDetailChildrenBuilder& ChildBuilder, IPropertyTypeCustomizationUtils& CustomizationUtils) override;
 
 private:
-	struct FParameterData
-	{
-	public:
-		TWeakObjectPtr<UPrimitiveComponent> Component;
-		TWeakObjectPtr<UMaterialInterface> Material;
-		FMaterialParameterInfo Info;
-		FGuid ExpressionID;
-		uint8 IndexOffset;
-	};
-
 	// TODO [jonathan.bard] Remove on non-release branch
 	bool bDeferringRefresh = false;
 
@@ -69,8 +59,8 @@ private:
 	TMap<TWeakObjectPtr<UPrimitiveComponent>, uint32> ComponentMaterialCounts;
 	TSet<TSoftObjectPtr<UMaterialInterface>> MaterialsToWatch;
 
-	TMap<uint8, TArray<FParameterData>> VectorParameterData;
-	TMap<uint8, TArray<FParameterData>> ScalarParameterData;
+	TMap<uint8, TArray<FCustomPrimitiveDataParameter>> VectorParameterData;
+	TMap<uint8, TArray<FCustomPrimitiveDataParameter>> ScalarParameterData;
 	TMap<uint8, TSharedPtr<SColorBlock>> ColorBlocks;
 
 	/**
diff --git a/Engine/Source/Runtime/Engine/Classes/Components/PrimitiveComponent.h b/Engine/Source/Runtime/Engine/Classes/Components/PrimitiveComponent.h
--- a/Engine/Source/Runtime/Engine/Classes/Components/PrimitiveComponent.h
+++ b/Engine/Source/Runtime/Engine/Classes/Components/PrimitiveComponent.h
@@ -35,7 +35,7 @@
 #include "MeshDrawCommandStatsDefines.h"
 #include "PrimitiveSceneInfoData.h"
 #include "Streaming/RenderAssetOwnerStreamingState.h"
-
+#include "Materials/MaterialParameters.h"
 #include "PrimitiveComponent.generated.h"
 
 DECLARE_CYCLE_STAT_EXTERN(TEXT("BeginComponentOverlap"), STAT_BeginComponentOverlap, STATGROUP_Game, ENGINE_API);
@@ -224,6 +224,19 @@ enum class EComponentPhysicsStateChange : uint8
 	Destroyed
 };
 
+/** Used to gather information on custom primitive data from the materials. */
+struct FCustomPrimitiveDataParameter
+{
+	TWeakObjectPtr<UPrimitiveComponent> Component;
+	TWeakObjectPtr<UMaterialInterface> Material;
+	FMaterialParameterInfo Info;
+	uint8 IndexOffset;
+	EMaterialParameterType Type;
+#if WITH_EDITOR
+	FGuid ExpressionID;
+#endif // WITH_EDITOR
+};
+
 /**
  * Delegate for notification of blocking collision against a specific component.  
  * NormalImpulse will be filled in for physics-simulating bodies, but will be zero for swept-component blocking collisions. 
@@ -1234,6 +1247,18 @@ public:
 	 */
 	const FCustomPrimitiveData& GetDefaultCustomPrimitiveData() const { return CustomPrimitiveData; }
 
+	/**
+	 * Gather the default custom primitive data parameters. Non-const as it keeps a non-const weak pointer on 'this'.
+	 * Outputs an array of pairs of primitive indices with the corresponding parameter and the max primitive data index.
+	 */
+	ENGINE_API void GatherDefaultCustomPrimitiveDataParametersFromMaterials(TArray<TPair<uint8 /*PrimitiveIndex*/, FCustomPrimitiveDataParameter>>& OutParameters, int32& OutMaxPrimitiveDataIndex);
+
+	/**
+	 * Setup the Default Custom Primitive Data based on the assigned materials. Wipes any existing data.
+	 */
+	UFUNCTION(BlueprintCallable, Category = "Rendering|Material")
+	ENGINE_API void SetupDefaultCustomPrimitiveDataFromMaterials();
+
 #if WITH_EDITOR
 	/** Override delegate used for checking the selection state of a component */
 	DECLARE_DELEGATE_RetVal_OneParam( bool, FSelectionOverride, const UPrimitiveComponent* );
diff --git a/Engine/Source/Runtime/Engine/Private/Components/PrimitiveComponent.cpp b/Engine/Source/Runtime/Engine/Private/Components/PrimitiveComponent.cpp
--- a/Engine/Source/Runtime/Engine/Private/Components/PrimitiveComponent.cpp
+++ b/Engine/Source/Runtime/Engine/Private/Components/PrimitiveComponent.cpp
@@ -2765,6 +2765,101 @@ void UPrimitiveComponent::SetDefaultCustomPrimitiveDataFloatArray(int32 DataInde
 	SetDefaultCustomPrimitiveData(DataIndex, Values);
 }
 
+void UPrimitiveComponent::GatherDefaultCustomPrimitiveDataParametersFromMaterials(TArray<TPair<uint8, FCustomPrimitiveDataParameter>>& OutParameters, int32& OutMaxPrimitiveDataIndex)
+{
+	OutParameters.Reset();
+
+	const int32 NumMaterials = GetNumMaterials();
+
+	for (int32 i = 0; i < NumMaterials; ++i)
+	{
+		UMaterialInterface* MaterialInterface = GetMaterial(i);
+
+		if (MaterialInterface == nullptr)
+		{
+			continue;
+		}
+
+		TMap<FMaterialParameterInfo, FMaterialParameterMetadata> Parameters;
+
+		auto GatherForType = [this, &Parameters, &OutParameters, MaterialInterface, &OutMaxPrimitiveDataIndex](const EMaterialParameterType Type, const int32 TypeSize)
+		{
+			Parameters.Reset();
+
+			MaterialInterface->GetAllParametersOfType(Type, Parameters);
+
+			for (const TPair<FMaterialParameterInfo, FMaterialParameterMetadata>& Parameter : Parameters)
+			{
+				const FMaterialParameterInfo& Info = Parameter.Key;
+				const FMaterialParameterMetadata& ParameterMetadata = Parameter.Value;
+
+				if (ParameterMetadata.PrimitiveDataIndex > INDEX_NONE)
+				{
+					OutParameters.Emplace(
+						static_cast<uint8>(ParameterMetadata.PrimitiveDataIndex),
+						FCustomPrimitiveDataParameter
+						{
+							.Component = this,
+							.Material = MaterialInterface,
+							.Info = Info,
+							.IndexOffset = 0,
+							.Type = Type,
+#if WITH_EDITOR
+							.ExpressionID = ParameterMetadata.ExpressionGuid,
+#endif // WITH_EDITOR
+						}
+					);
+
+					OutMaxPrimitiveDataIndex = FMath::Max(OutMaxPrimitiveDataIndex, ParameterMetadata.PrimitiveDataIndex + TypeSize - 1);
+				}
+			}
+		};
+
+		GatherForType(EMaterialParameterType::Vector, 4);
+		GatherForType(EMaterialParameterType::Scalar, 1);
+	}
+}
+
+void UPrimitiveComponent::SetupDefaultCustomPrimitiveDataFromMaterials()
+{
+	TArray<TPair<uint8, FCustomPrimitiveDataParameter>> Parameters;
+	int32 MaxPrimitiveDataIndex = INDEX_NONE;
+
+	GatherDefaultCustomPrimitiveDataParametersFromMaterials(Parameters, MaxPrimitiveDataIndex);
+
+	CustomPrimitiveData.Data.SetNumZeroed(MaxPrimitiveDataIndex + 1);
+
+	for (const auto& [PrimIdx, Parameter] : Parameters)
+	{
+		if (!Parameter.Material.IsValid())
+		{
+			continue;
+		}
+
+		if (Parameter.Type == EMaterialParameterType::Scalar)
+		{
+			float Value = 0.f;
+			if (Parameter.Material->GetScalarParameterValue(Parameter.Info, Value))
+			{
+				SetDefaultCustomPrimitiveDataFloat(PrimIdx, Value);
+			}
+		}
+		else if (Parameter.Type == EMaterialParameterType::Vector)
+		{
+			FVector4 Value(ForceInitToZero);
+			if (Parameter.Material->GetDoubleVectorParameterDefaultValue(Parameter.Info, Value))
+			{
+				SetDefaultCustomPrimitiveDataVector4(PrimIdx, Value);
+			}
+		}
+		else
+		{
+			// Should not happen
+			ensure(false);
+		}
+	}
+}
+
 UMaterialInterface* UPrimitiveComponent::GetMaterialFromCollisionFaceIndex(int32 FaceIndex, int32& SectionIndex) const
 {
 	//This function should be overriden
