In the editor, if bEnableWorldBoundsChecks is true, clamp component transforms to HALF_WORLD_MAX to avoid users placing components beyond the supported world size.

diff --git a/Engine/Source/Editor/DetailCustomizations/Private/ComponentTransformDetails.cpp b/Engine/Source/Editor/DetailCustomizations/Private/ComponentTransformDetails.cpp
--- a/Engine/Source/Editor/DetailCustomizations/Private/ComponentTransformDetails.cpp
+++ b/Engine/Source/Editor/DetailCustomizations/Private/ComponentTransformDetails.cpp
@@ -19,6 +19,7 @@
 #include "Framework/Commands/UICommandInfo.h"
 #include "Framework/MultiBox/MultiBoxBuilder.h"
 #include "GameFramework/Actor.h"
+#include "GameFramework/WorldSettings.h"
 #include "HAL/PlatformApplicationMisc.h"
 #include "IDetailChildrenBuilder.h"
 #include "IDetailPropertyRow.h"
@@ -1452,6 +1453,28 @@ void FComponentTransformDetails::OnSetTransform(ETransformField::Type TransformF
 				{
 					NewComponentValue = GetAxisFilteredVector(RemappedAxis, SwizzledNewValue, OldComponentValue);
 				}
+
+				// Clamp location to world size
+				if (TransformField == ETransformField::Location)
+				{
+					double MaxMagnitude = HALF_WORLD_MAX;
+					if (UWorld* World = SceneComponent->GetWorld())
+					{
+						AWorldSettings* WorldSettings = World->GetWorldSettings(false);
+						bool bEnableWorldBoundsChecks = WorldSettings ? WorldSettings->bEnableWorldBoundsChecks : true;
+						if (!bEnableWorldBoundsChecks)
+						{
+							MaxMagnitude = UE_DOUBLE_HUGE_DISTANCE;
+						}
+					}
+
+					NewComponentValue = FVector
+					{
+						FMath::Clamp(NewComponentValue.X, -MaxMagnitude, MaxMagnitude),
+						FMath::Clamp(NewComponentValue.Y, -MaxMagnitude, MaxMagnitude),
+						FMath::Clamp(NewComponentValue.Z, -MaxMagnitude, MaxMagnitude)
+					};
+				}
 				
 				auto AreValuesEqual = [this, TransformField](const FVector& NewComponentValue_, const FVector& OldComponentValue_)
 				{
