SoundAttenuation category reorder based on feedback from Design

#Synced-CL 48330817

diff --git a/Engine/Source/Editor/DetailCustomizations/Private/AttenuationSettingsCustomizations.cpp b/Engine/Source/Editor/DetailCustomizations/Private/AttenuationSettingsCustomizations.cpp
--- a/Engine/Source/Editor/DetailCustomizations/Private/AttenuationSettingsCustomizations.cpp
+++ b/Engine/Source/Editor/DetailCustomizations/Private/AttenuationSettingsCustomizations.cpp
@@ -54,12 +54,12 @@ namespace AttenuationSettingsUtils
 		static const TArray<FName> OrderedCategories = {
 			FName("AttenuationDistance"),
 			FName("AttenuationSpatialization"),
+			FName("AttenuationAirAbsorption"),
+			FName("AttenuationListenerFocus"),
 			FName("AttenuationOcclusion"),
 			FName("AttenuationSubmixSend"),
 			FName("AttenuationReverbSend"),
-			FName("AttenuationListenerFocus"),
 			FName("AttenuationPriority"),
-			FName("AttenuationAirAbsorption"),
 			FName("AttenuationPluginSettings"),
 			FName("AttenuationSourceDataOverride"),
 			FName("AttenuationAudioLink"),
