Added customization to core editor for FColorGradingSettings that allows its details panel property structure to match Post Process Settings.

#Synced-CL 48367791

diff --git a/Engine/Plugins/Experimental/ColorCorrectRegions/Source/ColorCorrectRegions/Public/ColorCorrectRegion.h b/Engine/Plugins/Experimental/ColorCorrectRegions/Source/ColorCorrectRegions/Public/ColorCorrectRegion.h
--- a/Engine/Plugins/Experimental/ColorCorrectRegions/Source/ColorCorrectRegions/Public/ColorCorrectRegion.h
+++ b/Engine/Plugins/Experimental/ColorCorrectRegions/Source/ColorCorrectRegions/Public/ColorCorrectRegion.h
@@ -173,7 +173,7 @@ public:
 	float Tint;
 
 	/** Color correction settings. */
-	UPROPERTY(BlueprintReadWrite, EditAnywhere, Category = "Color Grading", meta = (ShowPostProcessCategories))
+	UPROPERTY(BlueprintReadWrite, EditAnywhere, Category = "Color Grading", meta = (ShowOnlyInnerProperties))
 	FColorGradingSettings ColorGradingSettings;
 
 	/** Enable/Disable color correction provided by this region. */
diff --git a/Engine/Plugins/Experimental/ColorCorrectRegions/Source/ColorCorrectRegionsEditor/Private/ColorCorrectRegionCustomization.cpp b/Engine/Plugins/Experimental/ColorCorrectRegions/Source/ColorCorrectRegionsEditor/Private/ColorCorrectRegionCustomization.cpp
--- a/Engine/Plugins/Experimental/ColorCorrectRegions/Source/ColorCorrectRegionsEditor/Private/ColorCorrectRegionCustomization.cpp
+++ b/Engine/Plugins/Experimental/ColorCorrectRegions/Source/ColorCorrectRegionsEditor/Private/ColorCorrectRegionCustomization.cpp
@@ -68,14 +68,6 @@ void FColorCorrectWindowDetails::CustomizeDetails(IDetailLayoutBuilder& DetailBu
 		}
 	}
 
-	// Add all Color Grading struct properties directly to the custom category
-	TSharedRef<IPropertyHandle> ColorGradingPropertyHandle = DetailBuilder.GetProperty(GET_MEMBER_NAME_CHECKED(AColorCorrectionRegion, ColorGradingSettings));
-	if (ColorGradingPropertyHandle.Get().IsValidHandle())
-	{
-		ColorGradingPropertyHandle->MarkHiddenByCustomization();
-		MoveColorGradingPropertiesToCategory(ColorGradingPropertyHandle, ColorGradingCategory);
-	}
-
 	// Sort categories
 	DetailBuilder.SortCategories([&](const TMap<FName, IDetailCategoryBuilder*>& CategoryMap)
 	{
@@ -123,62 +115,4 @@ void FColorCorrectWindowDetails::CustomizeDetails(IDetailLayoutBuilder& DetailBu
 	});
 }
 
-void FColorCorrectWindowDetails::MoveColorGradingPropertiesToCategory(TSharedRef<IPropertyHandle> StructHandle, IDetailCategoryBuilder& RootCategory)
-{
-	// Map from color grading subproperties to the groups we want to assign them to
-	const TMap<FName, FName> ColorGradingGroupAssignments = {
-		{ GET_MEMBER_NAME_CHECKED(FColorGradingSettings, ShadowsMax), GET_MEMBER_NAME_CHECKED(FColorGradingSettings, Shadows) },
-		{ GET_MEMBER_NAME_CHECKED(FColorGradingSettings, HighlightsMin), GET_MEMBER_NAME_CHECKED(FColorGradingSettings, Highlights) },
-		{ GET_MEMBER_NAME_CHECKED(FColorGradingSettings, HighlightsMax), GET_MEMBER_NAME_CHECKED(FColorGradingSettings, Highlights) },
-	};
-
-	TMap<FName, IDetailGroup*> Groups;
-
-	uint32 NumColorGradingChildren = 0;
-	StructHandle->GetNumChildren(NumColorGradingChildren);
-	for (uint32 ColorGradingChildIndex = 0; ColorGradingChildIndex < NumColorGradingChildren; ++ColorGradingChildIndex)
-	{
-		TSharedPtr<IPropertyHandle> SubpropertyHandle = StructHandle->GetChildHandle(ColorGradingChildIndex);
-		FProperty* Subproperty = SubpropertyHandle->GetProperty();
-		const FName SubpropertyName = Subproperty->GetFName();
-
-		uint32 NumGroupChildren = 0;
-		SubpropertyHandle->GetNumChildren(NumGroupChildren);
-		if (NumGroupChildren > 0)
-		{
-			// Recreate the property as a group within the Color Grading category
-			IDetailGroup& NewGroup = RootCategory.AddGroup(SubpropertyName, Subproperty->GetDisplayNameText());
-			NewGroup.EnableReset(true);
-
-			Groups.Emplace(SubpropertyName, &NewGroup);
-
-			for (uint32 GroupChildIndex = 0; GroupChildIndex < NumGroupChildren; ++GroupChildIndex)
-			{
-				TSharedPtr<IPropertyHandle> GroupChildHandle = SubpropertyHandle->GetChildHandle(GroupChildIndex);
-				NewGroup.AddPropertyRow(GroupChildHandle.ToSharedRef());
-			}
-		}
-		else
-		{
-			// Try to assign the property to one of the color grading groups
-			bool bHasAssignedGroup = false;
-
-			if (const FName* AssignedGroupName = ColorGradingGroupAssignments.Find(SubpropertyName))
-			{
-				if (IDetailGroup** AssignedGroup = Groups.Find(*AssignedGroupName))
-				{
-					(*AssignedGroup)->AddPropertyRow(SubpropertyHandle.ToSharedRef());
-					bHasAssignedGroup = true;
-				}
-			}
-
-			// No group found, so just add it to the root color grading category
-			if (!bHasAssignedGroup)
-			{
-				RootCategory.AddProperty(SubpropertyHandle);
-			}
-		}
-	}
-}
-
 #undef LOCTEXT_NAMESPACE
diff --git a/Engine/Plugins/Experimental/ColorCorrectRegions/Source/ColorCorrectRegionsEditor/Private/ColorCorrectRegionCustomization.h b/Engine/Plugins/Experimental/ColorCorrectRegions/Source/ColorCorrectRegionsEditor/Private/ColorCorrectRegionCustomization.h
--- a/Engine/Plugins/Experimental/ColorCorrectRegions/Source/ColorCorrectRegionsEditor/Private/ColorCorrectRegionCustomization.h
+++ b/Engine/Plugins/Experimental/ColorCorrectRegions/Source/ColorCorrectRegionsEditor/Private/ColorCorrectRegionCustomization.h
@@ -17,8 +17,4 @@ public:
 
 	/** IDetailCustomization interface */
 	virtual void CustomizeDetails(IDetailLayoutBuilder& DetailBuilder) override;
-
-private:
-	/** Recreate the color grading property struct's children as groups or root properties of the color grading category */
-	void MoveColorGradingPropertiesToCategory(TSharedRef<IPropertyHandle> StructHandle, IDetailCategoryBuilder& RootCategory);
 };
diff --git a/Engine/Source/Editor/DetailCustomizations/Private/ColorGradingSettingsCustomization.cpp b/Engine/Source/Editor/DetailCustomizations/Private/ColorGradingSettingsCustomization.cpp
new file mode 100644
--- /dev/null
+++ b/Engine/Source/Editor/DetailCustomizations/Private/ColorGradingSettingsCustomization.cpp
@@ -0,0 +1,127 @@
+﻿// Copyright Epic Games, Inc. All Rights Reserved.
+
+#include "ColorGradingSettingsCustomization.h"
+
+#include "DetailWidgetRow.h"
+#include "IDetailChildrenBuilder.h"
+#include "IDetailGroup.h"
+#include "ObjectEditorUtils.h"
+#include "PropertyHandle.h"
+
+void FColorGradingSettingsCustomization::CustomizeHeader(TSharedRef<IPropertyHandle> PropertyHandle, FDetailWidgetRow& HeaderRow, IPropertyTypeCustomizationUtils& CustomizationUtils)
+{
+	// Allow ShowOnlyInnerProperties metadata flag to affect this struct
+	const bool bShowInners = PropertyHandle->HasMetaData("ShowOnlyInnerProperties");
+	if (!bShowInners)
+	{
+		HeaderRow.NameContent()
+		[
+			PropertyHandle->CreatePropertyNameWidget()
+		]
+		.ValueContent()
+		[
+			PropertyHandle->CreatePropertyValueWidget()
+		];
+	}
+}
+
+void FColorGradingSettingsCustomization::CustomizeChildren(TSharedRef<IPropertyHandle> PropertyHandle, IDetailChildrenBuilder& ChildBuilder, IPropertyTypeCustomizationUtils& CustomizationUtils)
+{
+	TMap<FName, FString> PropertyGroupMap;
+	TSet<FString> GroupSet;
+
+	GetPropertyGroups(PropertyHandle, PropertyGroupMap, GroupSet);
+	
+	uint32 NumChildProperties;
+	PropertyHandle->GetNumChildren(NumChildProperties);
+	for (uint32 Index = 0; Index < NumChildProperties; ++Index)
+	{
+		TSharedPtr<IPropertyHandle> ChildHandle = PropertyHandle->GetChildHandle(Index);
+		if (!ChildHandle.IsValid() || !ChildHandle->IsValidHandle())
+		{
+			continue;
+		}
+		
+		FProperty* ChildProperty = ChildHandle->GetProperty();
+
+		if (PropertyGroupMap.Contains(ChildProperty->GetFName()))
+		{
+			// If the property specified an explicit group to be part of, add the handle to that group
+			FString GroupName = PropertyGroupMap[ChildProperty->GetFName()];
+			IDetailGroup& Group = GetOrAddGroup(ChildBuilder, *GroupName, FText::FromString(GroupName));
+			Group.AddPropertyRow(ChildHandle.ToSharedRef());
+		}
+		else if (GroupSet.Contains(ChildHandle->GetPropertyDisplayName().ToString()))
+		{
+			// If an explicit group is needed whose name matches that of the property name, we need an explicit group that the property's children will be added to
+			AddPropertyChildrenExplicitly(ChildHandle.ToSharedRef(), ChildBuilder);
+		}
+		else
+		{
+			ChildBuilder.AddProperty(ChildHandle.ToSharedRef());
+		}
+	}
+}
+
+void FColorGradingSettingsCustomization::GetPropertyGroups(TSharedRef<IPropertyHandle> PropertyHandle, TMap<FName, FString>& OutPropertyGroupMap, TSet<FString>& OutGroupSet)
+{
+	uint32 NumChildProperties;
+	PropertyHandle->GetNumChildren(NumChildProperties);
+	for (uint32 Index = 0; Index < NumChildProperties; ++Index)
+	{
+		TSharedPtr<IPropertyHandle> ChildHandle = PropertyHandle->GetChildHandle(Index);
+		if (!ChildHandle.IsValid() || !ChildHandle->IsValidHandle())
+		{
+			continue;
+		}
+		
+		FProperty* ChildProperty = ChildHandle->GetProperty();
+		FName ChildCategoryFName = FObjectEditorUtils::GetCategoryFName(ChildProperty);
+		FString ChildCategoryStr = ChildCategoryFName.ToString();
+
+		if (!ChildCategoryStr.Contains(TEXT("|")))
+		{
+			continue;
+		}
+
+		TArray<FString> ParsedCategories;
+		ChildCategoryStr.ParseIntoArray(ParsedCategories, TEXT("|"));
+
+		if (ParsedCategories.Num() < 2)
+		{
+			continue;
+		}
+
+		// For now, only support one level of subcategories
+		FString Group = ParsedCategories.Last();
+		OutPropertyGroupMap.Add(ChildProperty->GetFName(), Group);
+		OutGroupSet.Add(Group);
+	}
+}
+
+void FColorGradingSettingsCustomization::AddPropertyChildrenExplicitly(TSharedRef<IPropertyHandle> PropertyHandle, IDetailChildrenBuilder& ChildBuilder)
+{
+	IDetailGroup& PropertyGroup = GetOrAddGroup(ChildBuilder, PropertyHandle->GetProperty()->GetFName(), PropertyHandle->GetPropertyDisplayName());
+	
+	uint32 NumChildren = 0;
+	PropertyHandle->GetNumChildren(NumChildren);
+	for (uint32 ChildIndex = 0; ChildIndex < NumChildren; ++ChildIndex)
+	{
+		TSharedPtr<IPropertyHandle> ChildHandle = PropertyHandle->GetChildHandle(ChildIndex);
+		if (ChildHandle && ChildHandle->IsValidHandle())
+		{
+			PropertyGroup.AddPropertyRow(ChildHandle.ToSharedRef());
+		}
+	}
+}
+
+IDetailGroup& FColorGradingSettingsCustomization::GetOrAddGroup(IDetailChildrenBuilder& ChildBuilder, const FName& GroupName, const FText& GroupDisplayName)
+{
+	IDetailGroup* ExistingGroup = ChildBuilder.GetGroup(GroupName);
+	if (!ExistingGroup)
+	{
+		ExistingGroup = &ChildBuilder.AddGroup(GroupName, GroupDisplayName);
+	}
+
+	return *ExistingGroup;
+}
diff --git a/Engine/Source/Editor/DetailCustomizations/Private/ColorGradingSettingsCustomization.h b/Engine/Source/Editor/DetailCustomizations/Private/ColorGradingSettingsCustomization.h
new file mode 100644
--- /dev/null
+++ b/Engine/Source/Editor/DetailCustomizations/Private/ColorGradingSettingsCustomization.h
@@ -0,0 +1,45 @@
+﻿// Copyright Epic Games, Inc. All Rights Reserved.
+
+#pragma once
+
+#include "IPropertyTypeCustomization.h"
+
+class IDetailGroup;
+
+/**
+ * Implements a details view customization for the FColorGradingSettings structure.
+ */
+class FColorGradingSettingsCustomization : public IPropertyTypeCustomization
+{
+public:
+	static TSharedRef<IPropertyTypeCustomization> MakeInstance()
+	{
+		return MakeShareable(new FColorGradingSettingsCustomization());
+	}
+
+public:
+	// IPropertyTypeCustomization interface
+	virtual void CustomizeHeader(TSharedRef<IPropertyHandle> PropertyHandle, FDetailWidgetRow& HeaderRow, IPropertyTypeCustomizationUtils& CustomizationUtils) override;
+	virtual void CustomizeChildren(TSharedRef<IPropertyHandle> PropertyHandle, IDetailChildrenBuilder& ChildBuilder, IPropertyTypeCustomizationUtils& CustomizationUtils) override;
+
+private:
+	/**
+	 * Iterates the property's children and finds all child properties that have a subcategory and need a group made for them
+	 * @param PropertyHandle The property handle being customized
+	 * @param OutPropertyGroupMap A map of properties and the group they want to be in
+	 * @param OutGroupSet The set of all groups that need to be explicitly added
+	 */
+	void GetPropertyGroups(TSharedRef<IPropertyHandle> PropertyHandle, TMap<FName, FString>& OutPropertyGroupMap, TSet<FString>& OutGroupSet);
+
+	/**
+	 * Adds the specified property's children explicitly to the details panel under a custom group that matches the property's name.
+	 * Used in case a group is needed to embed additional properties that don't live in the specified property
+	 * 
+	 * @param PropertyHandle The property handle whose children are added explicitly
+	 * @param ChildBuilder The detail children builder
+	 */
+	void AddPropertyChildrenExplicitly(TSharedRef<IPropertyHandle> PropertyHandle, IDetailChildrenBuilder& ChildBuilder);
+
+	/** Gets an existing group, or adds one if not found */
+	IDetailGroup& GetOrAddGroup(IDetailChildrenBuilder& ChildBuilder, const FName& GroupName, const FText& GroupDisplayName);
+};
diff --git a/Engine/Source/Editor/DetailCustomizations/Private/DetailCustomizations.cpp b/Engine/Source/Editor/DetailCustomizations/Private/DetailCustomizations.cpp
--- a/Engine/Source/Editor/DetailCustomizations/Private/DetailCustomizations.cpp
+++ b/Engine/Source/Editor/DetailCustomizations/Private/DetailCustomizations.cpp
@@ -33,6 +33,7 @@
 #include "CollectionReferenceStructCustomization.h"
 #include "CollisionProfileDetails.h"
 #include "CollisionProfileNameCustomization.h"
+#include "ColorGradingSettingsCustomization.h"
 #include "ComponentReferenceCustomization.h"
 #include "CompositeRerouteCustomization.h"
 #include "ConfigEditorPropertyDetails.h"
@@ -349,6 +350,7 @@ void FDetailCustomizationsModule::RegisterPropertyTypeCustomizations()
 	RegisterCustomPropertyTypeLayout("MeshMergingSettings", FOnGetPropertyTypeCustomizationInstance::CreateStatic(&FMeshMergingSettingsCustomization::MakeInstance));
 	RegisterCustomPropertyTypeLayout("MeshProxySettings", FOnGetPropertyTypeCustomizationInstance::CreateStatic(&FMeshProxySettingsCustomizations::MakeInstance));
 	RegisterCustomPropertyTypeLayout("PostProcessSettings", FOnGetPropertyTypeCustomizationInstance::CreateStatic(&FPostProcessSettingsCustomization::MakeInstance));
+	RegisterCustomPropertyTypeLayout("ColorGradingSettings", FOnGetPropertyTypeCustomizationInstance::CreateStatic(&FColorGradingSettingsCustomization::MakeInstance));
 	RegisterCustomPropertyTypeLayout("AssetImportInfo", FOnGetPropertyTypeCustomizationInstance::CreateStatic(&FAssetImportDataCustomization::MakeInstance));
 	RegisterCustomPropertyTypeLayout("CaptureResolution", FOnGetPropertyTypeCustomizationInstance::CreateStatic(&FCaptureResolutionCustomization::MakeInstance));
 	RegisterCustomPropertyTypeLayout("CompositionGraphCapturePasses", FOnGetPropertyTypeCustomizationInstance::CreateStatic(&FRenderPassesCustomization::MakeInstance));
diff --git a/Engine/Source/Runtime/Engine/Classes/Engine/Scene.h b/Engine/Source/Runtime/Engine/Classes/Engine/Scene.h
--- a/Engine/Source/Runtime/Engine/Classes/Engine/Scene.h
+++ b/Engine/Source/Runtime/Engine/Classes/Engine/Scene.h
@@ -177,13 +177,13 @@ struct FColorGradingSettings
 	UPROPERTY(Interp, BlueprintReadWrite, Category = "Color Grading")
 	FColorGradePerRangeSettings Highlights;
 
-	UPROPERTY(Interp, BlueprintReadWrite, Category = "Color Grading", meta = (UIMin = "-1.0", UIMax = "1.0", DisplayName = "ShadowsMax"))
+	UPROPERTY(Interp, BlueprintReadWrite, Category = "Color Grading|Shadows", meta = (UIMin = "-1.0", UIMax = "1.0", DisplayName = "ShadowsMax"))
 	float ShadowsMax;
 
-	UPROPERTY(Interp, BlueprintReadWrite, Category = "Color Grading", meta = (UIMin = "-1.0", UIMax = "1.0", DisplayName = "HighlightsMin"))
+	UPROPERTY(Interp, BlueprintReadWrite, Category = "Color Grading|Highlights", meta = (UIMin = "-1.0", UIMax = "1.0", DisplayName = "HighlightsMin"))
 	float HighlightsMin;
 
-	UPROPERTY(Interp, BlueprintReadWrite, Category = "Color Grading", meta = (UIMin = "1.0", UIMax = "10.0", DisplayName = "HighlightsMax"))
+	UPROPERTY(Interp, BlueprintReadWrite, Category = "Color Grading|Highlights", meta = (UIMin = "1.0", UIMax = "10.0", DisplayName = "HighlightsMax"))
 	float HighlightsMax;
 
 	FColorGradingSettings()
