[EditorDetails] Null pointer check for PerPlatform widgets on metadata maps.

* Metadata maps on Properties can be null, and it is not guarded

#rnx
#Synced-CL 48452988

diff --git a/Engine/Source/Editor/DetailCustomizations/Private/PerPlatformPropertyCustomization.cpp b/Engine/Source/Editor/DetailCustomizations/Private/PerPlatformPropertyCustomization.cpp
--- a/Engine/Source/Editor/DetailCustomizations/Private/PerPlatformPropertyCustomization.cpp
+++ b/Engine/Source/Editor/DetailCustomizations/Private/PerPlatformPropertyCustomization.cpp
@@ -111,18 +111,22 @@ TSharedRef<SWidget> FPerPlatformPropertyCustomization<PerPlatformType>::GetWidge
 	if (EditProperty.IsValid())
 	{
 		// First get the source map
-		const TMap<FName, FString>* SourceMap = StructPropertyHandle->GetMetaDataProperty()->GetMetaDataMap();
-		// Iterate through source map, setting each key/value pair in the destination
-		for (const auto& It : *SourceMap)
+		if (const TMap<FName, FString>* SourceMap = StructPropertyHandle->GetMetaDataProperty()->GetMetaDataMap())
 		{
-			EditProperty->SetInstanceMetaData(*It.Key.ToString(), *It.Value);
+			// Iterate through source map, setting each key/value pair in the destination
+			for (const auto& It : *SourceMap)
+			{
+				EditProperty->SetInstanceMetaData(*It.Key.ToString(), *It.Value);
+			}
 		}
 
 		// Copy instance metadata as well
-		const TMap<FName, FString>* InstanceSourceMap = StructPropertyHandle->GetInstanceMetaDataMap();		
-		for (const auto& It : *InstanceSourceMap)
+		if (const TMap<FName, FString>* InstanceSourceMap = StructPropertyHandle->GetInstanceMetaDataMap())
 		{
-			EditProperty->SetInstanceMetaData(*It.Key.ToString(), *It.Value);
+			for (const auto& It : *InstanceSourceMap)
+			{
+				EditProperty->SetInstanceMetaData(*It.Key.ToString(), *It.Value);
+			}
 		}
 
 		if (EditProperty->GetProperty()->IsA<FStructProperty>())
