Fix wrong registration name for LinearColorRamp

#Synced-CL 48963730

diff --git a/Engine/Source/Editor/DetailCustomizations/Private/DetailCustomizations.cpp b/Engine/Source/Editor/DetailCustomizations/Private/DetailCustomizations.cpp
--- a/Engine/Source/Editor/DetailCustomizations/Private/DetailCustomizations.cpp
+++ b/Engine/Source/Editor/DetailCustomizations/Private/DetailCustomizations.cpp
@@ -255,7 +255,7 @@ void FDetailCustomizationsModule::RegisterPropertyTypeCustomizations()
 	RegisterCustomPropertyTypeLayout("DataTableRowHandle", FOnGetPropertyTypeCustomizationInstance::CreateStatic(&FDataTableCustomizationLayout::MakeInstance));
 	RegisterCustomPropertyTypeLayout("DataTableCategoryHandle", FOnGetPropertyTypeCustomizationInstance::CreateStatic(&FDataTableCategoryCustomizationLayout::MakeInstance));
 	RegisterCustomPropertyTypeLayout("CurveTableRowHandle", FOnGetPropertyTypeCustomizationInstance::CreateStatic(&FCurveTableCustomizationLayout::MakeInstance));
-	RegisterCustomPropertyTypeLayout("DataflowColorRamp", FOnGetPropertyTypeCustomizationInstance::CreateStatic(&FLinearColorRampCustomization::MakeInstance));
+	RegisterCustomPropertyTypeLayout("LinearColorRamp", FOnGetPropertyTypeCustomizationInstance::CreateStatic(&FLinearColorRampCustomization::MakeInstance));
 	REGISTER_UIMINMAX_CUSTOMIZATION(NAME_Vector, &FVectorStructCustomization::MakeInstance);
 	REGISTER_UIMINMAX_CUSTOMIZATION(NAME_Vector3f, &FVectorStructCustomization::MakeInstance);
 	REGISTER_UIMINMAX_CUSTOMIZATION(NAME_Vector3d, &FVectorStructCustomization::MakeInstance);
