Fix world item context menu options all bound to OFPE, now each option will show up based on if its enabled or not.

#Synced-CL 48998606

diff --git a/Engine/Source/Editor/DetailCustomizations/Private/DetailCustomizations.cpp b/Engine/Source/Editor/DetailCustomizations/Private/DetailCustomizations.cpp
--- a/Engine/Source/Editor/DetailCustomizations/Private/DetailCustomizations.cpp
+++ b/Engine/Source/Editor/DetailCustomizations/Private/DetailCustomizations.cpp
@@ -218,17 +218,17 @@ void FDetailCustomizationsModule::ShutdownModule()
 	}
 }
 
-void FDetailCustomizationsModule::RegisterDevelopmentStatusWarningSupression(FName ClassName)
+void FDetailCustomizationsModule::RegisterDevelopmentStatusWarningSuppression(FName ClassName)
 {
 	SuppressedDevelopmentStatusWarnings.Add(ClassName);
 }
 
-void FDetailCustomizationsModule::UnregisterDevelopmentStatusWarningSupression(FName ClassName)
+void FDetailCustomizationsModule::UnregisterDevelopmentStatusWarningSuppression(FName ClassName)
 {
 	SuppressedDevelopmentStatusWarnings.Remove(ClassName);
 }
 
-bool FDetailCustomizationsModule::IsDevelopmentStatusWarningSupressed(const UClass* Class) const
+bool FDetailCustomizationsModule::IsDevelopmentStatusWarningSuppressed(const UClass* Class) const
 {
 	const UClass* ClassCursor = Class;
 	while(ClassCursor)
diff --git a/Engine/Source/Editor/DetailCustomizations/Private/ObjectDetails.cpp b/Engine/Source/Editor/DetailCustomizations/Private/ObjectDetails.cpp
--- a/Engine/Source/Editor/DetailCustomizations/Private/ObjectDetails.cpp
+++ b/Engine/Source/Editor/DetailCustomizations/Private/ObjectDetails.cpp
@@ -61,7 +61,7 @@ void FObjectDetails::AddExperimentalWarningCategory(IDetailLayoutBuilder& Detail
 	FObjectEditorUtils::GetClassDevelopmentStatus(DetailBuilder.GetBaseClass(), bBaseClassIsExperimental, bBaseClassIsEarlyAccess, MostDerivedDevelopmentClassName);
 
 	FDetailCustomizationsModule& DetailCustomizationsModule = FModuleManager::Get().GetModuleChecked<FDetailCustomizationsModule>(TEXT("DetailCustomizations"));
-	if (DetailCustomizationsModule.IsDevelopmentStatusWarningSupressed(DetailBuilder.GetBaseClass()))
+	if (DetailCustomizationsModule.IsDevelopmentStatusWarningSuppressed(DetailBuilder.GetBaseClass()))
 	{
 		return;
 	}
diff --git a/Engine/Source/Editor/DetailCustomizations/Public/DetailCustomizations.h b/Engine/Source/Editor/DetailCustomizations/Public/DetailCustomizations.h
--- a/Engine/Source/Editor/DetailCustomizations/Public/DetailCustomizations.h
+++ b/Engine/Source/Editor/DetailCustomizations/Public/DetailCustomizations.h
@@ -34,13 +34,13 @@ public:
 	/**
 	 * Suppresses the DevelopmentStatus warnings for the given class or any derived versions of this class
 	 */
-	void DETAILCUSTOMIZATIONS_API RegisterDevelopmentStatusWarningSupression(FName ClassName);
+	void DETAILCUSTOMIZATIONS_API RegisterDevelopmentStatusWarningSuppression(FName ClassName);
 	/**
 	 * Removes suppression of DevelopmentStatus warnings for the given class or any derived versions of this class
 	 */
-	void DETAILCUSTOMIZATIONS_API UnregisterDevelopmentStatusWarningSupression(FName ClassName);
+	void DETAILCUSTOMIZATIONS_API UnregisterDevelopmentStatusWarningSuppression(FName ClassName);
 
-	bool IsDevelopmentStatusWarningSupressed(const UClass* Class) const;
+	bool IsDevelopmentStatusWarningSuppressed(const UClass* Class) const;
 
 private:
 	void RegisterPropertyTypeCustomizations();
