Attenuation assets now display colored gutters for each expanded category

#Synced-CL 49102229

diff --git a/Engine/Source/Editor/DetailCustomizations/Private/AttenuationSettingsCustomizations.cpp b/Engine/Source/Editor/DetailCustomizations/Private/AttenuationSettingsCustomizations.cpp
--- a/Engine/Source/Editor/DetailCustomizations/Private/AttenuationSettingsCustomizations.cpp
+++ b/Engine/Source/Editor/DetailCustomizations/Private/AttenuationSettingsCustomizations.cpp
@@ -762,18 +762,33 @@ void FSoundAttenuationSettingsCustomization::CustomizeChildren(TSharedRef<IPrope
 			.EditCondition(IsAttenuationOverriddenAttribute(), nullptr);
 	}
 
-	// Set category display names
-	LayoutBuilder.EditCategory("AttenuationDistance", LOCTEXT("AttenuationVolume", "Attenuation (Volume)"));
-	LayoutBuilder.EditCategory("AttenuationSpatialization", LOCTEXT("AttenuationSpatialization", "Attenuation (Spatialization)"));
-	LayoutBuilder.EditCategory("AttenuationOcclusion", LOCTEXT("AttenuationOcclusion", "Attenuation (Occlusion)"));
-	LayoutBuilder.EditCategory("AttenuationSubmixSend", LOCTEXT("AttenuationSubmixSend", "Attenuation (Submix)"));
-	LayoutBuilder.EditCategory("AttenuationReverbSend", LOCTEXT("AttenuationReverbSend", "Attenuation (Reverb)"));
-	LayoutBuilder.EditCategory("AttenuationListenerFocus", LOCTEXT("AttenuationListenerFocus", "Attenuation (Focus)"));
-	LayoutBuilder.EditCategory("AttenuationPriority", LOCTEXT("AttenuationPriority", "Attenuation (Priority)"));
-	LayoutBuilder.EditCategory("AttenuationAirAbsorption", LOCTEXT("AttenuationAirAbsorption", "Attenuation (Air Absorption)"));
-	LayoutBuilder.EditCategory("AttenuationPluginSettings", LOCTEXT("AttenuationPluginSettings", "Attenuation (Plugin Settings)"));
-	LayoutBuilder.EditCategory("AttenuationSourceDataOverride", LOCTEXT("AttenuationSourceDataOverride", "Attenuation (Source Data Override)"));
-	LayoutBuilder.EditCategory("AttenuationAudioLink", LOCTEXT("AttenuationAudioLink", "Attenuation (AudioLink)"));
+	// Paul Tol's Muted qualitative color scheme (colorblind-safe palette; sRGB)
+	// https://sronpersonalpages.nl/~pault/data/colourschemes.pdf
+	constexpr FLinearColor CategoryColors[] = 
+	{
+		FLinearColor(FColor(51,   34, 136)), // #332288 indigo
+		FLinearColor(FColor(221, 204, 119)), // #DDCC77 sand
+		FLinearColor(FColor(17,  119, 51)),  // #117733 green
+		FLinearColor(FColor(170,  68, 153)), // #AA4499 purple
+		FLinearColor(FColor(136, 204, 238)), // #88CCEE cyan
+		FLinearColor(FColor(136,  34, 85)),  // #882255 wine
+		FLinearColor(FColor(153, 153, 51)),  // #999933 olive
+		FLinearColor(FColor(204, 102, 119)), // #CC6677 rose
+	};
+	constexpr FLinearColor DefaultColor(FColor(221, 221, 221)); // #DDDDDD pale grey
+	
+	// Set category display names & colors
+	LayoutBuilder.EditCategory("AttenuationDistance", LOCTEXT("AttenuationVolume", "Attenuation (Volume)")).SetColor(CategoryColors[0]);
+	LayoutBuilder.EditCategory("AttenuationSpatialization", LOCTEXT("AttenuationSpatialization", "Attenuation (Spatialization)")).SetColor(CategoryColors[1]);
+	LayoutBuilder.EditCategory("AttenuationAirAbsorption", LOCTEXT("AttenuationAirAbsorption", "Attenuation (Air Absorption)")).SetColor(CategoryColors[2]);
+	LayoutBuilder.EditCategory("AttenuationListenerFocus", LOCTEXT("AttenuationListenerFocus", "Attenuation (Focus)")).SetColor(CategoryColors[3]);
+	LayoutBuilder.EditCategory("AttenuationOcclusion", LOCTEXT("AttenuationOcclusion", "Attenuation (Occlusion)")).SetColor(CategoryColors[4]);
+	LayoutBuilder.EditCategory("AttenuationSubmixSend", LOCTEXT("AttenuationSubmixSend", "Attenuation (Submix)")).SetColor(CategoryColors[5]);
+	LayoutBuilder.EditCategory("AttenuationReverbSend", LOCTEXT("AttenuationReverbSend", "Attenuation (Reverb)")).SetColor(CategoryColors[6]);
+	LayoutBuilder.EditCategory("AttenuationPriority", LOCTEXT("AttenuationPriority", "Attenuation (Priority)")).SetColor(CategoryColors[7]);
+	LayoutBuilder.EditCategory("AttenuationPluginSettings", LOCTEXT("AttenuationPluginSettings", "Attenuation (Plugin Settings)")).SetColor(DefaultColor);
+	LayoutBuilder.EditCategory("AttenuationSourceDataOverride", LOCTEXT("AttenuationSourceDataOverride", "Attenuation (Source Data Override)")).SetColor(DefaultColor);
+	LayoutBuilder.EditCategory("AttenuationAudioLink", LOCTEXT("AttenuationAudioLink", "Attenuation (AudioLink)")).SetColor(DefaultColor);
 	LayoutBuilder.SortCategories(AttenuationSettingsUtils::SortCategories);
 
 	if (PropertyHandles.Num() != 70)
