Check the pointer valid before accessing it.

#rnx
#tests In PIE, create an Editor Utility Widget and add "Editable TextBlock (Multiline)" into it. In the text field of "Editable TextBlock (Multiline)", type some string and keep it in focused. Then press "ALT+P" to play the session.
#Synced-CL 49409147

diff --git a/Engine/Source/Editor/DetailCustomizations/Private/TextCustomization.cpp b/Engine/Source/Editor/DetailCustomizations/Private/TextCustomization.cpp
--- a/Engine/Source/Editor/DetailCustomizations/Private/TextCustomization.cpp
+++ b/Engine/Source/Editor/DetailCustomizations/Private/TextCustomization.cpp
@@ -114,7 +114,7 @@ namespace
 					TArray<UObject*> PropertyObjects;
 					PropertyHandle->GetOuterObjects(PropertyObjects);
 
-					if (PropertyObjects.IsValidIndex(InIndex))
+					if (PropertyObjects.IsValidIndex(InIndex) && PropertyObjects[InIndex])
 					{
 						PropertyPackage = PropertyObjects[InIndex]->GetPackage();
 						bForceUniqueId = PropertyHandle->HasMetaData(PropertyObjects[InIndex]->HasAnyFlags(RF_ClassDefaultObject | RF_ArchetypeObject) ? NAME_UniqueIdWhenEditedOnDefault : NAME_UniqueIdWhenEditedOnInstance);
