Hide actor's ExternalDataLayerAsset property when its world is not partitioned.

#Synced-CL 49593614

diff --git a/Engine/Source/Editor/DetailCustomizations/Private/ActorDetails.cpp b/Engine/Source/Editor/DetailCustomizations/Private/ActorDetails.cpp
--- a/Engine/Source/Editor/DetailCustomizations/Private/ActorDetails.cpp
+++ b/Engine/Source/Editor/DetailCustomizations/Private/ActorDetails.cpp
@@ -158,6 +158,7 @@ void FActorDetails::CustomizeDetails( IDetailLayoutBuilder& DetailLayout )
 			DetailLayout.HideProperty(DetailLayout.GetProperty(AActor::GetRuntimeGridPropertyName(), AActor::StaticClass()));
 			DetailLayout.HideProperty(DetailLayout.GetProperty(AActor::GetDataLayerAssetsPropertyName(), AActor::StaticClass()));
 			DetailLayout.HideProperty(DetailLayout.GetProperty(AActor::GetDataLayerPropertyName(), AActor::StaticClass()));
+			DetailLayout.HideProperty(DetailLayout.GetProperty(AActor::GetExternalDataLayerAssetPropertyName(), AActor::StaticClass()));
 			DetailLayout.HideProperty(DetailLayout.GetProperty(AActor::GetHLODLayerPropertyName(), AActor::StaticClass()));
 		}
 
diff --git a/Engine/Source/Runtime/Engine/Classes/GameFramework/Actor.h b/Engine/Source/Runtime/Engine/Classes/GameFramework/Actor.h
--- a/Engine/Source/Runtime/Engine/Classes/GameFramework/Actor.h
+++ b/Engine/Source/Runtime/Engine/Classes/GameFramework/Actor.h
@@ -1482,6 +1482,7 @@ public:
 	ENGINE_API bool IsPropertyChangedAffectingDataLayers(FPropertyChangedEvent& PropertyChangedEvent) const;
 	ENGINE_API void FixupDataLayers(bool bRevertChangesOnLockedDataLayer = false);
 	static const FName GetDataLayerAssetsPropertyName() { return GET_MEMBER_NAME_CHECKED(AActor, DataLayerAssets); }
+	static const FName GetExternalDataLayerAssetPropertyName() { return GET_MEMBER_NAME_CHECKED(AActor, ExternalDataLayerAsset); }
 	static const FName GetDataLayerPropertyName() { return GET_MEMBER_NAME_CHECKED(AActor, DataLayers); }
 
 	ENGINE_API TArray<const UDataLayerAsset*> GetDataLayerAssets(bool bIncludeExternalDataLayerAsset = true) const;
