[Crash Fix]: Update unsigned iteration loop indexing leading to crash when removing all elements (AKA PrimIdx == 0)

#Synced-CL 50448929

#Synced-CL 50448929

diff --git a/Engine/Source/Editor/DetailCustomizations/Private/CustomPrimitiveDataCustomization.cpp b/Engine/Source/Editor/DetailCustomizations/Private/CustomPrimitiveDataCustomization.cpp
--- a/Engine/Source/Editor/DetailCustomizations/Private/CustomPrimitiveDataCustomization.cpp
+++ b/Engine/Source/Editor/DetailCustomizations/Private/CustomPrimitiveDataCustomization.cpp
@@ -716,7 +716,7 @@ void FCustomPrimitiveDataCustomization::OnRemovedPrimitiveData(uint8 PrimIdx)
 	{
 		GEditor->BeginTransaction(LOCTEXT("OnRemovedPrimitiveData", "Removed Items"));
 
-		for (uint32 i = NumElements - 1; i >= static_cast<uint32>(PrimIdx); --i)
+		for (uint32 i = NumElements; i-- > static_cast<uint32>(PrimIdx); )
 		{
 			DataArrayHandle->DeleteItem(i);
 		}
