CallInEditor: Don't show buttons when no valid execution context

Fix regression from CL 35103367 where CallInEditor buttons would appear in the Blueprint Editor. The previous fix filters out archetype objects but still displays UI widgets even when no valid objects remain.

#Synced-CL 50591564

diff --git a/Engine/Source/Editor/DetailCustomizations/Private/ObjectDetails.cpp b/Engine/Source/Editor/DetailCustomizations/Private/ObjectDetails.cpp
--- a/Engine/Source/Editor/DetailCustomizations/Private/ObjectDetails.cpp
+++ b/Engine/Source/Editor/DetailCustomizations/Private/ObjectDetails.cpp
@@ -95,6 +95,12 @@ void FObjectDetails::AddCallInEditorMethods(IDetailLayoutBuilder& DetailBuilder)
 	DetailBuilder.GetObjectsBeingCustomized(/*out*/ SelectedObjectsList);
 	SelectedObjectsList.RemoveAllSwap([](TWeakObjectPtr<UObject> ObjPtr) { UObject* Obj = ObjPtr.Get(); return (Obj == nullptr) || Obj->HasAnyFlags(RF_ArchetypeObject); });
 
+	// Don't show CallInEditor buttons if there are no valid execution contexts (e.g. in Blueprint Editor)
+	if (SelectedObjectsList.IsEmpty())
+	{
+		return;
+	}
+
 	TArray<UFunction*> CallInEditorFunctions;
 	PropertyCustomizationHelpers::GetCallInEditorFunctionsForClass(
 		DetailBuilder.GetBaseClass(),
