Fix crash in TOverrideResetToDefaultWithStaticUStruct for nested properties
* IsResetToDefaultVisible() and OnResetToDefault() used Property->ContainerPtrToValuePtr(&DefaultObject) to compute the default value address. This is correct for direct children of UStructType, but these callbacks can also fire for deeply nested properties, in which case an invalid pointer was returned.
* The fix adds ResolveDefaultValuePtr() which walks the property handle parent chain to figure out the correct offset from DefaultObject.

#Synced-CL 50781450

diff --git a/Engine/Source/Editor/DetailCustomizations/Private/OverrideResetToDefault.h b/Engine/Source/Editor/DetailCustomizations/Private/OverrideResetToDefault.h
--- a/Engine/Source/Editor/DetailCustomizations/Private/OverrideResetToDefault.h
+++ b/Engine/Source/Editor/DetailCustomizations/Private/OverrideResetToDefault.h
@@ -18,12 +18,12 @@
  * reset archetype instances as well.
  */
 template<typename UStructType>
-class TOverrideResetToDefaultWithStaticUStruct 
+class TOverrideResetToDefaultWithStaticUStruct
 {
 public:
 	/** Called by the UI to show/hide the reset widgets */
 	static bool IsResetToDefaultVisible(TSharedPtr<IPropertyHandle> InPropertyHandle);
-	
+
 	/** Reset to default triggered in UI */
 	static void OnResetToDefault(TSharedPtr<IPropertyHandle> InPropertyHandle);
 
@@ -31,12 +31,50 @@ protected:
 	/** Adds callbacks used by the UI to determine if a property can be reset and reset it. */
 	static void AddResetToDefaultOverrides(IDetailPropertyRow& InDetailPropertyRow);
 
+private:
+	/**
+	 * Resolves the correct default value pointer for a property that may be nested within UStructType.
+	 * Walks the property handle parent chain to figure out the correct offset from DefaultObject.
+	 */
+	static const void* ResolveDefaultValuePtr(TSharedPtr<IPropertyHandle> InPropertyHandle);
+
 	static const UStructType DefaultObject;
 };
 
 template<typename UStructType>
 const UStructType TOverrideResetToDefaultWithStaticUStruct<UStructType>::DefaultObject;
 
+template<typename UStructType>
+const void* TOverrideResetToDefaultWithStaticUStruct<UStructType>::ResolveDefaultValuePtr(TSharedPtr<IPropertyHandle> InPropertyHandle)
+{
+	FProperty* Property = InPropertyHandle->GetProperty();
+	if (!Property)
+	{
+		return nullptr;
+	}
+
+	// If this property is a direct member of UStructType, use DefaultObject as the container
+	if (Property->GetOwnerStruct() == UStructType::StaticStruct())
+	{
+		return Property->ContainerPtrToValuePtr<void>(&DefaultObject);
+	}
+
+	// Property is nested deeper within UStructType - resolve the parent's default value to use as container
+	TSharedPtr<IPropertyHandle> ParentHandle = InPropertyHandle->GetParentHandle();
+	if (!ParentHandle.IsValid())
+	{
+		return nullptr;
+	}
+
+	const void* ParentDefaultValue = ResolveDefaultValuePtr(ParentHandle);
+	if (!ParentDefaultValue)
+	{
+		return nullptr;
+	}
+
+	return Property->ContainerPtrToValuePtr<void>(ParentDefaultValue);
+}
+
 template<typename UStructType>
 bool TOverrideResetToDefaultWithStaticUStruct<UStructType>::IsResetToDefaultVisible(TSharedPtr<IPropertyHandle> InPropertyHandle)
 {
@@ -46,7 +84,7 @@ bool TOverrideResetToDefaultWithStaticUStruct<UStructType>::IsResetToDefaultVisi
 
 	check(Property != nullptr);
 
-	DefaultValuePtr = Property->ContainerPtrToValuePtr<void>(&DefaultObject);
+	DefaultValuePtr = ResolveDefaultValuePtr(InPropertyHandle);
 	InPropertyHandle->GetValueData(ValuePtr);
 
 	if ((DefaultValuePtr != nullptr) && (ValuePtr != nullptr))
@@ -63,16 +101,16 @@ void TOverrideResetToDefaultWithStaticUStruct<UStructType>::OnResetToDefault(TSh
 	FProperty*	Property = InPropertyHandle->GetProperty();
 	const void* DefaultValuePtr = nullptr;
 	void*		ValuePtr = nullptr;
-	
+
 	check(Property != nullptr);
 
-	DefaultValuePtr = Property->ContainerPtrToValuePtr<void>(&DefaultObject);
+	DefaultValuePtr = ResolveDefaultValuePtr(InPropertyHandle);
 	InPropertyHandle->GetValueData(ValuePtr);
 
 	if ((DefaultValuePtr != nullptr) && (ValuePtr != nullptr))
 	{
-		// SetPropertyValue_DirectSingle when called on an archetype can propagate to instances. 
-		// See class comments: this ResetToDefault override was not written for contexts involving 
+		// SetPropertyValue_DirectSingle when called on an archetype can propagate to instances.
+		// See class comments: this ResetToDefault override was not written for contexts involving
 		// inheritance, so pass in empty archetype instances list.
 		TArray<void*> NoArchetypeInsts;
 		PropertyAccessUtil::SetPropertyValue_DirectSingle(Property, DefaultValuePtr, Property, ValuePtr, NoArchetypeInsts, 0, Property->HasAnyFlags(RF_ArchetypeObject), []() { return nullptr; });
