TOverrideResetToDefaultWithStaticUStruct::OnResetToDefault() - Removed usage of EObject flags on FField, soon to be fully deprecated

* Replaced the call to PropertyAccessUtil::SetPropertyValue_DirectSingle() in OnResetToDefault() with a direct Property->CopySingleValue() call. The SetPropertyValue_DirectSingle() call was passing an empty archetype instances array, no read-only flags, and a null notification function ? every code path except the final CopySingleValue() was dead. This allows the removal of the Property->HasAnyFlags(RF_ArchetypeObject) call.

#Synced-CL 50811210

diff --git a/Engine/Source/Editor/DetailCustomizations/Private/OverrideResetToDefault.h b/Engine/Source/Editor/DetailCustomizations/Private/OverrideResetToDefault.h
--- a/Engine/Source/Editor/DetailCustomizations/Private/OverrideResetToDefault.h
+++ b/Engine/Source/Editor/DetailCustomizations/Private/OverrideResetToDefault.h
@@ -7,7 +7,7 @@
 #include "Layout/Visibility.h"
 #include "IPropertyTypeCustomization.h"
 #include "PropertyHandle.h"
-#include "UObject/PropertyAccessUtil.h"
+#include "UObject/UnrealType.h"
 
 /**
  * Class template aiming to override ResetToDefault methods by comparing properties with an internal static DefaultObject.
@@ -109,11 +109,7 @@ void TOverrideResetToDefaultWithStaticUStruct<UStructType>::OnResetToDefault(TSh
 
 	if ((DefaultValuePtr != nullptr) && (ValuePtr != nullptr))
 	{
-		// SetPropertyValue_DirectSingle when called on an archetype can propagate to instances.
-		// See class comments: this ResetToDefault override was not written for contexts involving
-		// inheritance, so pass in empty archetype instances list.
-		TArray<void*> NoArchetypeInsts;
-		PropertyAccessUtil::SetPropertyValue_DirectSingle(Property, DefaultValuePtr, Property, ValuePtr, NoArchetypeInsts, 0, Property->HasAnyFlags(RF_ArchetypeObject), []() { return nullptr; });
+		Property->CopySingleValue(ValuePtr, DefaultValuePtr);
 	}
 }
 
