Component reference picker improvements

Allow picking from the actor's components in the blueprint editor

#Synced-CL 50814747

diff --git a/Engine/Source/Editor/DetailCustomizations/Private/ComponentReferenceCustomization.cpp b/Engine/Source/Editor/DetailCustomizations/Private/ComponentReferenceCustomization.cpp
--- a/Engine/Source/Editor/DetailCustomizations/Private/ComponentReferenceCustomization.cpp
+++ b/Engine/Source/Editor/DetailCustomizations/Private/ComponentReferenceCustomization.cpp
@@ -3,6 +3,8 @@
 #include "ComponentReferenceCustomization.h"
 
 #include "ActorPickerMode.h"
+#include "BlueprintEditor.h"
+#include "BlueprintEditorModule.h"
 #include "Brushes/SlateNoResource.h"
 #include "Components/ActorComponent.h"
 #include "Components/SceneComponent.h"
@@ -261,8 +263,8 @@ void FComponentReferenceCustomization::BuildComboBox()
 			SNew(SHorizontalBox)
 			+ SHorizontalBox::Slot()
 			.AutoWidth()
-			.HAlign(HAlign_Left)
 			.VAlign(VAlign_Center)
+			.Padding(0.0f, 0.0f, 4.0f, 0.0f)
 			[
 				SNew(SImage)
 				.Image(this, &FComponentReferenceCustomization::GetActorIcon)
@@ -284,8 +286,8 @@ void FComponentReferenceCustomization::BuildComboBox()
 		SNew(SHorizontalBox)
 		+ SHorizontalBox::Slot()
 		.AutoWidth()
-		.HAlign(HAlign_Left)
 		.VAlign(VAlign_Center)
+		.Padding(0.0f, 0.0f, 4.0f, 0.0f)
 		[
 			SNew(SImage)
 			.Image(this, &FComponentReferenceCustomization::GetComponentIcon)
@@ -304,7 +306,6 @@ void FComponentReferenceCustomization::BuildComboBox()
 	TSharedRef<SHorizontalBox> ComboButtonContent = SNew(SHorizontalBox)
 		+ SHorizontalBox::Slot()
 		.AutoWidth()
-		.HAlign(HAlign_Left)
 		.VAlign(VAlign_Center)
 		[
 			SNew(SImage)
@@ -319,12 +320,10 @@ void FComponentReferenceCustomization::BuildComboBox()
 
 	ComponentComboButton = SNew(SComboButton)
 		.ToolTipText(TooltipAttribute)
-		.ButtonStyle(FAppStyle::Get(), "PropertyEditor.AssetComboStyle")
-		.ForegroundColor(FAppStyle::GetColor("PropertyEditor.AssetName.ColorAndOpacity"))
 		.OnGetMenuContent(this, &FComponentReferenceCustomization::OnGetMenuContent)
 		.OnMenuOpenChanged(this, &FComponentReferenceCustomization::OnMenuOpenChanged)
 		.IsEnabled(IsEnabledAttribute)
-		.ContentPadding(2.0f)
+		.ContentPadding(0.0f)
 		.ButtonContent()
 		[
 			SNew(SWidgetSwitcher)
@@ -357,6 +356,14 @@ AActor* FComponentReferenceCustomization::GetFirstOuterActor() const
 			}
 			if (UActorComponent* Component = Cast<UActorComponent>(Obj))
 			{
+				if (Component->HasAnyFlags(RF_ArchetypeObject))
+				{
+					if (UClass* OuterClass = Cast<UClass>(Component->GetOuter()))
+					{
+						return Cast<AActor>(OuterClass->GetDefaultObject());
+					}
+				}
+
 				if (Component->GetOwner())
 				{
 					return Component->GetOwner();
@@ -464,7 +471,7 @@ bool FComponentReferenceCustomization::IsComponentReferenceValid(const FComponen
 		return false;
 	}
 
-	AActor* CachedActor = CachedFirstOuterActor.Get();
+	AActor* CachedActor = CachedFilterActor.Get();
 	if (UActorComponent* NewComponent = Value.GetComponent(CachedActor))
 	{
 		if (!IsFilteredComponent(NewComponent))
@@ -519,6 +526,7 @@ void FComponentReferenceCustomization::UpdateCachedValues(bool bResetValueIfInva
 {
 	CachedComponent.Reset();
 	CachedFirstOuterActor = GetFirstOuterActor();
+	CachedFilterActor = CachedFirstOuterActor;
 
 	FComponentReference TmpComponentReference;
 	CachedPropertyAccess = GetValue(TmpComponentReference);
@@ -639,12 +647,43 @@ TSharedRef<SWidget> FComponentReferenceCustomization::OnGetMenuContent()
 {
 	UActorComponent* InitialComponent = CachedComponent.Get();
 
+	// If the actor we are looking at is a BP, use its BP editor preview world, if any
+	CachedFilterActor = CachedFirstOuterActor;
+	bool bShowTransient = false;
+	UWorld* PickerWorld = nullptr;
+	if (AActor* Actor = CachedFirstOuterActor.Get())
+	{
+		if (Actor->HasAnyFlags(RF_ArchetypeObject))
+		{
+			if (UBlueprint* Blueprint = Cast<UBlueprint>(Actor->GetClass()->ClassGeneratedBy))
+			{
+				FBlueprintEditorModule& BlueprintEditorModule = FModuleManager::LoadModuleChecked<FBlueprintEditorModule>("Kismet");
+				for (TSharedRef<IBlueprintEditor>& EditorInterface : BlueprintEditorModule.GetBlueprintEditors())
+				{
+					TSharedRef<FBlueprintEditor> Editor = StaticCastSharedRef<FBlueprintEditor>(EditorInterface);
+					if (Editor->GetBlueprintObj() == Blueprint)
+					{
+						if(AActor* PreviewActor = Editor->GetPreviewActor())
+						{
+							CachedFilterActor = PreviewActor;
+							PickerWorld = PreviewActor->GetWorld();
+							bShowTransient = true;
+						}
+						break;
+					}
+				}
+			}
+		}
+	}
+
 	return PropertyCustomizationHelpers::MakeComponentPickerWithMenu(InitialComponent
 		, bAllowClear
 		, FOnShouldFilterActor::CreateSP(this, &FComponentReferenceCustomization::IsFilteredActor)
 		, FOnShouldFilterComponent::CreateSP(this, &FComponentReferenceCustomization::IsFilteredComponent)
 		, FOnComponentSelected::CreateSP(this, &FComponentReferenceCustomization::OnComponentSelected)
-		, FSimpleDelegate::CreateSP(this, &FComponentReferenceCustomization::CloseComboButton));
+		, FSimpleDelegate::CreateSP(this, &FComponentReferenceCustomization::CloseComboButton)
+		, PickerWorld
+		, bShowTransient);
 }
 
 void FComponentReferenceCustomization::OnMenuOpenChanged(bool bOpen)
@@ -657,18 +696,18 @@ void FComponentReferenceCustomization::OnMenuOpenChanged(bool bOpen)
 
 bool FComponentReferenceCustomization::IsFilteredActor(const AActor* const Actor) const
 {
-	return bAllowAnyActor || Actor == CachedFirstOuterActor.Get();
+	return bAllowAnyActor || Actor == CachedFilterActor.Get();
 }
 
 bool FComponentReferenceCustomization::IsFilteredComponent(const UActorComponent* const Component) const
 {
 	const USceneComponent* SceneComp = Cast<USceneComponent>(Component);
 	const USceneComponent* ParentSceneComp = SceneComp != nullptr ? SceneComp->GetAttachParent() : nullptr;
-	const AActor* OuterActor = CachedFirstOuterActor.Get();
+	const AActor* FilterActor = CachedFilterActor.Get();
 
 	return Component->GetOwner()
-		&& (bAllowAnyActor || Component->GetOwner() == CachedFirstOuterActor.Get())
-		&& (!bAllowAnyActor || (OuterActor && Component->GetOwner()->GetLevel() == OuterActor->GetLevel()))
+		&& (bAllowAnyActor || Component->GetOwner() == FilterActor)
+		&& (!bAllowAnyActor || (FilterActor && Component->GetOwner()->GetLevel() == FilterActor->GetLevel()))
 		&& FComponentEditorUtils::CanEditComponentInstance(Component, SceneComp, false)
 		&& IsFilteredObject(Component, AllowedComponentClassFilters, DisallowedComponentClassFilters)
 		&& IsFilteredObject(Component->GetOwner(), AllowedActorClassFilters, DisallowedActorClassFilters);
@@ -713,7 +752,7 @@ void FComponentReferenceCustomization::OnComponentSelected(const UActorComponent
 {
 	ComponentComboButton->SetIsOpen(false);
 
-	FComponentReference ComponentReference = FComponentEditorUtils::MakeComponentReference(CachedFirstOuterActor.Get(), InComponent);
+	FComponentReference ComponentReference = FComponentEditorUtils::MakeComponentReference(CachedFilterActor.Get(), InComponent);
 	SetValue(ComponentReference);
 }
 
diff --git a/Engine/Source/Editor/DetailCustomizations/Private/ComponentReferenceCustomization.h b/Engine/Source/Editor/DetailCustomizations/Private/ComponentReferenceCustomization.h
--- a/Engine/Source/Editor/DetailCustomizations/Private/ComponentReferenceCustomization.h
+++ b/Engine/Source/Editor/DetailCustomizations/Private/ComponentReferenceCustomization.h
@@ -142,6 +142,7 @@ private:
 
 	/** Cached ComponentReference */
 	TWeakObjectPtr<AActor> CachedFirstOuterActor;
+	TWeakObjectPtr<AActor> CachedFilterActor;
 	TWeakObjectPtr<UActorComponent> CachedComponent;
 	FPropertyAccess::Result CachedPropertyAccess;
 };
diff --git a/Engine/Source/Editor/PropertyEditor/Private/PropertyCustomizationHelpers.cpp b/Engine/Source/Editor/PropertyEditor/Private/PropertyCustomizationHelpers.cpp
--- a/Engine/Source/Editor/PropertyEditor/Private/PropertyCustomizationHelpers.cpp
+++ b/Engine/Source/Editor/PropertyEditor/Private/PropertyCustomizationHelpers.cpp
@@ -746,16 +746,18 @@ namespace PropertyCustomizationHelpers
 			.ShowTransient(bShowTransient);
 	}
 
-	TSharedRef<SWidget> MakeComponentPickerWithMenu( UActorComponent* const InitialComponent, const bool AllowClear, FOnShouldFilterActor ActorFilter, FOnShouldFilterComponent ComponentFilter, FOnComponentSelected OnSet, FSimpleDelegate OnClose )
+	TSharedRef<SWidget> MakeComponentPickerWithMenu( UActorComponent* const InitialComponent, const bool AllowClear, FOnShouldFilterActor ActorFilter, FOnShouldFilterComponent ComponentFilter, FOnComponentSelected OnSet, FSimpleDelegate OnClose, UWorld* PickerWorld, bool bShowTransient)
 	{
 		return
 			SNew(SPropertyMenuComponentPicker)
 			.InitialComponent(InitialComponent)
 			.AllowClear(AllowClear)
+			.ShowTransient(bShowTransient)
 			.ActorFilter(ActorFilter)
 			.ComponentFilter(ComponentFilter)
 			.OnSet(OnSet)
-			.OnClose(OnClose);
+			.OnClose(OnClose)
+			.PickerWorld(PickerWorld);
 	}
 
 	TSharedRef<SWidget> MakeInteractiveActorPicker( FOnGetAllowedClasses OnGetAllowedClasses, FOnShouldFilterActor OnShouldFilterActor, FOnActorSelected OnActorSelectedFromPicker)
diff --git a/Engine/Source/Editor/PropertyEditor/Private/UserInterface/PropertyEditor/SPropertyMenuComponentPicker.cpp b/Engine/Source/Editor/PropertyEditor/Private/UserInterface/PropertyEditor/SPropertyMenuComponentPicker.cpp
--- a/Engine/Source/Editor/PropertyEditor/Private/UserInterface/PropertyEditor/SPropertyMenuComponentPicker.cpp
+++ b/Engine/Source/Editor/PropertyEditor/Private/UserInterface/PropertyEditor/SPropertyMenuComponentPicker.cpp
@@ -82,6 +82,7 @@ void SPropertyMenuComponentPicker::Construct(const FArguments& InArgs)
 
 		FSceneOutlinerInitializationOptions InitOptions;
 		InitOptions.bFocusSearchBoxWhenOpened = true;
+		InitOptions.bShowTransient = InArgs._ShowTransient;
 
 		struct FPickerFilter : public FSceneOutlinerFilter
 		{
@@ -132,7 +133,7 @@ void SPropertyMenuComponentPicker::Construct(const FArguments& InArgs)
 				SNew(SBorder)
 				.BorderImage(FAppStyle::GetBrush("Menu.Background"))
 				[
-					SceneOutlinerModule.CreateComponentPicker(InitOptions, FOnComponentPicked::CreateSP(this, &SPropertyMenuComponentPicker::OnItemSelected))
+					SceneOutlinerModule.CreateComponentPicker(InitOptions, FOnComponentPicked::CreateSP(this, &SPropertyMenuComponentPicker::OnItemSelected), InArgs._PickerWorld)
 				]
 			];
 
diff --git a/Engine/Source/Editor/PropertyEditor/Private/UserInterface/PropertyEditor/SPropertyMenuComponentPicker.h b/Engine/Source/Editor/PropertyEditor/Private/UserInterface/PropertyEditor/SPropertyMenuComponentPicker.h
--- a/Engine/Source/Editor/PropertyEditor/Private/UserInterface/PropertyEditor/SPropertyMenuComponentPicker.h
+++ b/Engine/Source/Editor/PropertyEditor/Private/UserInterface/PropertyEditor/SPropertyMenuComponentPicker.h
@@ -20,14 +20,17 @@ public:
 	SLATE_BEGIN_ARGS(SPropertyMenuComponentPicker)
 		: _InitialComponent(nullptr)
 		, _AllowClear(true)
+		, _ShowTransient(false)
 		, _ActorFilter()
 	{}
 		SLATE_ARGUMENT(UActorComponent*, InitialComponent)
 		SLATE_ARGUMENT(bool, AllowClear)
+		SLATE_ARGUMENT(bool, ShowTransient)
 		SLATE_ARGUMENT(FOnShouldFilterActor, ActorFilter)
 		SLATE_ARGUMENT(FOnShouldFilterComponent, ComponentFilter)
 		SLATE_EVENT(FOnComponentSelected, OnSet)
 		SLATE_EVENT(FSimpleDelegate, OnClose)
+		SLATE_ARGUMENT(UWorld*, PickerWorld)
 	SLATE_END_ARGS()
 
 	/**
diff --git a/Engine/Source/Editor/PropertyEditor/Public/PropertyCustomizationHelpers.h b/Engine/Source/Editor/PropertyEditor/Public/PropertyCustomizationHelpers.h
--- a/Engine/Source/Editor/PropertyEditor/Public/PropertyCustomizationHelpers.h
+++ b/Engine/Source/Editor/PropertyEditor/Public/PropertyCustomizationHelpers.h
@@ -174,7 +174,7 @@ namespace PropertyCustomizationHelpers
 	PROPERTYEDITOR_API TSharedRef<SWidget> MakeActorPickerAnchorButton( FOnGetActorFilters OnGetActorFilters, FOnActorSelected OnActorSelectedFromPicker );
 	PROPERTYEDITOR_API TSharedRef<SWidget> MakeActorPickerWithMenu( AActor* const InitialActor, const bool AllowClear, FOnShouldFilterActor ActorFilter, FOnActorSelected OnSet, FSimpleDelegate OnClose, FSimpleDelegate OnUseSelected );
 	PROPERTYEDITOR_API TSharedRef<SWidget> MakeActorPickerWithMenu(AActor* const InitialActor, const bool AllowClear, const bool AllowPickingLevelInstanceContent, FOnShouldFilterActor ActorFilter, FOnActorSelected OnSet, FSimpleDelegate OnClose, FSimpleDelegate OnUseSelected, bool bDisplayUseSelected=true, bool bShowTransient=false);
-	PROPERTYEDITOR_API TSharedRef<SWidget> MakeComponentPickerWithMenu( UActorComponent* const InitialComponent, const bool AllowClear, FOnShouldFilterActor ActorFilter, FOnShouldFilterComponent ComponentFilter, FOnComponentSelected OnSet, FSimpleDelegate OnClose );
+	PROPERTYEDITOR_API TSharedRef<SWidget> MakeComponentPickerWithMenu( UActorComponent* const InitialComponent, const bool AllowClear, FOnShouldFilterActor ActorFilter, FOnShouldFilterComponent ComponentFilter, FOnComponentSelected OnSet, FSimpleDelegate OnClose, UWorld* PickerWorld = nullptr, bool bShowTransient = false);
 	PROPERTYEDITOR_API TSharedRef<SWidget> MakeInteractiveActorPicker(FOnGetAllowedClasses OnGetAllowedClasses, FOnShouldFilterActor OnShouldFilterActor, FOnActorSelected OnActorSelectedFromPicker);
 	PROPERTYEDITOR_API TSharedRef<SWidget> MakeSceneDepthPicker(FOnSceneDepthLocationSelected OnSceneDepthLocationSelected);
 	PROPERTYEDITOR_API TSharedRef<SWidget> MakeEditConfigHierarchyButton(FSimpleDelegate OnEditConfigClicked, TAttribute<FText> OptionalToolTipText = FText(), TAttribute<bool> IsEnabled = true);
