SplineComponentDetails: Fix point index bounds check

#Synced-CL 51054012

diff --git a/Engine/Source/Editor/DetailCustomizations/Private/SplineComponentDetails.cpp b/Engine/Source/Editor/DetailCustomizations/Private/SplineComponentDetails.cpp
--- a/Engine/Source/Editor/DetailCustomizations/Private/SplineComponentDetails.cpp
+++ b/Engine/Source/Editor/DetailCustomizations/Private/SplineComponentDetails.cpp
@@ -923,7 +923,7 @@ void FSplinePointDetails::UpdateValues()
 		bool bValidIndices = true;
 		for (int32 Index : SelectedKeys)
 		{
-			if (Index < 0 || Index > SplineComp->GetNumberOfSplinePoints())
+			if (Index < 0 || Index >= SplineComp->GetNumberOfSplinePoints())
 			{
 				bValidIndices = false;
 				if (!bAlreadyWarnedInvalidIndex)
