Cameras: detail customizations for CineCameraComponent presets were not correctly checking IsEditable on the properties, and missing applying some new properties

#Synced-CL 51024716

diff --git a/Engine/Source/Editor/DetailCustomizations/Private/CameraCropSettingsCustomization.cpp b/Engine/Source/Editor/DetailCustomizations/Private/CameraCropSettingsCustomization.cpp
--- a/Engine/Source/Editor/DetailCustomizations/Private/CameraCropSettingsCustomization.cpp
+++ b/Engine/Source/Editor/DetailCustomizations/Private/CameraCropSettingsCustomization.cpp
@@ -52,7 +52,7 @@ void FCameraCropSettingsCustomization::CustomizeHeader(TSharedRef<IPropertyHandl
 				.OptionsSource(&PresetComboList)
 				.OnGenerateWidget(this, &FCameraCropSettingsCustomization::MakePresetComboWidget)
 				.OnSelectionChanged(this, &FCameraCropSettingsCustomization::OnPresetChanged)
-				.IsEnabled(FSlateApplication::Get().GetNormalExecutionAttribute())
+				.IsEnabled(this, &FCameraCropSettingsCustomization::IsPresetEnabled)
 				.ContentPadding(2.f)
 				.Content()
 				[
@@ -97,6 +97,17 @@ TSharedRef<SWidget> FCameraCropSettingsCustomization::MakePresetComboWidget(TSha
 		.Font(IDetailLayoutBuilder::GetDetailFont());
 }
 
+bool FCameraCropSettingsCustomization::IsPresetEnabled() const
+{
+	bool bEnabled = false;
+	if (AspectRatioHandle.IsValid())
+	{
+		bEnabled = (AspectRatioHandle->IsEditable() &&
+				FSlateApplication::Get().GetNormalExecutionAttribute().Get());
+	}
+	return bEnabled;
+}
+
 void FCameraCropSettingsCustomization::OnPresetChanged(TSharedPtr<FString> NewSelection, ESelectInfo::Type SelectInfo)
 {
 	// if it's set from code, we did that on purpose
diff --git a/Engine/Source/Editor/DetailCustomizations/Private/CameraCropSettingsCustomization.h b/Engine/Source/Editor/DetailCustomizations/Private/CameraCropSettingsCustomization.h
--- a/Engine/Source/Editor/DetailCustomizations/Private/CameraCropSettingsCustomization.h
+++ b/Engine/Source/Editor/DetailCustomizations/Private/CameraCropSettingsCustomization.h
@@ -28,6 +28,7 @@ protected:
 	TSharedPtr<class SComboBox< TSharedPtr<FString> > > PresetComboBox;
 	TArray< TSharedPtr< FString > >						PresetComboList;
 
+	bool IsPresetEnabled() const;
 	void OnPresetChanged(TSharedPtr<FString> NewSelection, ESelectInfo::Type SelectInfo);
 	TSharedRef<SWidget> MakePresetComboWidget(TSharedPtr<FString> InItem);
 
diff --git a/Engine/Source/Editor/DetailCustomizations/Private/CameraFilmbackSettingsCustomization.cpp b/Engine/Source/Editor/DetailCustomizations/Private/CameraFilmbackSettingsCustomization.cpp
--- a/Engine/Source/Editor/DetailCustomizations/Private/CameraFilmbackSettingsCustomization.cpp
+++ b/Engine/Source/Editor/DetailCustomizations/Private/CameraFilmbackSettingsCustomization.cpp
@@ -108,6 +108,8 @@ void FCameraFilmbackSettingsCustomization::CustomizeChildren(TSharedRef<IPropert
 	// Retrieve special case properties
 	SensorWidthHandle = PropertyHandles.FindChecked(GET_MEMBER_NAME_CHECKED(FCameraFilmbackSettings, SensorWidth));
 	SensorHeightHandle = PropertyHandles.FindChecked(GET_MEMBER_NAME_CHECKED(FCameraFilmbackSettings, SensorHeight));
+	SensorHorizontalOffsetHandle = PropertyHandles.FindChecked(GET_MEMBER_NAME_CHECKED(FCameraFilmbackSettings, SensorHorizontalOffset));
+	SensorVerticalOffsetHandle = PropertyHandles.FindChecked(GET_MEMBER_NAME_CHECKED(FCameraFilmbackSettings, SensorVerticalOffset));
 
 	for( auto Iter(PropertyHandles.CreateConstIterator()); Iter; ++Iter  )
 	{
@@ -126,11 +128,16 @@ TSharedRef<SWidget> FCameraFilmbackSettingsCustomization::MakePresetComboWidget(
 bool FCameraFilmbackSettingsCustomization::IsPresetEnabled() const
 {
 	bool bEnabled = false;
-	if (SensorHeightHandle.IsValid() && SensorWidthHandle.IsValid())
+	if (SensorHeightHandle.IsValid() && 
+			SensorWidthHandle.IsValid() &&
+			SensorHorizontalOffsetHandle.IsValid() &&
+			SensorVerticalOffsetHandle.IsValid())
 	{
 		bEnabled = (
 				SensorHeightHandle->IsEditable() &&
 				SensorWidthHandle->IsEditable() &&
+				SensorHorizontalOffsetHandle->IsEditable() &&
+				SensorVerticalOffsetHandle->IsEditable() &&
 				FSlateApplication::Get().GetNormalExecutionAttribute().Get());
 	}
 	return bEnabled;
@@ -157,6 +164,8 @@ void FCameraFilmbackSettingsCustomization::OnPresetChanged(TSharedPtr<FString> N
 				// all SetValues except the last set to Interactive so we don't rerun construction scripts and invalidate subsequent property handles
 				ensure(SensorHeightHandle->SetValue(P.FilmbackSettings.SensorHeight, EPropertyValueSetFlags::InteractiveChange | EPropertyValueSetFlags::NotTransactable) == FPropertyAccess::Result::Success);
 				ensure(SensorWidthHandle->SetValue(P.FilmbackSettings.SensorWidth, EPropertyValueSetFlags::NotTransactable) == FPropertyAccess::Result::Success);
+				ensure(SensorHorizontalOffsetHandle->SetValue(P.FilmbackSettings.SensorHorizontalOffset, EPropertyValueSetFlags::InteractiveChange | EPropertyValueSetFlags::NotTransactable) == FPropertyAccess::Result::Success);
+				ensure(SensorVerticalOffsetHandle->SetValue(P.FilmbackSettings.SensorVerticalOffset, EPropertyValueSetFlags::InteractiveChange | EPropertyValueSetFlags::NotTransactable) == FPropertyAccess::Result::Success);
 				
 				break;
 			}
@@ -189,13 +198,22 @@ TSharedPtr<FString> FCameraFilmbackSettingsCustomization::GetPresetString() cons
 	float CurSensorHeight;
 	SensorHeightHandle->GetValue(CurSensorHeight);
 
+	float CurSensorHorizontalOffset;
+	SensorHorizontalOffsetHandle->GetValue(CurSensorHorizontalOffset);
+
+	float CurSensorVerticalOffset;
+	SensorVerticalOffsetHandle->GetValue(CurSensorVerticalOffset);
+
 	// search presets for one that matches
 	TArray<FNamedFilmbackPreset> const& Presets = UCineCameraSettings::GetFilmbackPresets();
 	int32 const NumPresets = Presets.Num();
 	for (int32 PresetIdx = 0; PresetIdx < NumPresets; ++PresetIdx)
 	{
 		FNamedFilmbackPreset const& P = Presets[PresetIdx];
-		if ((P.FilmbackSettings.SensorWidth == CurSensorWidth) && (P.FilmbackSettings.SensorHeight == CurSensorHeight))
+		if ((P.FilmbackSettings.SensorWidth == CurSensorWidth) && 
+				(P.FilmbackSettings.SensorHeight == CurSensorHeight) &&
+				(P.FilmbackSettings.SensorHorizontalOffset == CurSensorHorizontalOffset) &&
+				(P.FilmbackSettings.SensorVerticalOffset == CurSensorVerticalOffset))
 		{
 			// this is the one
 			if (PresetComboList.IsValidIndex(PresetIdx + 1))
diff --git a/Engine/Source/Editor/DetailCustomizations/Private/CameraFilmbackSettingsCustomization.h b/Engine/Source/Editor/DetailCustomizations/Private/CameraFilmbackSettingsCustomization.h
--- a/Engine/Source/Editor/DetailCustomizations/Private/CameraFilmbackSettingsCustomization.h
+++ b/Engine/Source/Editor/DetailCustomizations/Private/CameraFilmbackSettingsCustomization.h
@@ -38,6 +38,8 @@ protected:
 
 	TSharedPtr<IPropertyHandle> SensorWidthHandle;
 	TSharedPtr<IPropertyHandle> SensorHeightHandle;
+	TSharedPtr<IPropertyHandle> SensorHorizontalOffsetHandle;
+	TSharedPtr<IPropertyHandle> SensorVerticalOffsetHandle;
 
 	TSharedPtr<class SComboBox< TSharedPtr<FString> > > PresetComboBox;
 	TArray< TSharedPtr< FString > >						PresetComboList;
diff --git a/Engine/Source/Editor/DetailCustomizations/Private/CameraLensSettingsCustomization.cpp b/Engine/Source/Editor/DetailCustomizations/Private/CameraLensSettingsCustomization.cpp
--- a/Engine/Source/Editor/DetailCustomizations/Private/CameraLensSettingsCustomization.cpp
+++ b/Engine/Source/Editor/DetailCustomizations/Private/CameraLensSettingsCustomization.cpp
@@ -78,7 +78,7 @@ void FCameraLensSettingsCustomization::CustomizeHeader(TSharedRef<IPropertyHandl
 				.OnGenerateWidget(this, &FCameraLensSettingsCustomization::MakePresetComboWidget)
 				.OnSelectionChanged(this, &FCameraLensSettingsCustomization::OnPresetChanged)
 				.OnComboBoxOpening(this, &FCameraLensSettingsCustomization::BuildPresetComboList)
-				.IsEnabled(FSlateApplication::Get().GetNormalExecutionAttribute())
+				.IsEnabled(this, &FCameraLensSettingsCustomization::IsPresetEnabled)
 				.ContentPadding(2.f)
 				.Content()
 				[
@@ -135,6 +135,29 @@ TSharedRef<SWidget> FCameraLensSettingsCustomization::MakePresetComboWidget(TSha
 		.Font(IDetailLayoutBuilder::GetDetailFont());
 }
 
+bool FCameraLensSettingsCustomization::IsPresetEnabled() const
+{
+	bool bEnabled = false;
+	if (MinFocalLengthHandle.IsValid() &&
+			MaxFocalLengthHandle.IsValid() &&
+			MinFStopHandle.IsValid() &&
+			MaxFStopHandle.IsValid() &&
+			MinFocusDistanceHandle.IsValid() &&
+			SqueezeFactorHandle.IsValid() &&
+			DiaphragmBladeCountHandle.IsValid())
+	{
+		bEnabled = (MinFocalLengthHandle->IsEditable() &&
+				MaxFocalLengthHandle->IsEditable() &&
+				MinFStopHandle->IsEditable() &&
+				MaxFStopHandle->IsEditable() &&
+				MinFocusDistanceHandle->IsEditable() &&
+				SqueezeFactorHandle->IsEditable() &&
+				DiaphragmBladeCountHandle->IsEditable() &&
+				FSlateApplication::Get().GetNormalExecutionAttribute().Get());
+	}
+	return bEnabled;
+}
+
 void FCameraLensSettingsCustomization::OnPresetChanged(TSharedPtr<FString> NewSelection, ESelectInfo::Type SelectInfo)
 {
 	// if it's set from code, we did that on purpose
diff --git a/Engine/Source/Editor/DetailCustomizations/Private/CameraLensSettingsCustomization.h b/Engine/Source/Editor/DetailCustomizations/Private/CameraLensSettingsCustomization.h
--- a/Engine/Source/Editor/DetailCustomizations/Private/CameraLensSettingsCustomization.h
+++ b/Engine/Source/Editor/DetailCustomizations/Private/CameraLensSettingsCustomization.h
@@ -49,6 +49,7 @@ protected:
 
 	void BuildPresetComboList();
 	
+	bool IsPresetEnabled() const;
 	void OnPresetChanged(TSharedPtr<FString> NewSelection, ESelectInfo::Type SelectInfo);
 	TSharedRef<SWidget> MakePresetComboWidget(TSharedPtr<FString> InItem);
 
