allow priority to be set on transitions to conduits so it's possible to deprioritize them vs transitons to other states

#Synced-CL 51330716

diff --git a/Engine/Source/Editor/DetailCustomizations/Private/AnimTransitionNodeDetails.cpp b/Engine/Source/Editor/DetailCustomizations/Private/AnimTransitionNodeDetails.cpp
--- a/Engine/Source/Editor/DetailCustomizations/Private/AnimTransitionNodeDetails.cpp
+++ b/Engine/Source/Editor/DetailCustomizations/Private/AnimTransitionNodeDetails.cpp
@@ -106,7 +106,6 @@ void FAnimTransitionNodeDetails::CustomizeDetails( IDetailLayoutBuilder& DetailB
 		DetailBuilder.HideProperty(GET_MEMBER_NAME_CHECKED(UAnimStateTransitionNode, BlendMode));
 		DetailBuilder.HideProperty(GET_MEMBER_NAME_CHECKED(UAnimStateTransitionNode, CustomBlendCurve));
 		DetailBuilder.HideProperty(GET_MEMBER_NAME_CHECKED(UAnimStateTransitionNode, LogicType));
-		DetailBuilder.HideProperty(GET_MEMBER_NAME_CHECKED(UAnimStateTransitionNode, PriorityOrder));
 		DetailBuilder.HideProperty(GET_MEMBER_NAME_CHECKED(UAnimStateTransitionNode, BlendProfileWrapper));
 	}
 	else
