Component reference picker improvements

Fix not being able to display BP component names after selection.
Add lazy component caching (weak component can be invalidated by compilation after a pick, so the picker needs to refresh as it relies on a real component reference to display information)

#Synced-CL 51493025

diff --git a/Engine/Source/Editor/DetailCustomizations/Private/ComponentReferenceCustomization.cpp b/Engine/Source/Editor/DetailCustomizations/Private/ComponentReferenceCustomization.cpp
--- a/Engine/Source/Editor/DetailCustomizations/Private/ComponentReferenceCustomization.cpp
+++ b/Engine/Source/Editor/DetailCustomizations/Private/ComponentReferenceCustomization.cpp
@@ -527,6 +527,8 @@ void FComponentReferenceCustomization::UpdateCachedValues(bool bResetValueIfInva
 	CachedComponent.Reset();
 	CachedFirstOuterActor = GetFirstOuterActor();
 	CachedFilterActor = CachedFirstOuterActor;
+	UWorld* PickerWorld = nullptr;
+	RedirectActorToBPPreviewActor(CachedFilterActor, PickerWorld);
 
 	FComponentReference TmpComponentReference;
 	CachedPropertyAccess = GetValue(TmpComponentReference);
@@ -544,7 +546,7 @@ void FComponentReferenceCustomization::UpdateCachedValues(bool bResetValueIfInva
 		return;
 	}
 
-	CachedComponent = TmpComponentReference.GetComponent(CachedFirstOuterActor.Get());
+	CachedComponent = TmpComponentReference.GetComponent(CachedFilterActor.Get());
 	if (!IsComponentReferenceValid(TmpComponentReference))
 	{
 		CachedComponent.Reset();
@@ -555,6 +557,16 @@ void FComponentReferenceCustomization::UpdateCachedValues(bool bResetValueIfInva
 	}
 }
 
+UActorComponent* FComponentReferenceCustomization::GetCachedComponent() const
+{
+	UActorComponent* Component = CachedComponent.Get();
+	if (Component == nullptr)
+	{
+		const_cast<FComponentReferenceCustomization*>(this)->UpdateCachedValues(false);
+	}
+	return CachedComponent.Get();
+}
+
 int32 FComponentReferenceCustomization::OnGetComboContentWidgetIndex() const
 {
 	switch (CachedPropertyAccess)
@@ -578,7 +590,7 @@ bool FComponentReferenceCustomization::CanEditChildren() const
 
 const FSlateBrush* FComponentReferenceCustomization::GetActorIcon() const
 {
-	if (UActorComponent* Component = CachedComponent.Get())
+	if (UActorComponent* Component = GetCachedComponent())
 	{
 		if (AActor* Owner = Component->GetOwner())
 		{
@@ -590,7 +602,7 @@ const FSlateBrush* FComponentReferenceCustomization::GetActorIcon() const
 
 FText FComponentReferenceCustomization::OnGetActorName() const
 {
-	if (UActorComponent* Component = CachedComponent.Get())
+	if (UActorComponent* Component = GetCachedComponent())
 	{
 		if (AActor* Owner = Component->GetOwner())
 		{
@@ -602,7 +614,7 @@ FText FComponentReferenceCustomization::OnGetActorName() const
 
 const FSlateBrush* FComponentReferenceCustomization::GetComponentIcon() const
 {
-	if (const UActorComponent* ActorComponent = CachedComponent.Get())
+	if (const UActorComponent* ActorComponent = GetCachedComponent())
 	{
 		return FSlateIconFinder::FindIconBrushForClass(ActorComponent->GetClass());
 	}
@@ -613,7 +625,7 @@ FText FComponentReferenceCustomization::OnGetComponentName() const
 {
 	if (CachedPropertyAccess == FPropertyAccess::Success)
 	{
-		if (UActorComponent* ActorComponent = CachedComponent.Get())
+		if (UActorComponent* ActorComponent = GetCachedComponent())
 		{
 			const FName ComponentName = FComponentEditorUtils::FindVariableNameGivenComponentInstance(ActorComponent);
 			const bool bIsArrayVariable = !ComponentName.IsNone() && ActorComponent->GetOwner() != nullptr && FindFProperty<FArrayProperty>(ActorComponent->GetOwner()->GetClass(), ComponentName);
@@ -629,7 +641,7 @@ FText FComponentReferenceCustomization::OnGetComponentName() const
 	{
 		return LOCTEXT("MultipleValues", "Multiple Values");
 	}
-	return LOCTEXT("NoComponent", "None");
+	return LOCTEXT("NoComponent", "None"); 
 }
 
 const FSlateBrush* FComponentReferenceCustomization::GetStatusIcon() const
@@ -643,15 +655,10 @@ const FSlateBrush* FComponentReferenceCustomization::GetStatusIcon() const
 	return &EmptyBrush;
 }
 
-TSharedRef<SWidget> FComponentReferenceCustomization::OnGetMenuContent()
+bool FComponentReferenceCustomization::RedirectActorToBPPreviewActor(TWeakObjectPtr<AActor>& InOutActor, UWorld*& OutWorld) const
 {
-	UActorComponent* InitialComponent = CachedComponent.Get();
-
-	// If the actor we are looking at is a BP, use its BP editor preview world, if any
-	CachedFilterActor = CachedFirstOuterActor;
 	bool bShowTransient = false;
-	UWorld* PickerWorld = nullptr;
-	if (AActor* Actor = CachedFirstOuterActor.Get())
+	if (AActor* Actor = InOutActor.Get())
 	{
 		if (Actor->HasAnyFlags(RF_ArchetypeObject))
 		{
@@ -665,8 +672,8 @@ TSharedRef<SWidget> FComponentReferenceCustomization::OnGetMenuContent()
 					{
 						if(AActor* PreviewActor = Editor->GetPreviewActor())
 						{
-							CachedFilterActor = PreviewActor;
-							PickerWorld = PreviewActor->GetWorld();
+							InOutActor = PreviewActor;
+							OutWorld = PreviewActor->GetWorld();
 							bShowTransient = true;
 						}
 						break;
@@ -675,6 +682,17 @@ TSharedRef<SWidget> FComponentReferenceCustomization::OnGetMenuContent()
 			}
 		}
 	}
+	return bShowTransient;
+}
+
+TSharedRef<SWidget> FComponentReferenceCustomization::OnGetMenuContent()
+{
+	UActorComponent* InitialComponent = CachedComponent.Get();
+
+	// If the actor we are looking at is a BP, use its BP editor preview world, if any
+	CachedFilterActor = CachedFirstOuterActor;
+	UWorld* PickerWorld = nullptr;
+	bool bShowTransient = RedirectActorToBPPreviewActor(CachedFilterActor, PickerWorld);
 
 	return PropertyCustomizationHelpers::MakeComponentPickerWithMenu(InitialComponent
 		, bAllowClear
diff --git a/Engine/Source/Editor/DetailCustomizations/Private/ComponentReferenceCustomization.h b/Engine/Source/Editor/DetailCustomizations/Private/ComponentReferenceCustomization.h
--- a/Engine/Source/Editor/DetailCustomizations/Private/ComponentReferenceCustomization.h
+++ b/Engine/Source/Editor/DetailCustomizations/Private/ComponentReferenceCustomization.h
@@ -113,6 +113,12 @@ private:
 	 */
 	void UpdateCachedValues(bool bResetValueIfInvalid);
 
+	// Redirects the found actor to a BP preview actor to be able to choose components in a BP editor context
+	bool RedirectActorToBPPreviewActor(TWeakObjectPtr<AActor>& InOutActor, UWorld*& OutWorld) const;
+
+	// Gets (or updates) the cached component
+	UActorComponent* GetCachedComponent() const;
+
 private:
 	/** The property handle we are customizing */
 	TSharedPtr<IPropertyHandle> PropertyHandle;
