Reflection capture post process material property validation added, with an error icon displayed.  A CustomButtonSlot was added to SObjectPropertyEntryBox, similar to CustomContentSlot, except it's added to the end of the button box, allowing it to take up less space in the UX.  This location also prevents the dropdown from moving, as happens if added to the right as ExtensionContent (the first approach tried), forcing the user to move the mouse to edit the property again, when trying to fix the error.

#rnx

#Synced-CL 51974653

diff --git a/Engine/Source/Editor/DetailCustomizations/Private/DetailCustomizations.cpp b/Engine/Source/Editor/DetailCustomizations/Private/DetailCustomizations.cpp
--- a/Engine/Source/Editor/DetailCustomizations/Private/DetailCustomizations.cpp
+++ b/Engine/Source/Editor/DetailCustomizations/Private/DetailCustomizations.cpp
@@ -413,7 +413,7 @@ void FDetailCustomizationsModule::RegisterObjectCustomizations()
 	RegisterCustomClassLayout("StaticMeshActor", FOnGetDetailCustomizationInstance::CreateStatic(&FStaticMeshActorDetails::MakeInstance));
 	RegisterCustomClassLayout("MeshComponent", FOnGetDetailCustomizationInstance::CreateStatic(&FMeshComponentDetails::MakeInstance));
 	RegisterCustomClassLayout("LevelSequenceActor", FOnGetDetailCustomizationInstance::CreateStatic(&FLevelSequenceActorDetails::MakeInstance));
-	RegisterCustomClassLayout("ReflectionCapture", FOnGetDetailCustomizationInstance::CreateStatic(&FReflectionCaptureDetails::MakeInstance));
+	RegisterCustomClassLayout("ReflectionCaptureComponent", FOnGetDetailCustomizationInstance::CreateStatic(&FReflectionCaptureDetails::MakeInstance));
 	RegisterCustomClassLayout("SceneCaptureComponent", FOnGetDetailCustomizationInstance::CreateStatic(&FSceneCaptureDetails::MakeInstance));
 	RegisterCustomClassLayout("SkyLight", FOnGetDetailCustomizationInstance::CreateStatic(&FSkyLightComponentDetails::MakeInstance));
 	RegisterCustomClassLayout("Brush", FOnGetDetailCustomizationInstance::CreateStatic(&FBrushDetails::MakeInstance));
diff --git a/Engine/Source/Editor/DetailCustomizations/Private/ReflectionCaptureDetails.cpp b/Engine/Source/Editor/DetailCustomizations/Private/ReflectionCaptureDetails.cpp
--- a/Engine/Source/Editor/DetailCustomizations/Private/ReflectionCaptureDetails.cpp
+++ b/Engine/Source/Editor/DetailCustomizations/Private/ReflectionCaptureDetails.cpp
@@ -2,6 +2,13 @@
 
 #include "ReflectionCaptureDetails.h"
 
+#include "Components/ReflectionCaptureComponent.h"
+#include "DetailLayoutBuilder.h"
+#include "DetailWidgetRow.h"
+#include "Materials/Material.h"
+#include "PropertyCustomizationHelpers.h"
+#include "Styling/AppStyle.h"
+
 #define LOCTEXT_NAMESPACE "ReflectionCaptureDetails"
 
 
@@ -13,7 +20,67 @@ TSharedRef<IDetailCustomization> FReflectionCaptureDetails::MakeInstance()
 
 void FReflectionCaptureDetails::CustomizeDetails( IDetailLayoutBuilder& DetailLayout )
 {
-	
+	TSharedPtr<IPropertyHandle> PostProcessMaterialHandle = DetailLayout.GetProperty(GET_MEMBER_NAME_CHECKED(UReflectionCaptureComponent, PostProcessMaterial));
+	if (!PostProcessMaterialHandle->IsValidHandle())
+	{
+		return;
+	}
+
+	IDetailCategoryBuilder& RuntimeCaptureCategory = DetailLayout.EditCategory("RuntimeCapture");
+
+	// Explicitly add bRuntimeCapture first to preserve the declaration order when PostProcessMaterial is also added explicitly below.
+	RuntimeCaptureCategory.AddProperty(DetailLayout.GetProperty(GET_MEMBER_NAME_CHECKED(UReflectionCaptureComponent, bRuntimeCapture)));
+	IDetailPropertyRow& PropertyRow = RuntimeCaptureCategory.AddProperty(PostProcessMaterialHandle.ToSharedRef());
+
+	TSharedRef<SObjectPropertyEntryBox> ObjectEntryBox =
+		SNew(SObjectPropertyEntryBox)
+		.PropertyHandle(PostProcessMaterialHandle)
+		.AllowedClass(UMaterialInterface::StaticClass())
+		.ThumbnailPool(DetailLayout.GetThumbnailPool())
+		.CustomButtonSlot()
+		[
+			PropertyCustomizationHelpers::MakeCustomButton(
+				FAppStyle::Get().GetBrush("Icons.ErrorWithColor"),
+				FSimpleDelegate(),
+				LOCTEXT("InvalidPostProcessMaterial", "Invalid post process material"),
+				false,
+				TAttribute<EVisibility>::CreateLambda([PostProcessMaterialHandle]()
+				{
+					UObject* MaterialObj = nullptr;
+					PostProcessMaterialHandle->GetValue(MaterialObj);
+					UMaterialInterface* MaterialInterface = Cast<UMaterialInterface>(MaterialObj);
+					if (!MaterialInterface)
+					{
+						return EVisibility::Collapsed;
+					}
+
+					UMaterial* BaseMaterial = MaterialInterface->GetMaterial();
+					if (BaseMaterial &&
+						BaseMaterial->IsPostProcessMaterial() &&
+						MaterialInterface->GetUserSceneTextureOutput(BaseMaterial).IsNone() &&
+						MaterialInterface->GetBlendableLocation(BaseMaterial) == BL_SceneColorBeforeBloom)
+					{
+						return EVisibility::Collapsed;
+					}
+
+					return EVisibility::Visible;
+				}))
+		];
+
+	float MinWidth, MaxWidth;
+	ObjectEntryBox->GetDesiredWidth(MinWidth, MaxWidth);
+
+	PropertyRow.CustomWidget()
+		.NameContent()
+		[
+			PostProcessMaterialHandle->CreatePropertyNameWidget()
+		]
+		.ValueContent()
+		.MinDesiredWidth(MinWidth)
+		.MaxDesiredWidth(MaxWidth)
+		[
+			ObjectEntryBox
+		];
 }
 
 #undef LOCTEXT_NAMESPACE
diff --git a/Engine/Source/Editor/PropertyEditor/Private/PropertyCustomizationHelpers.cpp b/Engine/Source/Editor/PropertyEditor/Private/PropertyCustomizationHelpers.cpp
--- a/Engine/Source/Editor/PropertyEditor/Private/PropertyCustomizationHelpers.cpp
+++ b/Engine/Source/Editor/PropertyEditor/Private/PropertyCustomizationHelpers.cpp
@@ -965,6 +965,10 @@ void SObjectPropertyEntryBox::Construct( const FArguments& InArgs )
 				[
 					InArgs._CustomContentSlot.Widget
 				]
+				.CustomButtonSlot()
+				[
+					InArgs._CustomButtonSlot.Widget
+				]
 		]
 	];
 }
diff --git a/Engine/Source/Editor/PropertyEditor/Private/UserInterface/PropertyEditor/SPropertyEditorAsset.cpp b/Engine/Source/Editor/PropertyEditor/Private/UserInterface/PropertyEditor/SPropertyEditorAsset.cpp
--- a/Engine/Source/Editor/PropertyEditor/Private/UserInterface/PropertyEditor/SPropertyEditorAsset.cpp
+++ b/Engine/Source/Editor/PropertyEditor/Private/UserInterface/PropertyEditor/SPropertyEditorAsset.cpp
@@ -916,6 +916,16 @@ void SPropertyEditorAsset::Construct(const FArguments& InArgs, const TSharedPtr<
 		}
 	}	
 
+	if (InArgs._CustomButtonSlot.Widget != SNullWidget::NullWidget)
+	{
+		ButtonBox->AddSlot()
+			.VAlign(VAlign_Center)
+			.Padding(FMargin(0.0f, 2.0f))
+			[
+				InArgs._CustomButtonSlot.Widget
+			];
+	}
+
 	NumButtons = ButtonBox->NumSlots();
 	
 	if (ButtonBoxWrapper.IsValid())
diff --git a/Engine/Source/Editor/PropertyEditor/Private/UserInterface/PropertyEditor/SPropertyEditorAsset.h b/Engine/Source/Editor/PropertyEditor/Private/UserInterface/PropertyEditor/SPropertyEditorAsset.h
--- a/Engine/Source/Editor/PropertyEditor/Private/UserInterface/PropertyEditor/SPropertyEditorAsset.h
+++ b/Engine/Source/Editor/PropertyEditor/Private/UserInterface/PropertyEditor/SPropertyEditorAsset.h
@@ -63,6 +63,7 @@ public:
 		, _ObjectPath()
 		, _Class(nullptr)
 		, _CustomContentSlot()
+		, _CustomButtonSlot()
 	{}
 		SLATE_ARGUMENT(TOptional<bool>, AllowClear)
 		SLATE_ARGUMENT(TOptional<bool>, AllowCreate)
@@ -82,6 +83,7 @@ public:
 		SLATE_EVENT(FOnSetObject, OnSetObject)
 		SLATE_EVENT(FOnShouldFilterAsset, OnShouldFilterAsset)
 		SLATE_NAMED_SLOT(FArguments, CustomContentSlot)
+		SLATE_NAMED_SLOT(FArguments, CustomButtonSlot)
 		SLATE_ARGUMENT(TSharedPtr<IPropertyHandle>, PropertyHandle)
 		SLATE_ARGUMENT(TArray<FAssetData>, OwnerAssetDataArray)
 		SLATE_EVENT(FOnShouldFilterActor, OnShouldFilterActor)
diff --git a/Engine/Source/Editor/PropertyEditor/Public/PropertyCustomizationHelpers.h b/Engine/Source/Editor/PropertyEditor/Public/PropertyCustomizationHelpers.h
--- a/Engine/Source/Editor/PropertyEditor/Public/PropertyCustomizationHelpers.h
+++ b/Engine/Source/Editor/PropertyEditor/Public/PropertyCustomizationHelpers.h
@@ -325,6 +325,8 @@ public:
 		SLATE_ARGUMENT(bool, DisplayThumbnail)
 		/** A custom content slot for widgets */ 
 		SLATE_NAMED_SLOT(FArguments, CustomContentSlot)
+		/** A custom slot added to the row of buttons */
+		SLATE_NAMED_SLOT(FArguments, CustomButtonSlot)
 		SLATE_ATTRIBUTE(FIntPoint, ThumbnailSizeOverride)
 		/** Called to check if an actor is valid to use */
 		SLATE_EVENT(FOnShouldFilterActor, OnShouldFilterActor)
