Static Mesh Map Check error fixup tool.  Added custom button that cleans up excess override materials for static mesh components, allowing related map check errors to be fixed in bulk with a single click (especially by taking advantage of the "Select actors by message" feature of the map check dialog).  Button appears conditionally based on whether an affected mesh component is selected, and tool supports undo / redo.

#Synced-CL 52550019

diff --git a/Engine/Source/Editor/DetailCustomizations/Private/StaticMeshComponentDetails.cpp b/Engine/Source/Editor/DetailCustomizations/Private/StaticMeshComponentDetails.cpp
--- a/Engine/Source/Editor/DetailCustomizations/Private/StaticMeshComponentDetails.cpp
+++ b/Engine/Source/Editor/DetailCustomizations/Private/StaticMeshComponentDetails.cpp
@@ -7,6 +7,9 @@
 #include "DetailLayoutBuilder.h"
 #include "DetailWidgetRow.h"
 #include "Engine/StaticMesh.h"
+#include "PropertyCustomizationHelpers.h"
+#include "ScopedTransaction.h"
+#include "Styling/AppStyle.h"
 #include "Fonts/SlateFontInfo.h"
 #include "HAL/PlatformCrt.h"
 #include "IDetailPropertyRow.h"
@@ -39,7 +42,99 @@ TSharedRef<IDetailCustomization> FStaticMeshComponentDetails::MakeInstance()
 void FStaticMeshComponentDetails::CustomizeDetails(IDetailLayoutBuilder& DetailBuilder)
 {
 	// Create a category so this is displayed early in the properties
-	DetailBuilder.EditCategory("StaticMesh", FText::GetEmpty(), ECategoryPriority::Important);
+	IDetailCategoryBuilder& StaticMeshCategory = DetailBuilder.EditCategory("StaticMesh", FText::GetEmpty(), ECategoryPriority::Important);
+
+	// Customize the StaticMesh property to add a refresh button via CustomButtonSlot
+	StaticMeshProperty = DetailBuilder.GetProperty(GET_MEMBER_NAME_CHECKED(UStaticMeshComponent, StaticMesh));
+	if (StaticMeshProperty->IsValidHandle())
+	{
+		IDetailPropertyRow& StaticMeshRow = StaticMeshCategory.AddProperty(StaticMeshProperty.ToSharedRef());
+
+		TSharedPtr<IPropertyHandle> StaticMeshHandle = StaticMeshProperty;
+		TSharedRef<SObjectPropertyEntryBox> ObjectEntryBox =
+			SNew(SObjectPropertyEntryBox)
+			.PropertyHandle(StaticMeshProperty)
+			.AllowedClass(UStaticMesh::StaticClass())
+			.ThumbnailPool(DetailBuilder.GetThumbnailPool())
+			.CustomButtonSlot()
+			[
+				PropertyCustomizationHelpers::MakeCustomButton(
+					FAppStyle::Get().GetBrush("Icons.WarningWithColor"),
+					FSimpleDelegate::CreateLambda([StaticMeshHandle]()
+					{
+						if (!StaticMeshHandle.IsValid())
+						{
+							return;
+						}
+
+						TArray<UObject*> OuterObjects;
+						StaticMeshHandle->GetOuterObjects(OuterObjects);
+
+						TArray<UMeshComponent*> ComponentsToClean;
+						for (UObject* Obj : OuterObjects)
+						{
+							if (UMeshComponent* MeshComp = Cast<UMeshComponent>(Obj))
+							{
+								if (MeshComp->GetNumOverrideMaterials() > MeshComp->GetNumMaterials())
+								{
+									ComponentsToClean.Add(MeshComp);
+								}
+							}
+						}
+
+						if (ComponentsToClean.Num() > 0)
+						{
+							FScopedTransaction Transaction(LOCTEXT("CleanUpOverrideMaterials", "Clean Up Override Materials"));
+
+							for (UMeshComponent* MeshComp : ComponentsToClean)
+							{
+								MeshComp->Modify();
+								MeshComp->CleanUpOverrideMaterials();
+							}
+						}
+					}),
+					LOCTEXT("CleanUpOverrideMaterialsToolTip", "Clean up override materials that exceed the mesh material count"),
+					true,
+					TAttribute<EVisibility>::CreateLambda([StaticMeshHandle]()
+					{
+						if (!StaticMeshHandle.IsValid())
+						{
+							return EVisibility::Collapsed;
+						}
+
+						TArray<UObject*> OuterObjects;
+						StaticMeshHandle->GetOuterObjects(OuterObjects);
+
+						for (UObject* Obj : OuterObjects)
+						{
+							if (UMeshComponent* MeshComp = Cast<UMeshComponent>(Obj))
+							{
+								if (MeshComp->GetNumOverrideMaterials() > MeshComp->GetNumMaterials())
+								{
+									return EVisibility::Visible;
+								}
+							}
+						}
+
+						return EVisibility::Collapsed;
+					}))
+			];
+
+		float MinWidth, MaxWidth;
+		ObjectEntryBox->GetDesiredWidth(MinWidth, MaxWidth);
+
+		StaticMeshRow.CustomWidget()
+			.NameContent()
+			[
+				StaticMeshProperty->CreatePropertyNameWidget()
+			]
+			.ValueContent()
+			.MinDesiredWidth(MinWidth)
+			.MaxDesiredWidth(MaxWidth)
+			[
+				ObjectEntryBox
+			];
+	}
 
 	TSharedRef<IPropertyHandle> UseDefaultCollision = DetailBuilder.GetProperty(GET_MEMBER_NAME_CHECKED(UStaticMeshComponent, bUseDefaultCollision));
 	UseDefaultCollision->MarkHiddenByCustomization();
diff --git a/Engine/Source/Editor/DetailCustomizations/Private/StaticMeshComponentDetails.h b/Engine/Source/Editor/DetailCustomizations/Private/StaticMeshComponentDetails.h
--- a/Engine/Source/Editor/DetailCustomizations/Private/StaticMeshComponentDetails.h
+++ b/Engine/Source/Editor/DetailCustomizations/Private/StaticMeshComponentDetails.h
@@ -36,4 +36,5 @@ private:
 private:
 	TArray<TWeakObjectPtr<UObject>> ObjectsCustomized;
 	TSharedPtr<IPropertyHandle> OverrideLightResProperty;
+	TSharedPtr<IPropertyHandle> StaticMeshProperty;
 };
diff --git a/Engine/Source/Runtime/Engine/Classes/Components/StaticMeshComponent.h b/Engine/Source/Runtime/Engine/Classes/Components/StaticMeshComponent.h
--- a/Engine/Source/Runtime/Engine/Classes/Components/StaticMeshComponent.h
+++ b/Engine/Source/Runtime/Engine/Classes/Components/StaticMeshComponent.h
@@ -952,6 +952,7 @@ private:
 #endif
 	FOnGetNaniteResources OnGetNaniteResourcesEvent;
 
+	friend class FStaticMeshComponentDetails;
 	friend class FStaticMeshComponentRecreateRenderStateContext;
 	friend class FStaticMeshComponentHelper;
 	friend class FInstancedStaticMeshComponentHelper;
