EditorUsability : LiveEdit
- Fixed color picker would close and not transact any property change on any blueprint classes, this is caused by the blueprint rerun construction script and objects replacement happening in editor, the actor and its sub-objects are duplicated and replaced after a property change (value set), this would trigger a details panel refresh with the new objects and underlying property handle would change, old ones would be invalidated, and color picker would close immediately without processing transaction for live edit. Now user can interact with color picker and preview the property color change in the editor viewport, when OK is pressed, transaction will be processed by live edit and update the color on the client side. This would not happen with C++ actor class since they are not getting replaced.

#Synced-CL 52465190

diff --git a/Engine/Source/Editor/DetailCustomizations/Private/ColorStructCustomization.cpp b/Engine/Source/Editor/DetailCustomizations/Private/ColorStructCustomization.cpp
--- a/Engine/Source/Editor/DetailCustomizations/Private/ColorStructCustomization.cpp
+++ b/Engine/Source/Editor/DetailCustomizations/Private/ColorStructCustomization.cpp
@@ -51,6 +51,42 @@ struct FGeometry;
 
 #define LOCTEXT_NAMESPACE "FColorStructCustomization"
 
+namespace UE::Editor::DetailCustomizations
+{
+	bool IsPropertyInBlueprint(const TSharedPtr<IPropertyHandle>& StructHandle)
+	{
+		if (!StructHandle)
+		{
+			return false;
+		}
+
+		TArray<UObject*> OuterObjects;
+		StructHandle->GetOuterObjects(OuterObjects);
+
+		while (!OuterObjects.IsEmpty())
+		{
+			const UObject* Object = OuterObjects.Pop();
+			if (!Object)
+			{
+				continue;
+			}
+
+			const UClass* Class = Object->GetClass();
+			if (Class && Class->HasAnyClassFlags(CLASS_CompiledFromBlueprint))
+			{
+				return true;
+			}
+		
+			if (UObject* Outer = Object->GetOuter())
+			{
+				OuterObjects.Add(Outer);
+			}
+		}
+
+		return false;
+	};
+}
+
 
 TSharedRef<IPropertyTypeCustomization> FColorStructCustomization::MakeInstance() 
 {
@@ -62,6 +98,7 @@ void FColorStructCustomization::CustomizeHeader(TSharedRef<class IPropertyHandle
 {
 	StructPropertyHandle = InStructPropertyHandle;
 
+	bIsInBlueprint = UE::Editor::DetailCustomizations::IsPropertyInBlueprint(InStructPropertyHandle);
 	bIsLinearColor = CastFieldChecked<FStructProperty>(StructPropertyHandle->GetProperty())->Struct->GetFName() == NAME_LinearColor;
 	bIgnoreAlpha = TypeSupportsAlpha() == false || StructPropertyHandle->GetProperty()->HasMetaData(TEXT("HideAlphaChannel"));
 	
@@ -266,6 +303,37 @@ void FColorStructCustomization::CreateColorPicker(bool bUseAlpha)
 	}
 
 	OpenColorPicker(PickerArgs);
+	
+	// Close the color picker when a property is edited on the same blueprint actor 
+	// since it will trigger a rerun construction script and duplicate actor and its subobjects
+	// and invalidate property handle since details panel refreshes with objects replaced
+	if (!FCoreUObjectDelegates::OnObjectPropertyChanged.IsBoundToObject(this))
+	{
+		FCoreUObjectDelegates::OnObjectPropertyChanged.AddSPLambda(this, [this](UObject* InObject, FPropertyChangedEvent& InEvent)
+			{
+				const TSharedPtr<SColorPicker> ColorPicker = GetColorPicker();
+				if (StructPropertyHandle.IsValid()
+					&& ColorPicker.IsValid()
+					&& bIsInBlueprint 
+					&& InEvent.ChangeType != EPropertyChangeType::Interactive
+					&& InEvent.Property != StructPropertyHandle->GetProperty())
+				{
+					TArray<UObject*> OuterObjects;
+					StructPropertyHandle->GetOuterObjects(OuterObjects);
+					
+					const UObject* RootOwner = InObject->IsA<AActor>() ? InObject : InObject->GetTypedOuter<AActor>();
+					for (const UObject* Object : OuterObjects)
+					{
+						if (Object == RootOwner || Object->GetTypedOuter<AActor>() == RootOwner)
+						{
+							DestroyColorPicker();
+							return;
+						}
+					}
+				}
+			}
+		);
+	}
 }
 
 
@@ -313,11 +381,9 @@ void FColorStructCustomization::SetLastPickerColorString(const FLinearColor NewC
 void FColorStructCustomization::OnSetColorFromColorPicker(FLinearColor NewColor)
 {
 	SetLastPickerColorString(NewColor);
-
 	EPropertyValueSetFlags::Type PropertyFlags = EPropertyValueSetFlags::NotTransactable;
-	PropertyFlags |= bIsInteractive ? EPropertyValueSetFlags::InteractiveChange : 0;
+	PropertyFlags |= (bIsInteractive || bIsInBlueprint) ? EPropertyValueSetFlags::InteractiveChange : 0;
 	StructPropertyHandle->SetValueFromFormattedString(LastPickerColorString, PropertyFlags);
-	StructPropertyHandle->NotifyFinishedChangingProperties();
 }
 
 TArray<FString> FColorStructCustomization::ConvertToPerObjectColors(const TArray<FLinearOrSrgbColor>& Colors) const
@@ -359,6 +425,8 @@ void FColorStructCustomization::OnColorPickerCancelled(FLinearColor OriginalColo
 
 	GEditor->CancelTransaction(0);
 	TransactionIndex.Reset();
+
+	FCoreUObjectDelegates::OnObjectPropertyChanged.RemoveAll(this);
 }
 
 void FColorStructCustomization::OnColorPickerWindowClosed(const TSharedRef<SWindow>& Window)
@@ -377,6 +445,8 @@ void FColorStructCustomization::OnColorPickerWindowClosed(const TSharedRef<SWind
 
 	GEditor->EndTransaction();
 	TransactionIndex.Reset();
+	
+	FCoreUObjectDelegates::OnObjectPropertyChanged.RemoveAll(this);
 }
 
 
@@ -394,8 +464,9 @@ void FColorStructCustomization::OnColorPickerInteractiveEnd()
 
 		if (!LastPickerColorString.IsEmpty())
 		{
-			StructPropertyHandle->SetValueFromFormattedString(LastPickerColorString, EPropertyValueSetFlags::NotTransactable);
-			StructPropertyHandle->NotifyFinishedChangingProperties();
+			EPropertyValueSetFlags::Type PropertyFlags = EPropertyValueSetFlags::NotTransactable;
+			PropertyFlags |= bIsInBlueprint ? EPropertyValueSetFlags::InteractiveChange : 0;
+			StructPropertyHandle->SetValueFromFormattedString(LastPickerColorString, PropertyFlags);
 		}
 	}
 }
diff --git a/Engine/Source/Editor/DetailCustomizations/Public/Customizations/ColorStructCustomization.h b/Engine/Source/Editor/DetailCustomizations/Public/Customizations/ColorStructCustomization.h
--- a/Engine/Source/Editor/DetailCustomizations/Public/Customizations/ColorStructCustomization.h
+++ b/Engine/Source/Editor/DetailCustomizations/Public/Customizations/ColorStructCustomization.h
@@ -39,17 +39,19 @@ class FColorStructCustomization
 public:
 
 	static UE_API TSharedRef<IPropertyTypeCustomization> MakeInstance();
-	UE_API ~FColorStructCustomization();
+	UE_API virtual ~FColorStructCustomization() override;
 
 protected:
 
 	FColorStructCustomization()
-		: bIgnoreAlpha(false)
+		: bIsLinearColor(false)
+		, bIgnoreAlpha(false)
 		, bIsInlineColorPickerVisible(false)
 		, bIsInteractive(false)
+		, bIsInBlueprint(false)
 		, bDontUpdateWhileEditing(false)
-
-	{}
+	{
+	}
 
 protected:
 
@@ -216,6 +218,9 @@ protected:
 
 	/** True if the user is performing an interactive color change */
 	bool bIsInteractive;
+	
+	/** True if the property we are editing is contained within a blueprint */
+	bool bIsInBlueprint;
 
 	/** Last color set from color picker as string*/
 	FString LastPickerColorString;
diff --git a/Engine/Source/Editor/PropertyEditor/Private/SDetailsView.cpp b/Engine/Source/Editor/PropertyEditor/Private/SDetailsView.cpp
--- a/Engine/Source/Editor/PropertyEditor/Private/SDetailsView.cpp
+++ b/Engine/Source/Editor/PropertyEditor/Private/SDetailsView.cpp
@@ -975,14 +975,6 @@ void SDetailsView::PreSetObject(int32 InNewNumObjects)
 {
 	TRACE_CPUPROFILER_EVENT_SCOPE(SDetailsView::PreSetObject);
 
-	TSharedPtr<SColorPicker> ExistingColorPicker = GetColorPicker();
-	if (ExistingColorPicker.IsValid()
-		&& ExistingColorPicker->GetOptionalOwningDetailsView().IsValid()
-		&& ExistingColorPicker->GetOptionalOwningDetailsView().Get() != this)
-	{
-		DestroyColorPicker();
-	}
-
 	// Save existing expanded items first
 	for(TSharedPtr<FComplexPropertyNode>& RootNode : RootPropertyNodes)
 	{
@@ -1026,14 +1018,6 @@ void SDetailsView::PostSetObject(const TArray<FDetailsViewObjectRoot>& Roots)
 {
 	TRACE_CPUPROFILER_EVENT_SCOPE(SDetailsView::PostSetObject);
 
-	TSharedPtr<SColorPicker> ExistingColorPicker = GetColorPicker();
-	if (ExistingColorPicker.IsValid()
-		&& (!ExistingColorPicker->GetOptionalOwningDetailsView().IsValid()
-			|| ExistingColorPicker->GetOptionalOwningDetailsView().Get() == this))
-	{
-		DestroyColorPicker();
-	}
-
 	ColorPropertyNode = nullptr;
 
 	// Are we editing PIE objects?  If the bShowHiddenPropertiesWhilePlaying setting is enabled, we may want to
